Copy up of Unreal Insights from //Tasks/UE4/Dev-Core-Trace.

diff --git a/Engine/Source/Developer/TraceServices/Private/Analyzers/StatsTraceAnalysis.cpp b/Engine/Source/Developer/TraceServices/Private/Analyzers/StatsTraceAnalysis.cpp
new file mode 100644
--- /dev/null
+++ b/Engine/Source/Developer/TraceServices/Private/Analyzers/StatsTraceAnalysis.cpp
@@ -0,0 +1,140 @@
+// Copyright 1998-2019 Epic Games, Inc. All Rights Reserved.
+#include "StatsTraceAnalysis.h"
+#include "AnalysisServicePrivate.h"
+#include "Common/Utils.h"
+#include "Model/Counters.h"
+
+FStatsAnalyzer::FStatsAnalyzer(TSharedRef<Trace::FAnalysisSession> InSession)
+	: Session(InSession)
+	, CountersProvider(Session->EditCountersProvider())
+{
+
+}
+
+void FStatsAnalyzer::OnAnalysisBegin(const FOnAnalysisContext& Context)
+{
+	auto& Builder = Context.InterfaceBuilder;
+
+	Builder.RouteEvent(RouteId_Spec, "Stats", "Spec");
+	Builder.RouteEvent(RouteId_EventBatch, "Stats", "EventBatch");
+}
+
+void FStatsAnalyzer::OnAnalysisEnd()
+{
+}
+
+void FStatsAnalyzer::OnEvent(uint16 RouteId, const FOnEventContext& Context)
+{
+	const auto& EventData = Context.EventData;
+	switch (RouteId)
+	{
+	case RouteId_Spec:
+	{
+		Trace::FAnalysisSessionEditScope _(Session.Get());
+		uint32 StatId = EventData.GetValue("Id").As<uint32>();
+		check(!CountersMap.Contains(StatId));
+		const ANSICHAR* Name = reinterpret_cast<const ANSICHAR*>(EventData.GetAttachment());
+		const TCHAR* Description = reinterpret_cast<const TCHAR*>(EventData.GetAttachment() + strlen(Name) + 1);
+		Trace::ECounterDisplayHint DisplayHint = Trace::CounterDisplayHint_None;
+		if (EventData.GetValue("IsFloatingPoint").As<bool>())
+		{
+			DisplayHint = Trace::CounterDisplayHint_FloatingPoint;
+		}
+		else if (EventData.GetValue("IsMemory").As<bool>())
+		{
+			DisplayHint = Trace::CounterDisplayHint_Memory;
+		}
+		Trace::FCounterInternal* Counter = CountersProvider->CreateCounter(ANSI_TO_TCHAR(Name), Description, DisplayHint);
+		CountersMap.Add(StatId, Counter);
+		break;
+	}
+	case RouteId_EventBatch:
+	{
+		Trace::FAnalysisSessionEditScope _(Session.Get());
+		uint32 ThreadId = EventData.GetValue("ThreadId").As<uint32>();
+		TSharedRef<FThreadState> ThreadState = GetThreadState(ThreadId);
+		uint64 BufferSize = EventData.GetAttachmentSize();
+		const uint8* BufferPtr = EventData.GetAttachment();
+		const uint8* BufferEnd = BufferPtr + BufferSize;
+		while (BufferPtr < BufferEnd)
+		{
+			enum EOpType
+			{
+				Increment = 0,
+				Decrement = 1,
+				AddInteger = 2,
+				SetInteger = 3,
+				AddFloat = 4,
+				SetFloat = 5,
+			};
+
+			uint64 DecodedIdAndOp = FTraceAnalyzerUtils::Decode7bit(BufferPtr);
+			uint32 StatId = DecodedIdAndOp >> 3;
+			//check(CountersMap.Contains(StatId));
+			Trace::FCounterInternal& Counter = *CountersMap[StatId];
+			uint8 Op = DecodedIdAndOp & 0x7;
+			uint64 CycleDiff = FTraceAnalyzerUtils::Decode7bit(BufferPtr);
+			uint64 Cycle = ThreadState->LastCycle + CycleDiff;
+			double Time = Context.SessionContext.TimestampFromCycle(Cycle);
+			ThreadState->LastCycle = Cycle;
+			switch (Op)
+			{
+			case Increment:
+			{
+				CountersProvider->Add(Counter, Time, int64(1));
+				break;
+			}
+			case Decrement:
+			{
+				CountersProvider->Add(Counter, Time, int64(-1));
+				break;
+			}
+			case AddInteger:
+			{
+				int64 Amount = FTraceAnalyzerUtils::DecodeZigZag(BufferPtr);
+				CountersProvider->Add(Counter, Time, Amount);
+				break;
+			}
+			case SetInteger:
+			{
+				int64 Value = FTraceAnalyzerUtils::DecodeZigZag(BufferPtr);
+				CountersProvider->Set(Counter, Time, Value);
+				break;
+			}
+			case AddFloat:
+			{
+				double Amount;
+				memcpy(&Amount, BufferPtr, sizeof(double));
+				BufferPtr += sizeof(double);
+				CountersProvider->Add(Counter, Time, Amount);
+				break;
+			}
+			case SetFloat:
+			{
+				double Value;
+				memcpy(&Value, BufferPtr, sizeof(double));
+				BufferPtr += sizeof(double);
+				break;
+			}
+			}
+		}
+		check(BufferPtr == BufferEnd);
+		break;
+	}
+	}
+}
+
+TSharedRef<FStatsAnalyzer::FThreadState> FStatsAnalyzer::GetThreadState(uint32 ThreadId)
+{
+	if (!ThreadStatesMap.Contains(ThreadId))
+	{
+		TSharedRef<FThreadState> ThreadState = MakeShared<FThreadState>();
+		ThreadState->LastCycle = 0;
+		ThreadStatesMap.Add(ThreadId, ThreadState);
+		return ThreadState;
+	}
+	else
+	{
+		return ThreadStatesMap[ThreadId];
+	}
+}
diff --git a/Engine/Source/Developer/TraceServices/Private/Analyzers/StatsTraceAnalysis.h b/Engine/Source/Developer/TraceServices/Private/Analyzers/StatsTraceAnalysis.h
new file mode 100644
--- /dev/null
+++ b/Engine/Source/Developer/TraceServices/Private/Analyzers/StatsTraceAnalysis.h
@@ -0,0 +1,43 @@
+// Copyright 1998-2019 Epic Games, Inc. All Rights Reserved.
+
+#pragma once
+
+#include "Trace/Trace.h"
+#include "Trace/Analyzer.h"
+#include "Templates/SharedPointer.h"
+
+namespace Trace
+{
+	class FAnalysisSession;
+	class FCounterInternal;
+	class FCountersProvider;
+}
+
+class FStatsAnalyzer
+	: public Trace::IAnalyzer
+{
+public:
+	FStatsAnalyzer(TSharedRef<Trace::FAnalysisSession> Session);
+	virtual void OnAnalysisBegin(const FOnAnalysisContext& Context) override;
+	virtual void OnEvent(uint16 RouteId, const FOnEventContext& Context) override;
+	virtual void OnAnalysisEnd() override;
+
+private:
+	enum : uint16
+	{
+		RouteId_Spec,
+		RouteId_EventBatch,
+	};
+
+	struct FThreadState
+	{
+		uint64 LastCycle = 0;
+	};
+
+	TSharedRef<FThreadState> GetThreadState(uint32 ThreadId);
+
+	TSharedRef<Trace::FAnalysisSession> Session;
+	TSharedRef<Trace::FCountersProvider> CountersProvider;
+	TMap<uint32, Trace::FCounterInternal*> CountersMap;
+	TMap<uint32, TSharedRef<FThreadState>> ThreadStatesMap;
+};
diff --git a/Engine/Source/Developer/TraceServices/Private/Common/PagedArray.h b/Engine/Source/Developer/TraceServices/Private/Common/PagedArray.h
new file mode 100644
--- /dev/null
+++ b/Engine/Source/Developer/TraceServices/Private/Common/PagedArray.h
@@ -0,0 +1,229 @@
+// Copyright 1998-2019 Epic Games, Inc. All Rights Reserved.
+
+#pragma once
+
+#include "Common/SlabAllocator.h"
+
+struct FNoPageData
+{
+};
+
+template<typename ItemType, typename PageDataType = FNoPageData>
+class TPagedArray
+{
+public:
+
+	struct FPageInternal;
+
+	class FIterator
+	{
+	public:
+		const PageDataType* GetCurrentPage()
+		{
+			if (CurrentPage)
+			{
+				uint64 CurrentPageIndex = CurrentPage - FirstPage;
+				return Outer.PageUserDatas.GetData() + CurrentPageIndex;
+			}
+			else
+			{
+				return nullptr;
+			}
+		}
+
+		const ItemType* GetCurrentItem()
+		{
+			return CurrentItem;
+		}
+
+		const PageDataType* PrevPage()
+		{
+			if (CurrentPage == FirstPage)
+			{
+				return nullptr;
+			}
+			--CurrentPage;
+			CurrentPageBegin = CurrentPage->Items - 1;
+			CurrentPageEnd = CurrentPage->Items + CurrentPage->Count;
+			CurrentItem = CurrentPageEnd - 1;
+			return GetCurrentPage();
+		}
+
+		const PageDataType* NextPage()
+		{
+			if (CurrentPage == LastPage)
+			{
+				return nullptr;
+			}
+			++CurrentPage;
+			CurrentPageBegin = CurrentPage->Items - 1;
+			CurrentPageEnd = CurrentPage->Items + CurrentPage->Count;
+			CurrentItem = CurrentPage->Items;
+			return GetCurrentPage();
+		}
+
+		const ItemType* NextItem()
+		{
+			++CurrentItem;
+			if (CurrentItem >= CurrentPageEnd && !NextPage())
+			{
+				return nullptr;
+			}
+			return CurrentItem;
+		}
+
+		const ItemType* PrevItem()
+		{
+			--CurrentItem;
+			if (CurrentItem <= CurrentPageBegin && !PrevPage())
+			{
+				return nullptr;
+			}
+			return CurrentItem;
+		}
+
+	private:
+		friend class TPagedArray;
+
+		FIterator(const TPagedArray& InOuter, uint64 InitialBankIndex, uint64 InitialItemIndex)
+			: Outer(InOuter)
+		{
+			if (Outer.Pages.Num() > 0)
+			{
+				FirstPage = Outer.Pages.GetData();
+				LastPage = Outer.Pages.GetData() + Outer.Pages.Num() - 1;
+				check(InitialBankIndex < Outer.Pages.Num());
+				CurrentPage = Outer.Pages.GetData() + InitialBankIndex;
+				CurrentPageBegin = CurrentPage->Items - 1;
+				CurrentPageEnd = CurrentPage->Items + CurrentPage->Count;
+				check(InitialItemIndex < CurrentPage->Count);
+				CurrentItem = CurrentPage->Items + InitialItemIndex;
+			}
+		}
+
+		const TPagedArray& Outer;
+		const FPageInternal* CurrentPage = nullptr;
+		const FPageInternal* FirstPage = nullptr;
+		const FPageInternal* LastPage = nullptr;
+		const ItemType* CurrentItem = nullptr;
+		const ItemType* CurrentPageBegin = nullptr;
+		const ItemType* CurrentPageEnd = nullptr;
+	};
+
+	TPagedArray(FSlabAllocator& InAllocator, uint64 InPageSize)
+		: Allocator(InAllocator)
+		, PageSize(InPageSize)
+	{
+
+	}
+
+	~TPagedArray()
+	{
+		for (FPageInternal& Page : Pages)
+		{
+			ItemType* PageEnd = Page.Items + Page.Count;
+			for (ItemType* Item = Page.Items; Item != PageEnd; ++Item)
+			{
+				Item->~ItemType();
+			}
+		}
+	}
+
+	uint64 Num() const
+	{
+		return TotalItemCount;
+	}
+
+	ItemType& PushBack()
+	{
+		if (!LastPage || LastPage->Count == PageSize)
+		{
+			LastPage = &Pages.AddDefaulted_GetRef();
+			LastPage->Items = Allocator.Allocate<ItemType>(PageSize);
+			PageUserDatas.AddDefaulted_GetRef();
+		}
+		++TotalItemCount;
+		ItemType* ItemPtr = LastPage->Items + LastPage->Count;
+		new (ItemPtr) ItemType();
+		++LastPage->Count;
+		return *ItemPtr;
+	}
+
+	ItemType& Insert(uint64 Index)
+	{
+		PushBack();
+		for (uint64 CurrentIndex = TotalItemCount - 1; CurrentIndex > Index; --CurrentIndex)
+		{
+			(*this)[CurrentIndex] = (*this)[CurrentIndex - 1];
+		}
+		return (*this)[Index];
+	}
+
+	PageDataType* GetLastPage()
+	{
+		if (PageUserDatas.Num())
+		{
+			return &PageUserDatas.Last();
+		}
+		else
+		{
+			return nullptr;
+		}
+	}
+
+	PageDataType* GetPage(uint64 PageIndex)
+	{
+		return PageUserDatas.GetData() + PageIndex;
+	}
+
+	PageDataType* GetItemPage(uint64 ItemIndex)
+	{
+		uint64 PageIndex = ItemIndex / PageSize;
+		return PageUserDatas.GetData() + PageIndex;
+	}
+
+	FIterator GetIteratorFromPage(uint64 PageIndex) const
+	{
+		return FIterator(*this, PageIndex, 0);
+	}
+
+	FIterator GetIteratorFromItem(uint64 ItemIndex) const
+	{
+		uint64 PageIndex = ItemIndex / PageSize;
+		uint64 IndexInPage = ItemIndex % PageSize;
+		return FIterator(*this, PageIndex, IndexInPage);
+	}
+
+	const TArray<PageDataType> GetPages() const
+	{
+		return PageUserDatas;
+	}
+
+	ItemType& operator[](uint64 Index)
+	{
+		uint64 PageIndex = Index / PageSize;
+		uint64 IndexInPage = Index % PageSize;
+		FPageInternal* Page = Pages.GetData() + PageIndex;
+		ItemType* Item = Page->Items + IndexInPage;
+		return *Item;
+	}
+
+	const ItemType& operator[](uint64 Index) const
+	{
+		return const_cast<TPagedArray&>(*this)[Index];
+	}
+
+private:
+	struct FPageInternal
+	{
+		ItemType* Items = nullptr;
+		uint64 Count = 0;
+	};
+
+	FSlabAllocator& Allocator;
+	TArray<FPageInternal> Pages;
+	TArray<PageDataType> PageUserDatas;
+	FPageInternal* LastPage = nullptr;
+	uint64 PageSize;
+	uint64 TotalItemCount = 0;
+};
diff --git a/Engine/Source/Developer/TraceServices/Private/Common/SlabAllocator.h b/Engine/Source/Developer/TraceServices/Private/Common/SlabAllocator.h
new file mode 100644
--- /dev/null
+++ b/Engine/Source/Developer/TraceServices/Private/Common/SlabAllocator.h
@@ -0,0 +1,49 @@
+// Copyright 1998-2019 Epic Games, Inc. All Rights Reserved.
+
+#pragma once
+
+#include "HAL/UnrealMemory.h"
+
+class FSlabAllocator
+{
+public:
+	FSlabAllocator(uint64 InSlabSize)
+		: SlabSize(InSlabSize)
+	{
+
+	}
+
+	~FSlabAllocator()
+	{
+		for (void* Slab : Slabs)
+		{
+			FMemory::Free(Slab);
+		}
+	}
+
+	template<typename T>
+	T* Allocate(uint64 Count)
+	{
+		uint64 AllocationSize = Count * sizeof(T) + (16 - 1) / 16;
+		if (!CurrentSlab || CurrentSlabAllocatedSize + AllocationSize > SlabSize)
+		{
+			TotalAllocatedSize += SlabSize;
+			void* Allocation = FMemory::Malloc(SlabSize, 16);
+			CurrentSlab = reinterpret_cast<uint8*>(Allocation);
+			CurrentSlabAllocatedSize = 0;
+			Slabs.Add(CurrentSlab);
+		}
+		T* Allocation = reinterpret_cast<T*>(CurrentSlab + CurrentSlabAllocatedSize);
+		CurrentSlabAllocatedSize += AllocationSize;
+		TotalItemCount += Count;
+		return Allocation;
+	}
+
+private:
+	TArray<void*> Slabs;
+	uint8* CurrentSlab = nullptr;
+	const uint64 SlabSize;
+	uint64 CurrentSlabAllocatedSize = 0;
+	uint64 TotalItemCount = 0;
+	uint64 TotalAllocatedSize = 0;
+};
diff --git a/Engine/Source/Developer/TraceServices/Private/Model/Bookmarks.cpp b/Engine/Source/Developer/TraceServices/Private/Model/Bookmarks.cpp
new file mode 100644
--- /dev/null
+++ b/Engine/Source/Developer/TraceServices/Private/Model/Bookmarks.cpp
@@ -0,0 +1,78 @@
+// Copyright 1998-2019 Epic Games, Inc. All Rights Reserved.
+
+#include "Model/Bookmarks.h"
+#include "AnalysisServicePrivate.h"
+#include "Common/FormatArgs.h"
+
+namespace Trace
+{
+
+FBookmarkProvider::FBookmarkProvider(const FAnalysisSessionLock& InSessionLock)
+	: SessionLock(InSessionLock)
+{
+
+}
+
+FBookmarkSpec& FBookmarkProvider::GetSpec(uint64 BookmarkPoint)
+{
+	SessionLock.WriteAccessCheck();
+	if (SpecMap.Contains(BookmarkPoint))
+	{
+		return *SpecMap[BookmarkPoint].Get();
+	}
+	else
+	{
+		TSharedPtr<FBookmarkSpec> Spec = MakeShared<FBookmarkSpec>();
+		SpecMap.Add(BookmarkPoint, Spec);
+		return *Spec.Get();
+	}
+}
+
+void FBookmarkProvider::AppendBookmark(double Time, uint64 BookmarkPoint, uint16 FormatArgsSize, const uint8* FormatArgs)
+{
+	SessionLock.WriteAccessCheck();
+	FBookmarkSpec& Spec = GetSpec(BookmarkPoint);
+	TSharedRef<FBookmarkInternal> Bookmark = MakeShared<FBookmarkInternal>();
+	Bookmark->Time = Time;
+	TCHAR Buffer[65535];
+	FFormatArgsHelper::Format(Buffer, 65534, *Spec.FormatString, FormatArgs);
+	Bookmark->Text = Buffer;
+	Bookmarks.Add(Bookmark);
+}
+
+void FBookmarkProvider::EnumerateBookmarks(double IntervalStart, double IntervalEnd, TFunctionRef<void(const FBookmark &)> Callback) const
+{
+	SessionLock.ReadAccessCheck();
+	if (IntervalStart > IntervalEnd)
+	{
+		return;
+	}
+	uint64 FirstBookmarkIndex = Algo::LowerBoundBy(Bookmarks, IntervalStart, [](const TSharedRef<FBookmarkInternal>& B)
+	{
+		return B->Time;
+	});
+	uint64 BookmarkCount = Bookmarks.Num();
+	if (FirstBookmarkIndex >= BookmarkCount)
+	{
+		return;
+	}
+	uint64 LastBookmarkIndex = Algo::UpperBoundBy(Bookmarks, IntervalEnd, [](const TSharedRef<FBookmarkInternal>& B)
+	{
+		return B->Time;
+	});
+	if (LastBookmarkIndex == 0)
+	{
+		return;
+	}
+	--LastBookmarkIndex;
+	for (uint64 Index = FirstBookmarkIndex; Index <= LastBookmarkIndex; ++Index)
+	{
+		const FBookmarkInternal& InternalBookmark = Bookmarks[Index].Get();
+		FBookmark Bookmark;
+		Bookmark.Time = InternalBookmark.Time;
+		Bookmark.Text = *InternalBookmark.Text;
+		Callback(Bookmark);
+	}
+}
+
+}
\ No newline at end of file
diff --git a/Engine/Source/Developer/TraceServices/Private/Model/Bookmarks.h b/Engine/Source/Developer/TraceServices/Private/Model/Bookmarks.h
new file mode 100644
--- /dev/null
+++ b/Engine/Source/Developer/TraceServices/Private/Model/Bookmarks.h
@@ -0,0 +1,42 @@
+// Copyright 1998-2019 Epic Games, Inc. All Rights Reserved.
+
+#pragma once
+
+#include "TraceServices/AnalysisService.h"
+#include "Templates/SharedPointer.h"
+
+namespace Trace
+{
+
+class FAnalysisSessionLock;
+
+struct FBookmarkSpec
+{
+	FString File;
+	FString FormatString;
+	int32 Line;
+};
+
+struct FBookmarkInternal
+{
+	double Time;
+	FString Text;
+};
+
+class FBookmarkProvider
+	: public IBookmarkProvider
+{
+public:
+	FBookmarkProvider(const FAnalysisSessionLock& SessionLock);
+
+	FBookmarkSpec& GetSpec(uint64 BookmarkPoint);
+	void AppendBookmark(double Time, uint64 BookmarkPoint, uint16 FormatArgsSize, const uint8* FormatArgs);
+	virtual void EnumerateBookmarks(double IntervalStart, double IntervalEnd, TFunctionRef<void(const FBookmark&)> Callback) const override;
+
+private:
+	const FAnalysisSessionLock& SessionLock;
+	TMap<uint64, TSharedPtr<FBookmarkSpec>> SpecMap;
+	TArray<TSharedRef<FBookmarkInternal>> Bookmarks;
+};
+
+}
\ No newline at end of file
diff --git a/Engine/Source/Developer/TraceServices/Private/Model/Counters.cpp b/Engine/Source/Developer/TraceServices/Private/Model/Counters.cpp
new file mode 100644
--- /dev/null
+++ b/Engine/Source/Developer/TraceServices/Private/Model/Counters.cpp
@@ -0,0 +1,172 @@
+// Copyright 1998-2019 Epic Games, Inc. All Rights Reserved.
+
+#include "Model/Counters.h"
+#include "AnalysisServicePrivate.h"
+
+namespace Trace
+{
+
+FCounterInternal::FCounterInternal(FSlabAllocator& InAllocator, uint32 InId, const TCHAR* InName, const TCHAR* InDescription, ECounterDisplayHint InDisplayHint)
+	: Name(InName)
+	, Description(InDescription)
+	, Id(InId)
+	, DisplayHint(InDisplayHint)
+	, Timestamps(InAllocator, 1024)
+	, Ops(InAllocator, 1024)
+	, IntOpArguments(InAllocator, 1024)
+	, DoubleOpArguments(InAllocator, 1024)
+{
+
+}
+
+template<typename CounterType, typename EnumerationType>
+static void EnumerateCounterValuesInternal(const TPagedArray<double>& Timestamps, const TPagedArray<ECounterOpType>& Ops, const TPagedArray<CounterType>& OpArgs, double IntervalStart, double IntervalEnd, TFunctionRef<void(double, EnumerationType)> Callback)
+{
+	CounterType Value = 0;
+	auto TimeIterator = Timestamps.GetIteratorFromItem(0);
+	auto OpIterator = Ops.GetIteratorFromItem(0);
+	auto OpArgIterator = OpArgs.GetIteratorFromItem(0);
+	const double* Time = TimeIterator.GetCurrentItem();
+	const ECounterOpType* Op = OpIterator.GetCurrentItem();
+	const CounterType* OpArg = OpArgIterator.GetCurrentItem();
+	
+	double PrevTime = 0.0;
+	CounterType PrevValue = CounterType();
+	bool IsFirstValue = true;
+	bool IsFirstIncludedValue = true;
+	while (Time)
+	{
+		switch (*Op)
+		{
+		case CounterOpType_Add:
+			Value += *OpArg;
+			break;
+		case CounterOpType_Set:
+			Value = *OpArg;
+			break;
+		}
+
+		if (*Time >= IntervalStart)
+		{
+			if (IsFirstIncludedValue && !IsFirstValue)
+			{
+				Callback(PrevTime, PrevValue);
+			}
+			PrevTime = *Time;
+			PrevValue = EnumerationType(Value);
+			Callback(PrevTime, PrevValue);
+			IsFirstIncludedValue = false;
+		}
+		if (*Time > IntervalEnd)
+		{
+			break;
+		}
+
+		IsFirstValue = false;
+
+		Time = TimeIterator.NextItem();
+		Op = OpIterator.NextItem();
+		OpArg = OpArgIterator.NextItem();
+	}
+}
+
+void FCounterInternal::EnumerateValues(double IntervalStart, double IntervalEnd, TFunctionRef<void(double, int64)> Callback) const
+{
+	if (DisplayHint == CounterDisplayHint_FloatingPoint)
+	{
+		EnumerateCounterValuesInternal<double, int64>(Timestamps, Ops, DoubleOpArguments, IntervalStart, IntervalEnd, Callback);
+	}
+	else
+	{
+		EnumerateCounterValuesInternal<int64, int64>(Timestamps, Ops, IntOpArguments, IntervalStart, IntervalEnd, Callback);
+	}
+}
+
+void FCounterInternal::EnumerateFloatValues(double IntervalStart, double IntervalEnd, TFunctionRef<void(double, double)> Callback) const
+{
+	if (DisplayHint == CounterDisplayHint_FloatingPoint)
+	{
+		EnumerateCounterValuesInternal<double, double>(Timestamps, Ops, DoubleOpArguments, IntervalStart, IntervalEnd, Callback);
+	}
+	else
+	{
+		EnumerateCounterValuesInternal<int64, double>(Timestamps, Ops, IntOpArguments, IntervalStart, IntervalEnd, Callback);
+	}
+}
+
+FCountersProvider::FCountersProvider(FSlabAllocator& InAllocator, FAnalysisSessionLock& InSessionLock)
+	: Allocator(InAllocator)
+	, SessionLock(InSessionLock)
+{
+
+}
+
+FCountersProvider::~FCountersProvider()
+{
+	for (FCounterInternal* Counter : Counters)
+	{
+		delete Counter;
+	}
+}
+
+void FCountersProvider::EnumerateCounters(TFunctionRef<void(const ICounter &)> Callback) const
+{
+	uint32 Id = 0;
+	for (FCounterInternal* Counter : Counters)
+	{
+		Callback(*Counter);
+	}
+}
+
+FCounterInternal* FCountersProvider::CreateCounter(const TCHAR* Name, const TCHAR* Description, ECounterDisplayHint DisplayHint)
+{
+	FCounterInternal* Counter = new FCounterInternal(Allocator, Counters.Num(), Name, Description, DisplayHint);
+	Counters.Add(Counter);
+	return Counter;
+}
+
+template<typename T>
+static void InsertOp(FCounterInternal& Counter, double Time, ECounterOpType Type, T Arg)
+{
+	TPagedArray<double>& Timestamps = Counter.EditTimestamps();
+	uint64 Count = Timestamps.Num();
+	uint64 InsertionIndex = Count;
+	while (InsertionIndex > 0 && Timestamps[InsertionIndex - 1] > Time)
+	{
+		--InsertionIndex;
+	}
+	Timestamps.Insert(InsertionIndex) = Time;
+	Counter.EditOps().Insert(InsertionIndex) = Type;
+	if (Counter.GetDisplayHint() == CounterDisplayHint_FloatingPoint)
+	{
+		Counter.EditDoubleOpArguments().Insert(InsertionIndex) = double(Arg);
+	}
+	else
+	{
+		Counter.EditIntOpArguments().Insert(InsertionIndex) = int64(Arg);
+	}
+}
+
+void FCountersProvider::Add(FCounterInternal& Counter, double Time, int64 Value)
+{
+	InsertOp(Counter, Time, CounterOpType_Add, Value);
+}
+
+void FCountersProvider::Add(FCounterInternal& Counter, double Time, double Value)
+{
+	InsertOp(Counter, Time, CounterOpType_Add, Value);
+}
+
+void FCountersProvider::Set(FCounterInternal& Counter, double Time, int64 Value)
+{
+	InsertOp(Counter, Time, CounterOpType_Set, Value);
+}
+
+void FCountersProvider::Set(FCounterInternal& Counter, double Time, double Value)
+{
+	InsertOp(Counter, Time, CounterOpType_Set, Value);
+}
+
+
+
+}
\ No newline at end of file
diff --git a/Engine/Source/Developer/TraceServices/Private/Model/Counters.h b/Engine/Source/Developer/TraceServices/Private/Model/Counters.h
new file mode 100644
--- /dev/null
+++ b/Engine/Source/Developer/TraceServices/Private/Model/Counters.h
@@ -0,0 +1,71 @@
+// Copyright 1998-2019 Epic Games, Inc. All Rights Reserved.
+
+#pragma once
+
+#include "TraceServices/AnalysisService.h"
+#include "Templates/SharedPointer.h"
+#include "Common/PagedArray.h"
+
+namespace Trace
+{
+
+class FAnalysisSessionLock;
+
+enum ECounterOpType : uint8
+{
+	CounterOpType_Add,
+	CounterOpType_Set,
+};
+
+class FCounterInternal
+	: public ICounter
+{
+public:
+	FCounterInternal(FSlabAllocator& InAllocator, uint32 Id, const TCHAR* Name, const TCHAR* Description, ECounterDisplayHint DisplayHint);
+	virtual const TCHAR* GetName() const override { return *Name; }
+	virtual uint32 GetId() const override { return Id; }
+	virtual ECounterDisplayHint GetDisplayHint() const { return DisplayHint; }
+	virtual void EnumerateValues(double IntervalStart, double IntervalEnd, TFunctionRef<void(double, int64)> Callback) const override;
+	virtual void EnumerateFloatValues(double IntervalStart, double IntervalEnd, TFunctionRef<void(double, double)> Callback) const override;
+
+	const TPagedArray<int64>& ReadIntOpArguments() const { return IntOpArguments; }
+	const TPagedArray<double>& ReadDoubleOpArguments() const { return DoubleOpArguments; }
+	TPagedArray<double>& EditTimestamps() { return Timestamps; }
+	TPagedArray<ECounterOpType>& EditOps() { return Ops; }
+	TPagedArray<int64>& EditIntOpArguments() { return IntOpArguments; }
+	TPagedArray<double>& EditDoubleOpArguments() { return DoubleOpArguments; }
+
+private:
+	friend class FCountersProvider;
+
+	FString Name;
+	FString Description;
+	uint32 Id;
+	ECounterDisplayHint DisplayHint;
+	TPagedArray<double> Timestamps;
+	TPagedArray<ECounterOpType> Ops;
+	TPagedArray<int64> IntOpArguments;
+	TPagedArray<double> DoubleOpArguments;
+};
+
+class FCountersProvider
+	: public ICountersProvider
+{
+public:
+	FCountersProvider(FSlabAllocator& Allocator, FAnalysisSessionLock& SessionLock);
+	virtual ~FCountersProvider();
+	virtual void EnumerateCounters(TFunctionRef<void(const ICounter&)> Callback) const override;
+
+	FCounterInternal* CreateCounter(const TCHAR* Name, const TCHAR* Description, ECounterDisplayHint DisplayHint);
+	void Add(FCounterInternal& Counter, double Time, int64 Value);
+	void Add(FCounterInternal& Counter, double Time, double Value);
+	void Set(FCounterInternal& Counter, double Time, int64 Value);
+	void Set(FCounterInternal& Counter, double Time, double Value);
+
+private:
+	FSlabAllocator& Allocator;
+	FAnalysisSessionLock& SessionLock;
+	TArray<FCounterInternal*> Counters;
+};
+
+}
diff --git a/Engine/Source/Developer/TraceServices/Private/Model/FileActivity.cpp b/Engine/Source/Developer/TraceServices/Private/Model/FileActivity.cpp
new file mode 100644
--- /dev/null
+++ b/Engine/Source/Developer/TraceServices/Private/Model/FileActivity.cpp
@@ -0,0 +1,64 @@
+// Copyright 1998-2019 Epic Games, Inc. All Rights Reserved.
+
+#include "Model/FileActivity.h"
+#include "Model/IntervalTimeline.h"
+#include "AnalysisServicePrivate.h"
+
+namespace Trace
+{
+
+FFileActivityProvider::FFileActivityProvider(FSlabAllocator& InAllocator, FAnalysisSessionLock& InSessionLock)
+	: Allocator(InAllocator)
+	, SessionLock(InSessionLock)
+	, Files(Allocator, 1024)
+{
+}
+
+void FFileActivityProvider::EnumerateFileActivity(TFunctionRef<bool(const FFileInfo&, const ITimeline<FFileActivity>&)> Callback) const
+{
+	for (uint64 FileIndex = 0, FileCount = Files.Num(); FileIndex < FileCount; ++FileIndex)
+	{
+		const FFileInfoInternal& FileInfoInternal = Files[FileIndex];
+		FFileInfo FileInfo;
+		FileInfo.Id = FileIndex;
+		FileInfo.Path = *FileInfoInternal.Path;
+		if (!Callback(FileInfo, *FileInfoInternal.ActivityTimeline))
+		{
+			return;
+		}
+	}
+}
+
+uint32 FFileActivityProvider::GetFileIndex(const TCHAR* Path)
+{
+	uint32 FileIndex = Files.Num();
+	FFileInfoInternal& FileInfo = Files.PushBack();
+	FileInfo.Path = Path;
+	FileInfo.ActivityTimeline = MakeShared<TimelineInternal>(Allocator);
+	return FileIndex;
+}
+
+uint64 FFileActivityProvider::BeginActivity(uint32 FileIndex, EFileActivityType Type, uint64 Offset, uint64 Size, double Time)
+{
+	FFileInfoInternal& FileInfo = Files[FileIndex];
+	FFileActivity FileActivity;
+	FileActivity.ActivityType = Type;
+	FileActivity.Offset = Offset;
+	FileActivity.Size = Size;
+	FileActivity.Failed = false;
+	return FileInfo.ActivityTimeline->AppendBeginEvent(Time, FileActivity);
+}
+
+void FFileActivityProvider::EndActivity(uint32 FileIndex, uint64 ActivityIndex, double Time, bool Failed)
+{
+	FFileInfoInternal& FileInfo = Files[FileIndex];
+	FFileActivity& Activity = FileInfo.ActivityTimeline->EndEvent(ActivityIndex, Time);
+	Activity.Failed = Failed;
+}
+
+const TCHAR* FFileActivityProvider::GetFilePath(uint32 FileIndex) const
+{
+	return *Files[FileIndex].Path;
+}
+
+}
\ No newline at end of file
diff --git a/Engine/Source/Developer/TraceServices/Private/Model/FileActivity.h b/Engine/Source/Developer/TraceServices/Private/Model/FileActivity.h
new file mode 100644
--- /dev/null
+++ b/Engine/Source/Developer/TraceServices/Private/Model/FileActivity.h
@@ -0,0 +1,46 @@
+// Copyright 1998-2019 Epic Games, Inc. All Rights Reserved.
+
+#pragma once
+
+#include "TraceServices/AnalysisService.h"
+#include "Common/PagedArray.h"
+#include "Model/IntervalTimeline.h"
+
+namespace Trace
+{
+
+class FAnalysisSessionLock;
+
+class FFileActivityProvider
+	: public IFileActivityProvider
+{
+public:
+	struct FTimelineSettings
+	{
+		enum
+		{
+			EventsPerPage = 128
+		};
+	};
+	typedef TIntervalTimeline<FFileActivity, FTimelineSettings> TimelineInternal;
+
+	FFileActivityProvider(FSlabAllocator& Allocator, FAnalysisSessionLock& SessionLock);
+	virtual void EnumerateFileActivity(TFunctionRef<bool(const FFileInfo&, const Timeline&)> Callback) const override;
+	uint32 GetFileIndex(const TCHAR* Path);
+	uint64 BeginActivity(uint32 FileIndex, EFileActivityType Type, uint64 Offset, uint64 Size, double Time);
+	void EndActivity(uint32 FileIndex, uint64 ActivityIndex, double Time, bool Failed);
+	const TCHAR* GetFilePath(uint32 FileIndex) const;
+
+private:
+	struct FFileInfoInternal
+	{
+		FString Path;
+		TSharedPtr<TimelineInternal> ActivityTimeline;
+	};
+
+	FSlabAllocator& Allocator;
+	FAnalysisSessionLock& SessionLock;
+	TPagedArray<FFileInfoInternal> Files;
+};
+
+}
diff --git a/Engine/Source/Developer/TraceServices/Private/Model/Frames.cpp b/Engine/Source/Developer/TraceServices/Private/Model/Frames.cpp
new file mode 100644
--- /dev/null
+++ b/Engine/Source/Developer/TraceServices/Private/Model/Frames.cpp
@@ -0,0 +1,62 @@
+// Copyright 1998-2019 Epic Games, Inc. All Rights Reserved.
+
+#include "Model/Frames.h"
+#include "AnalysisServicePrivate.h"
+#include <limits>
+
+namespace Trace
+{
+
+FFrameProvider::FFrameProvider(FSlabAllocator& InAllocator, FAnalysisSessionLock& InSessionLock)
+	: Allocator(InAllocator)
+	, SessionLock(InSessionLock)
+{
+	for (int32 FrameType = 0; FrameType < TraceFrameType_Count; ++FrameType)
+	{
+		Frames.Emplace(Allocator, 65536);
+	}
+}
+
+uint64 FFrameProvider::GetFrameCount(ETraceFrameType FrameType) const
+{
+	SessionLock.ReadAccessCheck();
+	return Frames[FrameType].Num();
+}
+
+void FFrameProvider::EnumerateFrames(ETraceFrameType FrameType, uint64 Start, uint64 End, TFunctionRef<void(const FFrame&)> Callback) const
+{
+	SessionLock.ReadAccessCheck();
+
+	End = FMath::Min(End, Frames[FrameType].Num());
+	if (Start >= End)
+	{
+		return;
+	}
+	auto Iterator = Frames[FrameType].GetIteratorFromItem(Start);
+	const FFrame* Frame = Iterator.GetCurrentItem();
+	for (uint64 Index = Start; Index < End; ++Index)
+	{
+		Callback(*Frame);
+		Frame = Iterator.NextItem();
+	}
+}
+
+void FFrameProvider::BeginFrame(ETraceFrameType FrameType, double Time)
+{
+	SessionLock.WriteAccessCheck();
+	
+	uint64 Index = Frames[FrameType].Num();
+	FFrame& Frame = Frames[FrameType].PushBack();
+	Frame.StartTime = Time;
+	Frame.EndTime = std::numeric_limits<double>::infinity();
+	Frame.Index = Index;
+}
+
+void FFrameProvider::EndFrame(ETraceFrameType FrameType, double Time)
+{
+	SessionLock.WriteAccessCheck();
+	FFrame& Frame = Frames[FrameType][Frames[FrameType].Num() - 1];
+	Frame.EndTime = Time;
+}
+
+}
\ No newline at end of file
diff --git a/Engine/Source/Developer/TraceServices/Private/Model/Frames.h b/Engine/Source/Developer/TraceServices/Private/Model/Frames.h
new file mode 100644
--- /dev/null
+++ b/Engine/Source/Developer/TraceServices/Private/Model/Frames.h
@@ -0,0 +1,31 @@
+// Copyright 1998-2019 Epic Games, Inc. All Rights Reserved.
+
+#pragma once
+
+#include "TraceServices/AnalysisService.h"
+#include "ProfilingDebugging/MiscTrace.h"
+#include "Common/PagedArray.h"
+
+namespace Trace
+{
+
+class FAnalysisSessionLock;
+
+class FFrameProvider
+	: public IFrameProvider
+{
+public:
+	FFrameProvider(FSlabAllocator& Allocator, FAnalysisSessionLock& SessionLock);
+
+	virtual uint64 GetFrameCount(ETraceFrameType FrameType) const override;
+	virtual void EnumerateFrames(ETraceFrameType FrameType, uint64 Start, uint64 End, TFunctionRef<void(const FFrame&)> Callback) const override;
+	void BeginFrame(ETraceFrameType FrameType, double Time);
+	void EndFrame(ETraceFrameType FrameType, double Time);
+
+private:
+	FSlabAllocator& Allocator;
+	FAnalysisSessionLock& SessionLock;
+	TArray<TPagedArray<FFrame>> Frames;
+};
+
+}
\ No newline at end of file
diff --git a/Engine/Source/Developer/TraceServices/Private/Model/IntervalTimeline.h b/Engine/Source/Developer/TraceServices/Private/Model/IntervalTimeline.h
new file mode 100644
--- /dev/null
+++ b/Engine/Source/Developer/TraceServices/Private/Model/IntervalTimeline.h
@@ -0,0 +1,126 @@
+// Copyright 1998-2019 Epic Games, Inc. All Rights Reserved.
+
+#pragma once
+
+#include "TraceServices/AnalysisService.h"
+#include "Common/PagedArray.h"
+#include <limits>
+
+namespace Trace
+{
+
+struct FIntervalTimelineDefaultSettings
+{
+	enum
+	{
+		EventsPerPage = 65536
+	};
+};
+
+template<typename EventType, typename SettingsType = FIntervalTimelineDefaultSettings>
+class TIntervalTimeline
+	: public ITimeline<EventType>
+{
+public:
+	TIntervalTimeline(FSlabAllocator& Allocator)
+		: Events(Allocator, SettingsType::EventsPerPage)
+	{
+	}
+
+	virtual ~TIntervalTimeline() = default;
+	virtual uint64 GetEventCount() const override { return Events.Num(); }
+	
+	virtual void EnumerateEventsDownSampled(double IntervalStart, double IntervalEnd, double Resolution, TFunctionRef<void(bool, double, const EventType&)> Callback) const override { check(false); }
+	virtual void EnumerateEventsDownSampled(double IntervalStart, double IntervalEnd, double Resolution, TFunctionRef<void(double, double, uint32, const EventType&)> Callback) const override { check(false); }
+	virtual void EnumerateEvents(double IntervalStart, double IntervalEnd, TFunctionRef<void(bool, double, const EventType&)> Callback) const override { check(false); }
+	
+	virtual void EnumerateEvents(double IntervalStart, double IntervalEnd, TFunctionRef<void(double, double, uint32, const EventType&)> Callback) const override
+	{
+		auto EventIterator = Events.GetIteratorFromPage(0);
+		const FEventInternal* Event = EventIterator.GetCurrentItem();
+		while (Event)
+		{
+			const FEventPage* Page = EventIterator.GetCurrentPage();
+			if (Page->EndTime < IntervalStart || IntervalEnd < Page->BeginTime)
+			{
+				if (!EventIterator.NextPage())
+				{
+					return;
+				}
+				Event = EventIterator.GetCurrentItem();
+			}
+			else
+			{
+				if (!(Event->EndTime < IntervalStart || IntervalEnd < Event->StartTime))
+				{
+					Callback(Event->StartTime, Event->EndTime, 0, Event->Event);
+				}
+				Event = EventIterator.NextItem();
+			}
+		}
+	}
+	
+	uint64 AppendBeginEvent(double StartTime, const EventType& Event)
+	{
+		//check(StartTime >= LastTime);
+		//LastTime = StartTime;
+
+		uint64 Index = Events.Num();
+		FEventInternal& EventInternal = Events.PushBack();
+		EventInternal.StartTime = StartTime;
+		EventInternal.EndTime = std::numeric_limits<double>::infinity();
+		EventInternal.Event = Event;
+
+		FEventPage* LastPage = Events.GetLastPage();
+		if (LastPage != CurrentPage)
+		{
+			CurrentPage = LastPage;
+			LastPage->BeginTime = StartTime;
+			LastPage->EndTime = StartTime;
+			LastPage->BeginEventIndex = Index;
+		}
+		else
+		{
+			LastPage->BeginTime = FMath::Min(LastPage->BeginTime, StartTime);
+			LastPage->EndTime = FMath::Max(LastPage->EndTime, StartTime);
+		}
+
+		LastPage->EndEventIndex = Index + 1;
+
+		return Index;
+	}
+
+	EventType& EndEvent(uint64 EventIndex, double EndTime)
+	{
+		//check(EndTime >= LastTime);
+		//LastTime = EndTime;
+		FEventInternal& EventInternal = Events[EventIndex];
+		check(EventIndex < Events.Num());
+		EventInternal.EndTime = EndTime;
+		FEventPage* Page = Events.GetItemPage(EventIndex);
+		Page->EndTime = FMath::Max(Page->EndTime, EndTime);
+		return EventInternal.Event;
+	}
+
+private:
+	struct FEventPage
+	{
+		double BeginTime = 0.0;
+		double EndTime = 0.0;
+		uint64 BeginEventIndex = 0;
+		uint64 EndEventIndex = 0;
+	};
+
+	struct FEventInternal
+	{
+		double StartTime;
+		double EndTime;
+		EventType Event;
+	};
+
+	TPagedArray<FEventInternal, FEventPage> Events;
+	FEventPage* CurrentPage = nullptr;
+	double LastTime = 0.0;
+};
+
+}
\ No newline at end of file
diff --git a/Engine/Source/Developer/TraceServices/Private/Model/LoadTimeProfiler.cpp b/Engine/Source/Developer/TraceServices/Private/Model/LoadTimeProfiler.cpp
new file mode 100644
--- /dev/null
+++ b/Engine/Source/Developer/TraceServices/Private/Model/LoadTimeProfiler.cpp
@@ -0,0 +1,66 @@
+// Copyright 1998-2019 Epic Games, Inc. All Rights Reserved.
+
+#include "Model/LoadTimeProfiler.h"
+#include "AnalysisServicePrivate.h"
+
+namespace Trace
+{
+
+FLoadTimeProfilerProvider::FLoadTimeProfilerProvider(FSlabAllocator& InAllocator, FAnalysisSessionLock& InSessionLock)
+	: Allocator(InAllocator)
+	, SessionLock(InSessionLock)
+	, Packages(Allocator, 4096)
+	, Exports(Allocator, 4090)
+	, MainThreadCpuTimeline(MakeShared<CpuTimelineInternal>(Allocator))
+	, AsyncLoadingThreadCpuTimeline(MakeShared<CpuTimelineInternal>(Allocator))
+{
+	
+}
+
+void FLoadTimeProfilerProvider::EnumeratePackages(TFunctionRef<void(const FPackageInfo&)> Callback) const
+{
+	SessionLock.ReadAccessCheck();
+
+	auto Iterator = Packages.GetIteratorFromItem(0);
+	const FPackageInfo* Package = Iterator.GetCurrentItem();
+	while (Package)
+	{
+		Callback(*Package);
+		Package = Iterator.NextItem();
+	}
+}
+
+void FLoadTimeProfilerProvider::ReadMainThreadCpuTimeline(TFunctionRef<void(const CpuTimeline &)> Callback) const
+{
+	SessionLock.ReadAccessCheck();
+	Callback(*MainThreadCpuTimeline);
+}
+
+void FLoadTimeProfilerProvider::ReadAsyncLoadingThreadCpuTimeline(TFunctionRef<void(const CpuTimeline &)> Callback) const
+{
+	SessionLock.ReadAccessCheck();
+	Callback(*AsyncLoadingThreadCpuTimeline);
+}
+
+Trace::FPackageInfo& FLoadTimeProfilerProvider::CreatePackage(const TCHAR* PackageName)
+{
+	SessionLock.WriteAccessCheck();
+
+	uint32 PackageId = Packages.Num();
+	FPackageInfo& Package = Packages.PushBack();
+	Package.Id = PackageId;
+	Package.Name = PackageName;
+	return Package;
+}
+
+Trace::FPackageExportInfo& FLoadTimeProfilerProvider::CreateExport()
+{
+	SessionLock.WriteAccessCheck();
+
+	uint32 ExportId = Exports.Num();
+	FPackageExportInfo& Export = Exports.PushBack();
+	Export.Id = ExportId;
+	return Export;
+}
+
+}
\ No newline at end of file
diff --git a/Engine/Source/Developer/TraceServices/Private/Model/LoadTimeProfiler.h b/Engine/Source/Developer/TraceServices/Private/Model/LoadTimeProfiler.h
new file mode 100644
--- /dev/null
+++ b/Engine/Source/Developer/TraceServices/Private/Model/LoadTimeProfiler.h
@@ -0,0 +1,39 @@
+// Copyright 1998-2019 Epic Games, Inc. All Rights Reserved.
+
+#pragma once
+
+#include "TraceServices/AnalysisService.h"
+#include "Common/SlabAllocator.h"
+#include "Common/PagedArray.h"
+#include "Model/MonotonicTimeline.h"
+
+namespace Trace
+{
+
+class FAnalysisSessionLock;
+
+class FLoadTimeProfilerProvider
+	: public ILoadTimeProfilerProvider
+{
+public:
+	typedef TMonotonicTimeline<FLoadTimeProfilerCpuEvent> CpuTimelineInternal;
+
+	FLoadTimeProfilerProvider(FSlabAllocator& Allocator, FAnalysisSessionLock& InSessionLock);
+	virtual void EnumeratePackages(TFunctionRef<void(const FPackageInfo&)> Callback) const override;
+	virtual void ReadMainThreadCpuTimeline(TFunctionRef<void(const CpuTimeline&)> Callback) const override;
+	virtual void ReadAsyncLoadingThreadCpuTimeline(TFunctionRef<void(const CpuTimeline&)> Callback) const override;
+	FPackageInfo& CreatePackage(const TCHAR* PackageName);
+	FPackageExportInfo& CreateExport();
+	TSharedRef<CpuTimelineInternal> EditMainThreadCpuTimeline() { return MainThreadCpuTimeline; }
+	TSharedRef<CpuTimelineInternal> EditAsyncLoadingThreadCpuTimeline() { return AsyncLoadingThreadCpuTimeline; }
+
+private:
+	FSlabAllocator& Allocator;
+	FAnalysisSessionLock& SessionLock;
+	TPagedArray<FPackageInfo> Packages;
+	TPagedArray<FPackageExportInfo> Exports;
+	TSharedRef<CpuTimelineInternal> MainThreadCpuTimeline;
+	TSharedRef<CpuTimelineInternal> AsyncLoadingThreadCpuTimeline;
+};
+
+}
diff --git a/Engine/Source/Developer/TraceServices/Private/Model/Log.cpp b/Engine/Source/Developer/TraceServices/Private/Model/Log.cpp
new file mode 100644
--- /dev/null
+++ b/Engine/Source/Developer/TraceServices/Private/Model/Log.cpp
@@ -0,0 +1,165 @@
+// Copyright 1998-2019 Epic Games, Inc. All Rights Reserved.
+
+#include "Model/Log.h"
+#include "AnalysisServicePrivate.h"
+#include "Common/FormatArgs.h"
+
+namespace Trace
+{
+
+FLogProvider::FLogProvider(FSlabAllocator& InAllocator, FAnalysisSessionLock& InSessionLock)
+	: Allocator(InAllocator)
+	, SessionLock(InSessionLock)
+	, Categories(Allocator, 128)
+	, MessageSpecs(Allocator, 1024)
+{
+
+}
+
+FLogCategory& FLogProvider::GetCategory(uint64 CategoryPointer)
+{
+	SessionLock.WriteAccessCheck();
+	if (CategoryMap.Contains(CategoryPointer))
+	{
+		return *CategoryMap[CategoryPointer];
+	}
+	else
+	{
+		FLogCategory& Category = Categories.PushBack();
+		CategoryMap.Add(CategoryPointer, &Category);
+		return Category;
+	}
+}
+
+FLogMessageSpec& FLogProvider::GetMessageSpec(uint64 LogPoint)
+{
+	SessionLock.WriteAccessCheck();
+	if (SpecMap.Contains(LogPoint))
+	{
+		return *SpecMap[LogPoint];
+	}
+	else
+	{
+		FLogMessageSpec& Spec = MessageSpecs.PushBack();
+		SpecMap.Add(LogPoint, &Spec);
+		return Spec;
+	}
+}
+
+void FLogProvider::AppendMessage(uint64 LogPoint, double Time, uint16 FormatArgsSize, const uint8* FormatArgs)
+{
+	SessionLock.WriteAccessCheck();
+	check(SpecMap.Contains(LogPoint));
+	FLogMessageInternal& InternalMessage = Messages.AddDefaulted_GetRef();
+	InternalMessage.Time = Time;
+	InternalMessage.Spec = SpecMap[LogPoint];
+	InternalMessage.FormatArgsOffset = FormatArgsMemory.Num();
+	FormatArgsMemory.AddUninitialized(FormatArgsSize);
+	memcpy(FormatArgsMemory.GetData() + InternalMessage.FormatArgsOffset, FormatArgs, FormatArgsSize);
+}
+
+uint64 FLogProvider::GetMessageCount() const
+{
+	SessionLock.ReadAccessCheck();
+	return Messages.Num();
+}
+
+bool FLogProvider::ReadMessage(uint64 Index, TFunctionRef<void(const FLogMessage &)> Callback, bool ResolveFormatString) const
+{
+	SessionLock.ReadAccessCheck();
+	if (Index >= Messages.Num())
+	{
+		return false;
+	}
+	ConstructMessage(Index, Callback, ResolveFormatString);
+	return true;
+}
+
+void FLogProvider::EnumerateMessages(double IntervalStart, double IntervalEnd, TFunctionRef<void(const FLogMessage&)> Callback, bool ResolveFormatString) const
+{
+	SessionLock.ReadAccessCheck();
+	if (IntervalStart > IntervalEnd)
+	{
+		return;
+	}
+	uint64 FirstMessageIndex = Algo::LowerBoundBy(Messages, IntervalStart, [](const FLogMessageInternal& M)
+	{
+		return M.Time;
+	});
+	uint64 MessageCount = Messages.Num();
+	if (FirstMessageIndex >= MessageCount)
+	{
+		return;
+	}
+	uint64 LastMessageIndex = Algo::UpperBoundBy(Messages, IntervalEnd, [](const FLogMessageInternal& M)
+	{
+		return M.Time;
+	});
+	if (LastMessageIndex == 0)
+	{
+		return;
+	}
+	--LastMessageIndex;
+	for (uint64 Index = FirstMessageIndex; Index <= LastMessageIndex; ++Index)
+	{
+		ConstructMessage(Index, Callback, ResolveFormatString);
+	}
+}
+
+void FLogProvider::EnumerateMessagesByIndex(uint64 Start, uint64 End, TFunctionRef<void(const FLogMessage &)> Callback, bool ResolveFormatString) const
+{
+	SessionLock.ReadAccessCheck();
+	uint64 Count = Messages.Num();
+	if (Start >= Count)
+	{
+		return;
+	}
+	if (End > Count)
+	{
+		End = Count;
+	}
+	if (Start >= End)
+	{
+		return;
+	}
+	for (uint64 Index = Start; Index < End; ++Index)
+	{
+		ConstructMessage(Index, Callback, ResolveFormatString);
+	}
+}
+
+void FLogProvider::ConstructMessage(uint64 Index, TFunctionRef<void(const FLogMessage &)> Callback, bool ResolveFormatString) const
+{
+	const FLogMessageInternal& InternalMessage = Messages[Index];
+	FLogMessage Message;
+	Message.Index = Index;
+	Message.Time = InternalMessage.Time;
+	Message.Category = InternalMessage.Spec->Category;
+	Message.File = *InternalMessage.Spec->File;
+	Message.Line = InternalMessage.Spec->Line;
+	Message.Verbosity = InternalMessage.Spec->Verbosity;
+	if (ResolveFormatString)
+	{
+		TCHAR Buffer[65535];
+		FFormatArgsHelper::Format(Buffer, 65534, *InternalMessage.Spec->FormatString, FormatArgsMemory.GetData() + InternalMessage.FormatArgsOffset);
+		Message.Message = Buffer;
+	}
+	else
+	{
+		Message.Message = *InternalMessage.Spec->FormatString;
+	}
+	Callback(Message);
+}
+
+void FLogProvider::EnumerateCategories(TFunctionRef<void(const FLogCategory &)> Callback) const
+{
+	auto Iterator = Categories.GetIteratorFromItem(0);
+	const FLogCategory* Category = Iterator.GetCurrentItem();
+	while (Category)
+	{
+		Callback(*Category);
+		Category = Iterator.NextItem();
+	}
+}
+
+}
\ No newline at end of file
diff --git a/Engine/Source/Developer/TraceServices/Private/Model/Log.h b/Engine/Source/Developer/TraceServices/Private/Model/Log.h
new file mode 100644
--- /dev/null
+++ b/Engine/Source/Developer/TraceServices/Private/Model/Log.h
@@ -0,0 +1,64 @@
+// Copyright 1998-2019 Epic Games, Inc. All Rights Reserved.
+
+#pragma once
+
+#include "TraceServices/AnalysisService.h"
+#include "Templates/SharedPointer.h"
+#include "Common/PagedArray.h"
+
+namespace Trace
+{
+
+class FAnalysisSessionLock;
+
+struct FLogMessageSpec
+{
+	FLogCategory* Category;
+	FString File;
+	FString FormatString;
+	int32 Line;
+	ELogVerbosity::Type Verbosity;
+};
+
+struct FLogMessageInternal
+{
+	FLogMessageSpec* Spec;
+	double Time;
+	uint64 FormatArgsOffset;
+};
+
+class FLogProvider :
+	public ILogProvider
+{
+public:
+	enum
+	{
+		ReservedLogCategory_Bookmark = 0
+	};
+	
+	FLogProvider(FSlabAllocator& Allocator, FAnalysisSessionLock& SessionLock);
+
+	FLogCategory& GetCategory(uint64 CategoryPointer);
+	FLogMessageSpec& GetMessageSpec(uint64 LogPoint);
+	void AppendMessage(uint64 LogPoint, double Time, uint16 FormatArgsSize, const uint8* FormatArgs);
+
+	virtual uint64 GetMessageCount() const override;
+	virtual bool ReadMessage(uint64 Index, TFunctionRef<void(const FLogMessage &)> Callback, bool ResolveFormatString) const override;
+	virtual void EnumerateMessages(double IntervalStart, double IntervalEnd, TFunctionRef<void(const FLogMessage&)> Callback, bool ResolveFormatString) const override;
+	virtual void EnumerateMessagesByIndex(uint64 Start, uint64 End, TFunctionRef<void(const FLogMessage&)> Callback, bool ResolveFormatString) const override;
+	virtual void EnumerateCategories(TFunctionRef<void(const FLogCategory&)> Callback) const override;
+
+private:
+	void ConstructMessage(uint64 Id, TFunctionRef<void(const FLogMessage &)> Callback, bool ResolveFormatString) const;
+	
+	FSlabAllocator& Allocator;
+	FAnalysisSessionLock& SessionLock;
+	TMap<uint64, FLogCategory*> CategoryMap;
+	TMap<uint64, FLogMessageSpec*> SpecMap;
+	TPagedArray<FLogCategory> Categories;
+	TPagedArray<FLogMessageSpec> MessageSpecs;
+	TArray<FLogMessageInternal> Messages;
+	TArray<uint8> FormatArgsMemory;
+};
+
+}
\ No newline at end of file
diff --git a/Engine/Source/Developer/TraceServices/Private/Model/Threads.cpp b/Engine/Source/Developer/TraceServices/Private/Model/Threads.cpp
new file mode 100644
--- /dev/null
+++ b/Engine/Source/Developer/TraceServices/Private/Model/Threads.cpp
@@ -0,0 +1,182 @@
+// Copyright 1998-2019 Epic Games, Inc. All Rights Reserved.
+
+#include "Model/Threads.h"
+#include "Misc/ScopeLock.h"
+#include "AnalysisServicePrivate.h"
+
+namespace Trace
+{
+
+FThreadProvider::FThreadProvider(const FAnalysisSessionLock& InSessionLock)
+	: SessionLock(InSessionLock)
+{
+
+}
+
+void FThreadProvider::EnsureThreadExists(uint32 Id)
+{
+	SessionLock.WriteAccessCheck();
+
+	if (!ThreadMap.Contains(Id))
+	{
+		TSharedRef<FThreadInfoInternal> ThreadInfo = MakeShared<FThreadInfoInternal>();
+		ThreadInfo->Id = Id;
+		ThreadInfo->FallbackSortOrder = SortedThreads.Num();
+		SortedThreads.Add(ThreadInfo);
+		ThreadMap.Add(Id, ThreadInfo);
+		SortThreads();
+	}
+}
+
+void FThreadProvider::AddGameThread(uint32 Id)
+{
+	SessionLock.WriteAccessCheck();
+
+	check(!ThreadMap.Contains(Id));
+	TSharedRef<FThreadInfoInternal> ThreadInfo = MakeShared<FThreadInfoInternal>();
+	ThreadInfo->Id = Id;
+	ThreadInfo->PrioritySortOrder = GetPrioritySortOrder(TPri_Normal);
+	ThreadInfo->Name = *FName(NAME_GameThread).GetPlainNameString();
+	ThreadInfo->FallbackSortOrder = SortedThreads.Num();
+	ThreadInfo->IsGameThread = true;
+	SortedThreads.Add(ThreadInfo);
+	ThreadMap.Add(Id, ThreadInfo);
+}
+
+void FThreadProvider::AddThread(uint32 Id, const FString& Name, EThreadPriority Priority)
+{
+	SessionLock.WriteAccessCheck();
+
+	TSharedPtr<FThreadInfoInternal> ThreadInfo;
+	if (!ThreadMap.Contains(Id))
+	{
+		ThreadInfo = MakeShared<FThreadInfoInternal>();
+		ThreadInfo->Id = Id;
+		ThreadInfo->FallbackSortOrder = SortedThreads.Num();
+		SortedThreads.Add(ThreadInfo.ToSharedRef());
+		ThreadMap.Add(Id, ThreadInfo.ToSharedRef());
+	}
+	else
+	{
+		ThreadInfo = ThreadMap[Id];
+	}
+	ThreadInfo->PrioritySortOrder = GetPrioritySortOrder(Priority);
+	ThreadInfo->Name = Name;
+	SortThreads();
+}
+
+void FThreadProvider::SetThreadPriority(uint32 Id, EThreadPriority Priority)
+{
+	SessionLock.WriteAccessCheck();
+
+	check(ThreadMap.Contains(Id));
+	TSharedRef<FThreadInfoInternal> ThreadInfo = ThreadMap[Id];
+	ThreadInfo->PrioritySortOrder = GetPrioritySortOrder(Priority);
+	SortThreads();
+}
+
+void FThreadProvider::SetThreadGroup(uint32 Id, ETraceThreadGroup Group)
+{
+	SessionLock.WriteAccessCheck();
+
+	check(ThreadMap.Contains(Id));
+	TSharedRef<FThreadInfoInternal> ThreadInfo = ThreadMap[Id];
+	ThreadInfo->GroupSortOrder = GetGroupSortOrder(Group);
+	switch (Group)
+	{
+	case TraceThreadGroup_Render:
+		ThreadInfo->GroupName = TEXT("Render");
+		break;
+	case TraceThreadGroup_TaskGraphHigh:
+		ThreadInfo->GroupName = TEXT("TaskGraph (High Priority)");
+		break;
+	case TraceThreadGroup_TaskGraphNormal:
+		ThreadInfo->GroupName = TEXT("TaskGraph (Normal Priority)");
+		break;
+	case TraceThreadGroup_TaskGraphLow:
+		ThreadInfo->GroupName = TEXT("TaskGraph (Low Priority)");
+		break;
+	case TraceThreadGroup_ThreadPool:
+		ThreadInfo->GroupName = TEXT("ThreadPool");
+		break;
+	case TraceThreadGroup_BackgroundThreadPool:
+		ThreadInfo->GroupName = TEXT("ThreadPool (Background)");
+		break;
+	case TraceThreadGroup_LargeThreadPool:
+		ThreadInfo->GroupName = TEXT("ThreadPool (Large)");
+		break;
+	case TraceThreadGroup_IOThreadPool:
+		ThreadInfo->GroupName = TEXT("ThreadPool (IO)");
+		break;
+	}
+	SortThreads();
+}
+
+void FThreadProvider::EnumerateThreads(TFunctionRef<void(const FThreadInfo &)> Callback) const
+{
+	SessionLock.ReadAccessCheck();
+
+	for (const TSharedRef<FThreadInfoInternal>& Thread : SortedThreads)
+	{
+		FThreadInfo ThreadInfo;
+		ThreadInfo.Id = Thread->Id;
+		ThreadInfo.Name = *Thread->Name;
+		ThreadInfo.GroupName = Thread->GroupName;
+		Callback(ThreadInfo);
+	}
+}
+
+void FThreadProvider::SortThreads()
+{
+	Algo::SortBy(SortedThreads, [](const TSharedRef<FThreadInfoInternal>& In) -> const FThreadInfoInternal&
+	{
+		return In.Get();
+	});
+}
+
+uint32 FThreadProvider::GetGroupSortOrder(ETraceThreadGroup ThreadGroup)
+{
+	return uint32(ThreadGroup);
+}
+
+uint32 FThreadProvider::GetPrioritySortOrder(EThreadPriority ThreadPriority)
+{
+	switch (ThreadPriority)
+	{
+	case TPri_TimeCritical:
+		return 0;
+	case TPri_Highest:
+		return 1;
+	case TPri_AboveNormal:
+		return 2;
+	case TPri_Normal:
+		return 3;
+	case TPri_SlightlyBelowNormal:
+		return 4;
+	case TPri_BelowNormal:
+		return 5;
+	case TPri_Lowest:
+		return 6;
+	default:
+		return 7;
+	}
+}
+
+bool FThreadProvider::FThreadInfoInternal::operator<(const FThreadInfoInternal& Other) const
+{
+	if (IsGameThread == Other.IsGameThread)
+	{
+		if (GroupSortOrder == Other.GroupSortOrder)
+		{
+			if (PrioritySortOrder == Other.PrioritySortOrder)
+			{
+				return FallbackSortOrder < Other.FallbackSortOrder;
+			}
+			return PrioritySortOrder < Other.PrioritySortOrder;
+		}
+		return GroupSortOrder < Other.GroupSortOrder;
+	}
+	return IsGameThread;
+}
+
+}
\ No newline at end of file
diff --git a/Engine/Source/Developer/TraceServices/Private/Model/Threads.h b/Engine/Source/Developer/TraceServices/Private/Model/Threads.h
new file mode 100644
--- /dev/null
+++ b/Engine/Source/Developer/TraceServices/Private/Model/Threads.h
@@ -0,0 +1,48 @@
+// Copyright 1998-2019 Epic Games, Inc. All Rights Reserved.
+
+#pragma once
+
+#include "TraceServices/AnalysisService.h"
+#include "ProfilingDebugging/MiscTrace.h"
+
+namespace Trace
+{
+class FAnalysisSessionLock;
+
+class FThreadProvider
+	: public IThreadProvider
+{
+public:
+	FThreadProvider(const FAnalysisSessionLock& SessionLock);
+	void EnsureThreadExists(uint32 Id);
+	void AddGameThread(uint32 Id);
+	void AddThread(uint32 Id, const FString& Name, EThreadPriority Priority);
+	void SetThreadPriority(uint32 Id, EThreadPriority Priority);
+	void SetThreadGroup(uint32 Id, ETraceThreadGroup Group);
+
+	virtual void EnumerateThreads(TFunctionRef<void(const FThreadInfo&)> Callback) const override;
+
+private:
+	struct FThreadInfoInternal
+	{
+		uint32 Id = 0;
+		FString Name;
+		uint32 GroupSortOrder = ~0ul;
+		uint32 PrioritySortOrder = ~0ul;
+		uint32 FallbackSortOrder = ~0ul;
+		const TCHAR* GroupName = nullptr;
+		bool IsGameThread = false;
+
+		bool operator<(const FThreadInfoInternal& Other) const;
+	};
+
+	void SortThreads();
+	static uint32 GetPrioritySortOrder(EThreadPriority ThreadPriority);
+	static uint32 GetGroupSortOrder(ETraceThreadGroup ThreadGroup);
+
+	const FAnalysisSessionLock& SessionLock;
+	TMap<uint32, TSharedRef<FThreadInfoInternal>> ThreadMap;
+	TArray<TSharedRef<FThreadInfoInternal>> SortedThreads;
+};
+
+}
\ No newline at end of file
diff --git a/Engine/Source/Developer/TraceServices/Private/Model/TimingProfiler.cpp b/Engine/Source/Developer/TraceServices/Private/Model/TimingProfiler.cpp
new file mode 100644
--- /dev/null
+++ b/Engine/Source/Developer/TraceServices/Private/Model/TimingProfiler.cpp
@@ -0,0 +1,129 @@
+// Copyright 1998-2019 Epic Games, Inc. All Rights Reserved.
+
+#include "Model/TimingProfiler.h"
+#include "AnalysisServicePrivate.h"
+
+namespace Trace
+{
+
+FTimingProfilerProvider::FTimingProfilerProvider(FSlabAllocator& InAllocator, FAnalysisSessionLock& InSessionLock)
+	: Allocator(InAllocator)
+	, SessionLock(InSessionLock)
+{
+	Timelines.Add(MakeShared<TimelineInternal>(Allocator));
+}
+
+FTimingProfilerProvider::~FTimingProfilerProvider()
+{
+}
+
+uint32 FTimingProfilerProvider::AddCpuTimer(const TCHAR* Name)
+{
+	SessionLock.WriteAccessCheck();
+
+	FTimingProfilerTimer& Timer = AddTimerInternal(Name, false);
+	return Timer.Id;
+}
+
+uint32 FTimingProfilerProvider::AddGpuTimer(const TCHAR* Name)
+{
+	SessionLock.WriteAccessCheck();
+
+	FTimingProfilerTimer& Timer = AddTimerInternal(Name, true);
+	return Timer.Id;
+}
+
+FTimingProfilerTimer& FTimingProfilerProvider::AddTimerInternal(const TCHAR* Name, bool IsGpuTimer)
+{
+	FTimingProfilerTimer& Timer = Timers.AddDefaulted_GetRef();
+	Timer.Id = Timers.Num() - 1;
+	Timer.Name = Name;
+	uint32 NameHash = 0;
+	for (const TCHAR* c = Name; *c; ++c)
+	{
+		NameHash = (NameHash + *c) * 0x2c2c57ed;
+	}
+	Timer.NameHash = NameHash;
+	Timer.IsGpuTimer = IsGpuTimer;
+	return Timer;
+}
+
+TSharedRef<FTimingProfilerProvider::TimelineInternal> FTimingProfilerProvider::EditCpuThreadTimeline(uint32 ThreadId)
+{
+	SessionLock.WriteAccessCheck();
+
+	if (!CpuThreadTimelineIndexMap.Contains(ThreadId))
+	{
+		TSharedRef<TimelineInternal> Timeline = MakeShared<TimelineInternal>(Allocator);
+		uint32 TimelineIndex = Timelines.Num();
+		CpuThreadTimelineIndexMap.Add(ThreadId, TimelineIndex);
+		Timelines.Add(Timeline);
+		return Timeline;
+	}
+	else
+	{
+		uint32 TimelineIndex = CpuThreadTimelineIndexMap[ThreadId];
+		return Timelines[TimelineIndex];
+	}
+}
+
+TSharedRef<FTimingProfilerProvider::TimelineInternal> FTimingProfilerProvider::EditGpuTimeline()
+{
+	SessionLock.WriteAccessCheck();
+
+	return Timelines[GpuTimelineIndex];
+}
+
+bool FTimingProfilerProvider::GetCpuThreadTimelineIndex(uint32 ThreadId, uint32& OutTimelineIndex) const
+{
+	SessionLock.ReadAccessCheck();
+
+	if (CpuThreadTimelineIndexMap.Contains(ThreadId))
+	{
+		OutTimelineIndex = CpuThreadTimelineIndexMap[ThreadId];
+		return true;
+	}
+	return false;
+}
+
+bool FTimingProfilerProvider::GetGpuTimelineIndex(uint32& OutTimelineIndex) const
+{
+	SessionLock.ReadAccessCheck();
+
+	OutTimelineIndex = GpuTimelineIndex;
+	return true;
+}
+
+bool FTimingProfilerProvider::ReadTimeline(uint32 Index, TFunctionRef<void(const Timeline &)> Callback) const
+{
+	SessionLock.ReadAccessCheck();
+
+	if (Index < uint32(Timelines.Num()))
+	{
+		Callback(*Timelines[Index]);
+		return true;
+	}
+	else
+	{
+		return false;
+	}
+}
+
+void FTimingProfilerProvider::EnumerateTimelines(TFunctionRef<void(const Timeline&)> Callback) const
+{
+	SessionLock.ReadAccessCheck();
+
+	for (const auto& Timeline : Timelines)
+	{
+		Callback(*Timeline);
+	}
+}
+
+void FTimingProfilerProvider::ReadTimers(TFunctionRef<void(const FTimingProfilerTimer*, uint64)> Callback) const
+{
+	SessionLock.ReadAccessCheck();
+
+	Callback(Timers.GetData(), Timers.Num());
+}
+
+}
\ No newline at end of file
diff --git a/Engine/Source/Developer/TraceServices/Private/Model/TimingProfiler.h b/Engine/Source/Developer/TraceServices/Private/Model/TimingProfiler.h
new file mode 100644
--- /dev/null
+++ b/Engine/Source/Developer/TraceServices/Private/Model/TimingProfiler.h
@@ -0,0 +1,43 @@
+// Copyright 1998-2019 Epic Games, Inc. All Rights Reserved.
+
+#pragma once
+
+#include "TraceServices/AnalysisService.h"
+#include "Common/SlabAllocator.h"
+#include "Model/MonotonicTimeline.h"
+
+namespace Trace
+{
+
+class FAnalysisSessionLock;
+
+class FTimingProfilerProvider
+	: public ITimingProfilerProvider
+{
+public:
+	typedef TMonotonicTimeline<FTimingProfilerEvent> TimelineInternal;
+
+	FTimingProfilerProvider(FSlabAllocator& Allocator, FAnalysisSessionLock& InSessionLock);
+	virtual ~FTimingProfilerProvider();
+	uint32 AddCpuTimer(const TCHAR* Name);
+	uint32 AddGpuTimer(const TCHAR* Name);
+	TSharedRef<TimelineInternal> EditCpuThreadTimeline(uint32 ThreadId);
+	TSharedRef<TimelineInternal> EditGpuTimeline();
+	virtual bool GetCpuThreadTimelineIndex(uint32 ThreadId, uint32& OutTimelineIndex) const override;
+	virtual bool GetGpuTimelineIndex(uint32& OutTimelineIndex) const override;
+	virtual bool ReadTimeline(uint32 Index, TFunctionRef<void(const Timeline&)> Callback) const override;
+	virtual void EnumerateTimelines(TFunctionRef<void(const Timeline&)> Callback) const override;
+	virtual void ReadTimers(TFunctionRef<void(const FTimingProfilerTimer*, uint64)> Callback) const override;
+
+private:
+	FTimingProfilerTimer& AddTimerInternal(const TCHAR* Name, bool IsGpuEvent);
+
+	FSlabAllocator& Allocator;
+	FAnalysisSessionLock& SessionLock;
+	TArray<FTimingProfilerTimer> Timers;
+	TArray<TSharedRef<TimelineInternal>> Timelines;
+	TMap<uint32, uint32> CpuThreadTimelineIndexMap;
+	uint32 GpuTimelineIndex = 0;
+};
+
+}
\ No newline at end of file
diff --git a/Engine/Source/Developer/TraceServices/Private/SessionService.cpp b/Engine/Source/Developer/TraceServices/Private/SessionService.cpp
new file mode 100644
--- /dev/null
+++ b/Engine/Source/Developer/TraceServices/Private/SessionService.cpp
@@ -0,0 +1,63 @@
+// Copyright 1998-2019 Epic Games, Inc. All Rights Reserved.
+
+#include "TraceServices/SessionService.h"
+#include "SessionServicePrivate.h"
+#include "Trace/Store.h"
+#include "Trace/Recorder.h"
+#include "Trace/DataStream.h"
+
+namespace Trace
+{
+
+FSessionService::FSessionService(TSharedRef<IStore> InTraceStore)
+	: TraceStore(InTraceStore)
+{
+	
+}
+
+FSessionService::~FSessionService()
+{
+	
+}
+
+void FSessionService::StartRecorderServer()
+{
+	Trace::Recorder_StartServer(TraceStore);
+}
+
+void FSessionService::GetAvailableSessions(TArray<Trace::FSessionHandle>& OutSessions)
+{
+	TraceStore->GetAvailableSessions(OutSessions);
+}
+
+bool FSessionService::GetSessionInfo(Trace::FSessionHandle SessionHandle, Trace::FSessionInfo& OutSessionInfo)
+{
+	return TraceStore->GetSessionInfo(SessionHandle, OutSessionInfo);
+}
+
+Trace::IInDataStream* FSessionService::OpenSessionStream(Trace::FSessionHandle SessionHandle)
+{
+	return TraceStore->OpenSessionStream(SessionHandle);
+}
+
+Trace::IInDataStream* FSessionService::OpenSessionFromFile(const TCHAR* FilePath)
+{
+	return Trace::DataStream_ReadFile(FilePath);
+}
+
+/*void FSessionService::GetSessions(TArray<TSharedRef<FSessionInfo>>& OutSessions)
+{
+	OutSessions.Reset(AvailableSessions.Num());
+	for (const auto& KV : AvailableSessions)
+	{
+		OutSessions.Add(KV.Value);
+	}
+}*/
+
+/*void FSessionService::SelectSession(TSharedPtr<FSessionInfo> Session)
+{
+	SelectedSession = Session;
+	OnSessionSelected().Broadcast(SelectedSession);
+}*/
+
+}
\ No newline at end of file
diff --git a/Engine/Source/Developer/TraceServices/Private/SessionServicePrivate.h b/Engine/Source/Developer/TraceServices/Private/SessionServicePrivate.h
new file mode 100644
--- /dev/null
+++ b/Engine/Source/Developer/TraceServices/Private/SessionServicePrivate.h
@@ -0,0 +1,26 @@
+// Copyright 1998-2019 Epic Games, Inc. All Rights Reserved.
+
+#pragma once
+
+#include "TraceServices/SessionService.h"
+
+namespace Trace
+{
+
+class FSessionService
+	: public ISessionService
+{
+public:
+	FSessionService(TSharedRef<IStore> InTraceStore);
+	virtual ~FSessionService();
+	virtual void StartRecorderServer() override;
+	virtual void GetAvailableSessions(TArray<Trace::FSessionHandle>& OutSessions) override;
+	virtual bool GetSessionInfo(Trace::FSessionHandle SessionHandle, Trace::FSessionInfo& OutSessionInfo) override;
+	virtual Trace::IInDataStream* OpenSessionStream(Trace::FSessionHandle SessionHandle) override;
+	virtual Trace::IInDataStream* OpenSessionFromFile(const TCHAR* FilePath) override;
+
+private:
+	TSharedRef<IStore> TraceStore;
+};
+
+}
\ No newline at end of file
diff --git a/Engine/Source/Developer/TraceServices/Private/TraceServicesModule.cpp b/Engine/Source/Developer/TraceServices/Private/TraceServicesModule.cpp
new file mode 100644
--- /dev/null
+++ b/Engine/Source/Developer/TraceServices/Private/TraceServicesModule.cpp
@@ -0,0 +1,45 @@
+// Copyright 1998-2019 Epic Games, Inc. All Rights Reserved.
+
+#include "TraceServices/ITraceServicesModule.h"
+#include "SessionServicePrivate.h"
+#include "AnalysisServicePrivate.h"
+#include "Trace/Recorder.h"
+
+class FTraceServicesModule
+	: public ITraceServicesModule
+{
+public:
+	virtual void StartupModule() override;
+	virtual TSharedPtr<Trace::ISessionService> GetSessionService() override;
+	virtual TSharedPtr<Trace::IAnalysisService> GetAnalysisService() override;
+
+private:
+	TSharedPtr<Trace::IStore> TraceStore;
+	TSharedPtr<Trace::FSessionService> SessionService;
+	TSharedPtr<Trace::FAnalysisService> AnalysisService;
+};
+
+void FTraceServicesModule::StartupModule()
+{
+	TraceStore = Trace::Store_Create(TEXT("D:\\Trace"));
+}
+
+TSharedPtr<Trace::ISessionService> FTraceServicesModule::GetSessionService()
+{
+	if (!SessionService.IsValid())
+	{
+		SessionService = MakeShareable(new Trace::FSessionService(TraceStore.ToSharedRef()));
+	}
+	return SessionService;
+}
+
+TSharedPtr<Trace::IAnalysisService> FTraceServicesModule::GetAnalysisService()
+{
+	if (!AnalysisService.IsValid())
+	{
+		AnalysisService = MakeShareable(new Trace::FAnalysisService(TraceStore.ToSharedRef()));
+	}
+	return AnalysisService;
+}
+
+IMPLEMENT_MODULE(FTraceServicesModule, TraceServices);
\ No newline at end of file
diff --git a/Engine/Source/Developer/TraceServices/Public/TraceServices/AnalysisService.h b/Engine/Source/Developer/TraceServices/Public/TraceServices/AnalysisService.h
new file mode 100644
--- /dev/null
+++ b/Engine/Source/Developer/TraceServices/Public/TraceServices/AnalysisService.h
@@ -0,0 +1,284 @@
+// Copyright 1998-2019 Epic Games, Inc. All Rights Reserved.
+
+#pragma once
+
+#include "CoreTypes.h"
+#include "Containers/UnrealString.h"
+#include "Templates/SharedPointer.h"
+#include "Trace/Store.h"
+#include "Delegates/Delegate.h"
+#include "GenericPlatform/GenericPlatformAffinity.h"
+#include "Serialization/LoadTimeTrace.h"
+
+namespace Trace
+{
+
+struct FBookmark
+{
+	double Time;
+	const TCHAR* Text;
+};
+
+class IBookmarkProvider
+{
+public:
+	virtual ~IBookmarkProvider() = default;
+	virtual void EnumerateBookmarks(double IntervalStart, double IntervalEnd, TFunctionRef<void(const FBookmark&)> Callback) const = 0;
+};
+
+struct FLogCategory
+{
+	FString Name;
+	ELogVerbosity::Type DefaultVerbosity;
+};
+
+struct FLogMessage
+{
+	uint64 Index;
+	double Time;
+	const FLogCategory* Category;
+	const TCHAR* File;
+	const TCHAR* Message;
+	int32 Line;
+	ELogVerbosity::Type Verbosity;
+};
+
+class ILogProvider
+{
+public:
+	virtual ~ILogProvider() = default;
+	virtual uint64 GetMessageCount() const = 0;
+	virtual void EnumerateMessages(double IntervalStart, double IntervalEnd, TFunctionRef<void(const FLogMessage&)> Callback, bool ResolveFormatString = true) const = 0;
+	virtual void EnumerateMessagesByIndex(uint64 Start, uint64 End, TFunctionRef<void(const FLogMessage&)> Callback, bool ResolveFormatString = true) const = 0;
+	virtual bool ReadMessage(uint64 Index, TFunctionRef<void(const FLogMessage&)> Callback, bool ResolveFormatString = true) const = 0;
+	virtual void EnumerateCategories(TFunctionRef<void(const FLogCategory&)> Callback) const = 0;
+};
+
+struct FThreadInfo
+{
+	uint32 Id;
+	const TCHAR* Name;
+	const TCHAR* GroupName;
+};
+
+class IThreadProvider
+{
+public:
+	virtual ~IThreadProvider() = default;
+	virtual void EnumerateThreads(TFunctionRef<void(const FThreadInfo&)> Callback) const = 0;
+};
+
+template<typename EventType>
+class ITimeline
+{
+public:
+	virtual ~ITimeline() = default;
+	virtual uint64 GetEventCount() const = 0;
+	virtual void EnumerateEventsDownSampled(double IntervalStart, double IntervalEnd, double Resolution, TFunctionRef<void(bool, double, const EventType&)> Callback) const = 0;
+	virtual void EnumerateEventsDownSampled(double IntervalStart, double IntervalEnd, double Resolution, TFunctionRef<void(double, double, uint32, const EventType&)> Callback) const = 0;
+	virtual void EnumerateEvents(double IntervalStart, double IntervalEnd, TFunctionRef<void(bool, double, const EventType&)> Callback) const = 0;
+	virtual void EnumerateEvents(double IntervalStart, double IntervalEnd, TFunctionRef<void(double, double, uint32, const EventType&)> Callback) const = 0;
+};
+
+struct FTimingProfilerTimer
+{
+	FString Name;
+	uint32 Id;
+	uint32 NameHash;
+	bool IsGpuTimer;
+};
+
+struct FTimingProfilerEvent
+{
+	uint32 TimerIndex;
+};
+
+class ITimingProfilerProvider
+{
+public:
+	typedef ITimeline<FTimingProfilerEvent> Timeline;
+
+	virtual ~ITimingProfilerProvider() = default;
+	virtual bool GetCpuThreadTimelineIndex(uint32 ThreadId, uint32& OutTimelineIndex) const = 0;
+	virtual bool GetGpuTimelineIndex(uint32& OutTimelineIndex) const = 0;
+	virtual bool ReadTimeline(uint32 Index, TFunctionRef<void(const Timeline&)> Callback) const = 0;
+	virtual void EnumerateTimelines(TFunctionRef<void(const Timeline&)> Callback) const = 0;
+	virtual void ReadTimers(TFunctionRef<void(const FTimingProfilerTimer*, uint64)> Callback) const = 0;
+};
+
+struct FFileInfo
+{
+	uint32 Id;
+	const TCHAR* Path;
+};
+
+enum EFileActivityType
+{
+	FileActivityType_Open,
+	FileActivityType_Close,
+	FileActivityType_Read,
+	FileActivityType_Write,
+
+	FileActivityType_Count,
+};
+
+struct FFileActivity
+{
+	uint64 Offset;
+	uint64 Size;
+	EFileActivityType ActivityType;
+	bool Failed;
+};
+
+class IFileActivityProvider
+{
+public:
+	typedef ITimeline<FFileActivity> Timeline;
+
+	virtual ~IFileActivityProvider() = default;
+	virtual void EnumerateFileActivity(TFunctionRef<bool(const FFileInfo&, const Timeline&)> Callback) const = 0;
+};
+
+struct FFrame
+{
+	uint64 Index;
+	double StartTime;
+	double EndTime;
+};
+
+class IFrameProvider
+{
+public:
+	virtual ~IFrameProvider() = default;
+	virtual uint64 GetFrameCount(ETraceFrameType FrameType) const = 0;
+	virtual void EnumerateFrames(ETraceFrameType FrameType, uint64 Start, uint64 End, TFunctionRef<void(const FFrame&)> Callback) const = 0;
+};
+
+struct FPackageSummaryInfo
+{
+	uint32 TotalHeaderSize = 0;
+	uint32 NameCount = 0;
+	uint32 ImportCount = 0;
+	uint32 ExportCount = 0;
+};
+
+struct FClassInfo
+{
+	FString Name;
+};
+
+struct FPackageExportInfo
+{
+	uint32 Id;
+	const FClassInfo* Class = nullptr;
+	uint64 SerialOffset = 0;
+	uint64 SerialSize = 0;
+	bool IsAsset = false;
+};
+
+struct FPackageInfo
+{
+	uint32 Id;
+	FString Name;
+	FPackageSummaryInfo Summary;
+	TArray<FPackageExportInfo*> Exports;
+};
+
+struct FLoadTimeProfilerCpuEvent
+{
+	const FPackageInfo* Package = nullptr;
+	const FPackageExportInfo* Export = nullptr;
+	ELoadTimeProfilePackageEventType PackageEventType = LoadTimeProfilerPackageEventType_None;
+	ELoadTimeProfilerObjectEventType ExportEventType = LoadTimeProfilerObjectEventType_None;
+};
+
+class ILoadTimeProfilerProvider
+{
+public:
+	typedef ITimeline<FLoadTimeProfilerCpuEvent> CpuTimeline;
+
+	virtual ~ILoadTimeProfilerProvider() = default;
+	virtual void EnumeratePackages(TFunctionRef<void(const FPackageInfo&)> Callback) const = 0; 
+	virtual void ReadMainThreadCpuTimeline(TFunctionRef<void(const CpuTimeline&)> Callback) const = 0;
+	virtual void ReadAsyncLoadingThreadCpuTimeline(TFunctionRef<void(const CpuTimeline&)> Callback) const = 0;
+};
+
+enum ECounterDisplayHint
+{
+	CounterDisplayHint_None,
+	CounterDisplayHint_Memory,
+	CounterDisplayHint_FloatingPoint,
+};
+
+class ICounter
+{
+public:
+	virtual ~ICounter() = default;
+
+	virtual const TCHAR* GetName() const = 0;
+	virtual uint32 GetId() const = 0;
+	virtual ECounterDisplayHint GetDisplayHint() const = 0;
+	virtual void EnumerateValues(double IntervalStart, double IntervalEnd, TFunctionRef<void(double, int64)> Callback) const = 0;
+	virtual void EnumerateFloatValues(double IntervalStart, double IntervalEnd, TFunctionRef<void(double, double)> Callback) const = 0;
+};
+
+class ICountersProvider
+{
+public:
+	virtual ~ICountersProvider() = default;
+	virtual void EnumerateCounters(TFunctionRef<void(const ICounter&)> Callback) const = 0;
+};
+
+class IAnalysisSession
+{
+public:
+	virtual ~IAnalysisSession() = default;
+	
+	virtual const TCHAR* GetName() const = 0;
+	virtual bool IsAnalysisComplete() const = 0;
+	virtual double GetDurationSeconds() const = 0;
+
+	virtual void BeginRead() const = 0;
+	virtual void EndRead() const = 0;
+	virtual void ReadBookmarkProvider(TFunctionRef<void(const IBookmarkProvider&)> Callback) const = 0;
+	virtual void ReadLogProvider(TFunctionRef<void(const ILogProvider&)> Callback) const = 0;
+	virtual void ReadThreadProvider(TFunctionRef<void(const IThreadProvider&)> Callback) const = 0;
+	virtual void ReadFramesProvider(TFunctionRef<void(const IFrameProvider&)> Callback) const = 0;
+	virtual void ReadTimingProfilerProvider(TFunctionRef<void(const ITimingProfilerProvider&)> Callback) const = 0;
+	virtual void ReadFileActivityProvider(TFunctionRef<void(const IFileActivityProvider&)> Callback) const = 0;
+	virtual void ReadLoadTimeProfilerProvider(TFunctionRef<void(const ILoadTimeProfilerProvider&)> Callback) const = 0;
+	virtual void ReadCountersProvider(TFunctionRef<void(const ICountersProvider&)> Callback) const = 0;
+};
+
+struct FAnalysisSessionReadScope
+{
+	FAnalysisSessionReadScope(const IAnalysisSession& InAnalysisSession)
+		: AnalysisSession(InAnalysisSession)
+	{
+		AnalysisSession.BeginRead();
+	}
+
+	~FAnalysisSessionReadScope()
+	{
+		AnalysisSession.EndRead();
+	}
+
+private:
+	const IAnalysisSession& AnalysisSession;
+};
+
+class IAnalysisService
+{
+public:
+	virtual TSharedPtr<const IAnalysisSession> Analyze(const TCHAR* SessionName, TUniquePtr<Trace::IInDataStream>&& DataStream) = 0;
+	virtual TSharedPtr<const IAnalysisSession> StartAnalysis(const TCHAR* SessionName, TUniquePtr<Trace::IInDataStream>&& DataStream) = 0;
+	virtual TSharedPtr<const IAnalysisSession> MockAnalysis() = 0;
+
+	DECLARE_EVENT_OneParam(IAnalysisService, FAnalysisStartedEvent, TSharedRef<const IAnalysisSession>)
+	virtual FAnalysisStartedEvent& OnAnalysisStarted() = 0;
+
+	DECLARE_EVENT_OneParam(IAnalysisService, FAnalysisFinishedEvent, TSharedRef<const IAnalysisSession>)
+	virtual FAnalysisFinishedEvent& OnAnalysisFinished() = 0;
+};
+
+}
\ No newline at end of file
diff --git a/Engine/Source/Developer/TraceServices/Public/TraceServices/ITraceServicesModule.h b/Engine/Source/Developer/TraceServices/Public/TraceServices/ITraceServicesModule.h
new file mode 100644
--- /dev/null
+++ b/Engine/Source/Developer/TraceServices/Public/TraceServices/ITraceServicesModule.h
@@ -0,0 +1,25 @@
+// Copyright 1998-2019 Epic Games, Inc. All Rights Reserved.
+
+#pragma once
+
+#include "CoreMinimal.h"
+#include "Trace/Store.h"
+#include "Modules/ModuleInterface.h"
+
+namespace Trace
+{
+
+class IAnalysisService;
+class ISessionService;
+
+}
+
+class ITraceServicesModule
+	: public IModuleInterface
+{
+public:
+	virtual TSharedPtr<Trace::ISessionService> GetSessionService() = 0;
+	virtual TSharedPtr<Trace::IAnalysisService> GetAnalysisService() = 0;
+
+	virtual ~ITraceServicesModule() { }
+};
diff --git a/Engine/Source/Developer/TraceServices/Public/TraceServices/SessionService.h b/Engine/Source/Developer/TraceServices/Public/TraceServices/SessionService.h
new file mode 100644
--- /dev/null
+++ b/Engine/Source/Developer/TraceServices/Public/TraceServices/SessionService.h
@@ -0,0 +1,24 @@
+// Copyright 1998-2019 Epic Games, Inc. All Rights Reserved.
+
+#pragma once
+
+#include "CoreMinimal.h"
+#include "Templates/SharedPointer.h"
+#include "Trace/Store.h"
+#include "Trace/DataStream.h"
+
+namespace Trace
+{
+
+class ISessionService
+{
+public:
+	virtual ~ISessionService() = default;
+	virtual void StartRecorderServer() = 0;
+	virtual void GetAvailableSessions(TArray<Trace::FSessionHandle>& OutSessions) = 0;
+	virtual bool GetSessionInfo(Trace::FSessionHandle SessionHandle, Trace::FSessionInfo& OutSessionInfo) = 0;
+	virtual Trace::IInDataStream* OpenSessionStream(Trace::FSessionHandle SessionHandle) = 0;
+	virtual Trace::IInDataStream* OpenSessionFromFile(const TCHAR* FilePath) = 0;
+};
+
+}
diff --git a/Engine/Source/Programs/UnrealInsights/Private/Mac/MacUnrealInsightsMain.cpp b/Engine/Source/Programs/UnrealInsights/Private/Mac/MacUnrealInsightsMain.cpp
new file mode 100644
--- /dev/null
+++ b/Engine/Source/Programs/UnrealInsights/Private/Mac/MacUnrealInsightsMain.cpp
@@ -0,0 +1,116 @@
+// Copyright 1998-2019 Epic Games, Inc. All Rights Reserved.
+
+#include "UnrealInsightsMain.h"
+#include "HAL/ExceptionHandling.h"
+#include "Mac/CocoaThread.h"
+
+
+static FString GSavedCommandLine;
+
+
+@interface UE4AppDelegate : NSObject<NSApplicationDelegate, NSFileManagerDelegate>
+{
+}
+
+@end
+
+
+@implementation UE4AppDelegate
+
+//handler for the quit apple event used by the Dock menu
+- (void)handleQuitEvent:(NSAppleEventDescriptor*)Event withReplyEvent:(NSAppleEventDescriptor*)ReplyEvent
+{
+	[self requestQuit:self];
+}
+
+- (IBAction)requestQuit:(id)Sender
+{
+	GIsRequestingExit = true;
+}
+
+- (void) runGameThread:(id)Arg
+{
+	FPlatformMisc::SetGracefulTerminationHandler();
+	FPlatformMisc::SetCrashHandler(nullptr);
+	
+#if !UE_BUILD_SHIPPING
+	if (FParse::Param(*GSavedCommandLine,TEXT("crashreports")))
+	{
+		GAlwaysReportCrash = true;
+	}
+#endif
+	
+#if UE_BUILD_DEBUG
+	if (!GAlwaysReportCrash)
+#else
+		if (FPlatformMisc::IsDebuggerPresent() && !GAlwaysReportCrash)
+#endif
+		{
+			UnrealInsightsMain(*GSavedCommandLine);
+		}
+		else
+		{
+			GIsGuarded = 1;
+			UnrealInsightsMain(*GSavedCommandLine);
+			GIsGuarded = 0;
+		}
+	
+	FEngineLoop::AppExit();
+
+	[NSApp terminate: self];
+}
+
+- (NSApplicationTerminateReply)applicationShouldTerminate:(NSApplication *)Sender;
+{
+	if(!GIsRequestingExit || ([NSThread gameThread] && [NSThread gameThread] != [NSThread mainThread]))
+	{
+		[self requestQuit:self];
+		return NSTerminateLater;
+	}
+	else
+	{
+		return NSTerminateNow;
+	}
+}
+
+- (void)applicationDidFinishLaunching:(NSNotification *)Notification
+{
+	//install the custom quit event handler
+	NSAppleEventManager* appleEventManager = [NSAppleEventManager sharedAppleEventManager];
+	[appleEventManager setEventHandler:self andSelector:@selector(handleQuitEvent:withReplyEvent:) forEventClass:kCoreEventClass andEventID:kAEQuitApplication];
+
+	RunGameThread(self, @selector(runGameThread:));
+}
+
+@end
+
+
+int main(int argc, char *argv[])
+{
+	for (int32 Option = 1; Option < argc; Option++)
+	{
+		GSavedCommandLine += TEXT(" ");
+		FString Argument(ANSI_TO_TCHAR(argv[Option]));
+		if (Argument.Contains(TEXT(" ")))
+		{
+			if (Argument.Contains(TEXT("=")))
+			{
+				FString ArgName;
+				FString ArgValue;
+				Argument.Split( TEXT("="), &ArgName, &ArgValue );
+				Argument = FString::Printf( TEXT("%s=\"%s\""), *ArgName, *ArgValue );
+			}
+			else
+			{
+				Argument = FString::Printf(TEXT("\"%s\""), *Argument);
+			}
+		}
+		GSavedCommandLine += Argument;
+	}
+
+	SCOPED_AUTORELEASE_POOL;
+	[NSApplication sharedApplication];
+	[NSApp setDelegate:[UE4AppDelegate new]];
+	[NSApp run];
+	return 0;
+}
diff --git a/Engine/Source/Programs/UnrealInsights/Private/UnrealInsightsMain.cpp b/Engine/Source/Programs/UnrealInsights/Private/UnrealInsightsMain.cpp
new file mode 100644
--- /dev/null
+++ b/Engine/Source/Programs/UnrealInsights/Private/UnrealInsightsMain.cpp
@@ -0,0 +1,46 @@
+// Copyright 1998-2019 Epic Games, Inc. All Rights Reserved.
+
+#include "UnrealInsightsMain.h"
+
+#include "RequiredProgramMainCPPInclude.h"
+#include "UserInterfaceCommand.h"
+
+
+IMPLEMENT_APPLICATION(UnrealInsights, "UnrealInsights");
+
+
+/**
+ * Platform agnostic implementation of the main entry point.
+ */
+int32 UnrealInsightsMain(const TCHAR* CommandLine)
+{
+	// Override the stack size for the thread pool.
+	FQueuedThreadPool::OverrideStackSize = 256 * 1024;
+
+	//im: ???
+	FCommandLine::Set(CommandLine);
+
+	// Initialize core.
+	GEngineLoop.PreInit(CommandLine);
+
+	// Make sure all UObject classes are registered and default properties have been initialized.
+	//im: ??? ProcessNewlyLoadedUObjects();
+
+	// Tell the module manager it may now process newly-loaded UObjects when new C++ modules are loaded.
+	FModuleManager::Get().StartProcessingNewlyLoadedObjects();
+
+	FUserInterfaceCommand::Run();
+
+	// Shut down.
+	//im: ??? FCoreDelegates::OnExit.Broadcast();
+	FEngineLoop::AppPreExit(); //im: ???
+	FModuleManager::Get().UnloadModulesAtShutdown();
+
+#if STATS
+	FThreadStats::StopThread();
+#endif
+
+	FTaskGraphInterface::Shutdown(); //im: ???
+
+	return 0;
+}
diff --git a/Engine/Source/Programs/UnrealInsights/UnrealInsights.Target.cs b/Engine/Source/Programs/UnrealInsights/UnrealInsights.Target.cs
new file mode 100644
--- /dev/null
+++ b/Engine/Source/Programs/UnrealInsights/UnrealInsights.Target.cs
@@ -0,0 +1,30 @@
+// Copyright 1998-2019 Epic Games, Inc. All Rights Reserved.
+
+using UnrealBuildTool;
+using System.Collections.Generic;
+
+[SupportedPlatforms(UnrealPlatformClass.Desktop)]
+[SupportedPlatforms(UnrealTargetPlatform.IOS)]
+public class UnrealInsightsTarget : TargetRules
+{
+	public UnrealInsightsTarget(TargetInfo Target) : base(Target)
+	{
+		Type = TargetType.Program;
+		LinkType = TargetLinkType.Modular; // TargetLinkType.Monolithic;
+
+		LaunchModuleName = "UnrealInsights";
+		ExtraModuleNames.Add("EditorStyle");
+
+		bBuildDeveloperTools = false;
+		bCompileAgainstEngine = false;
+		bCompileAgainstCoreUObject = true;
+		bForceBuildTargetPlatforms = true;
+		bCompileWithStatsWithoutEngine = true;
+		bCompileWithPluginSupport = true;
+
+		// For UI functionality
+		bBuildDeveloperTools = true;
+
+		bHasExports = false;
+	}
+}
diff --git a/Engine/Source/Runtime/Core/Core.Build.cs b/Engine/Source/Runtime/Core/Core.Build.cs
--- a/Engine/Source/Runtime/Core/Core.Build.cs
+++ b/Engine/Source/Runtime/Core/Core.Build.cs
@@ -14,6 +14,8 @@ public class Core : ModuleRules
 
 		PrivateDependencyModuleNames.Add("BuildSettings");
 
+		PublicDependencyModuleNames.Add("Trace");
+
 		PrivateIncludePaths.AddRange(
 			new string[] {
 				"Developer/DerivedDataCache/Public",
diff --git a/Engine/Source/Runtime/Core/Private/Async/TaskGraph.cpp b/Engine/Source/Runtime/Core/Private/Async/TaskGraph.cpp
--- a/Engine/Source/Runtime/Core/Private/Async/TaskGraph.cpp
+++ b/Engine/Source/Runtime/Core/Private/Async/TaskGraph.cpp
@@ -28,6 +28,7 @@
 #include "HAL/LowLevelMemTracker.h"
 #include "HAL/ThreadHeartBeat.h"
 #include "ProfilingDebugging/ExternalProfiler.h"
+#include "ProfilingDebugging/MiscTrace.h"
 
 DEFINE_LOG_CATEGORY_STATIC(LogTaskGraph, Log, All);
 
@@ -1210,12 +1211,14 @@ public:
 		for (int32 ThreadIndex = LastExternalThread + 1; ThreadIndex < NumThreads; ThreadIndex++)
 		{
 			FString Name;
+			ETraceThreadGroup Group = TraceThreadGroup_TaskGraphNormal;
 			int32 Priority = ThreadIndexToPriorityIndex(ThreadIndex);
 			EThreadPriority ThreadPri;
 			uint64 Affinity = FPlatformAffinity::GetTaskGraphThreadMask();
 			if (Priority == 1)
 			{
 				Name = FString::Printf(TEXT("TaskGraphThreadHP %d"), ThreadIndex - (LastExternalThread + 1));
+				Group = TraceThreadGroup_TaskGraphHigh;
 				ThreadPri = TPri_SlightlyBelowNormal; // we want even hi priority tasks below the normal threads
 
 				// If the platform defines FPlatformAffinity::GetTaskGraphHighPriorityTaskMask then use it
@@ -1227,6 +1230,7 @@ public:
 			else if (Priority == 2)
 			{
 				Name = FString::Printf(TEXT("TaskGraphThreadBP %d"), ThreadIndex - (LastExternalThread + 1));
+				Group = TraceThreadGroup_TaskGraphLow;
 				ThreadPri = TPri_Lowest;
 				// If the platform defines FPlatformAffinity::GetTaskGraphBackgroundTaskMask then use it
 				if ( FPlatformAffinity::GetTaskGraphBackgroundTaskMask() != 0xFFFFFFFFFFFFFFFF )
@@ -1248,6 +1252,10 @@ public:
 #endif
 			WorkerThreads[ThreadIndex].RunnableThread = FRunnableThread::Create(&Thread(ThreadIndex), *Name, StackSize, ThreadPri, Affinity); // these are below normal threads so that they sleep when the named threads are active
 			WorkerThreads[ThreadIndex].bAttached = true;
+			if (WorkerThreads[ThreadIndex].RunnableThread)
+			{
+				TRACE_SET_THREAD_GROUP(WorkerThreads[ThreadIndex].RunnableThread->GetThreadID(), Group);
+			}
 		}
 	}
 
diff --git a/Engine/Source/Runtime/Core/Private/Logging/LogTrace.cpp b/Engine/Source/Runtime/Core/Private/Logging/LogTrace.cpp
new file mode 100644
--- /dev/null
+++ b/Engine/Source/Runtime/Core/Private/Logging/LogTrace.cpp
@@ -0,0 +1,84 @@
+// Copyright 1998-2019 Epic Games, Inc. All Rights Reserved.
+#include "Logging/LogTrace.h"
+
+#if LOGTRACE_ENABLED
+
+#include "Trace/Trace.h"
+#include "Templates/Function.h"
+#include "HAL/PlatformTime.h"
+#include "HAL/PlatformTLS.h"
+
+UE_TRACE_EVENT_BEGIN(Logging, LogCategory)
+	UE_TRACE_EVENT_FIELD(const void*, CategoryPointer)
+	UE_TRACE_EVENT_FIELD(uint16, NameLength)
+	UE_TRACE_EVENT_FIELD(uint8, DefaultVerbosity)
+UE_TRACE_EVENT_END()
+
+UE_TRACE_EVENT_BEGIN(Logging, LogMessageSpec)
+	UE_TRACE_EVENT_FIELD(const void*, LogPoint)
+	UE_TRACE_EVENT_FIELD(const void*, CategoryPointer)
+	UE_TRACE_EVENT_FIELD(int32, Line)
+	UE_TRACE_EVENT_FIELD(uint16, FileNameLength)
+	UE_TRACE_EVENT_FIELD(uint16, FormatStringLength)
+	UE_TRACE_EVENT_FIELD(uint8, Verbosity)
+UE_TRACE_EVENT_END()
+
+UE_TRACE_EVENT_BEGIN(Logging, LogMessage)
+	UE_TRACE_EVENT_FIELD(const void*, LogPoint)
+	UE_TRACE_EVENT_FIELD(uint64, Cycle)
+	UE_TRACE_EVENT_FIELD(uint32, ThreadId)
+UE_TRACE_EVENT_END()
+
+void FLogTrace::OutputLogCategory(const FLogCategoryBase* Category, const TCHAR* Name, ELogVerbosity::Type DefaultVerbosity)
+{
+	uint16 NameLength = FCString::Strlen(Name);
+	auto NameCopy = [NameLength, Name](uint8* Out) {
+		TCHAR* OutTChar = reinterpret_cast<TCHAR*>(Out);
+		for (int i = 0; i < NameLength; ++i)
+		{
+			*OutTChar++ = Name[i];
+		}
+	};
+	UE_TRACE_LOG(Logging, LogCategory, NameLength * 2)
+		<< LogCategory.CategoryPointer(Category)
+		<< LogCategory.NameLength(NameLength)
+		<< LogCategory.DefaultVerbosity(DefaultVerbosity)
+		<< LogCategory.Attachment(NameCopy);
+}
+
+void FLogTrace::OutputLogMessageSpec(const void* LogPoint, const FLogCategoryBase* Category, ELogVerbosity::Type Verbosity, const ANSICHAR* File, int32 Line, const TCHAR* Format)
+{
+	uint16 FileNameLength = strlen(File);
+	uint16 FormatStringLength = FCString::Strlen(Format);
+	auto StringCopyFunc = [FileNameLength, FormatStringLength, File, Format](uint8* Out) {
+		ANSICHAR* OutAnsiChar = reinterpret_cast<ANSICHAR*>(Out);
+		for (int i = 0; i < FileNameLength; ++i)
+		{
+			*OutAnsiChar++ = File[i];
+		}
+		TCHAR* OutTChar = reinterpret_cast<TCHAR*>(OutAnsiChar);
+		for (int i = 0; i < FormatStringLength; ++i)
+		{
+			*OutTChar++ = Format[i];
+		}
+	};
+	UE_TRACE_LOG(Logging, LogMessageSpec, FileNameLength + FormatStringLength * 2)
+		<< LogMessageSpec.LogPoint(LogPoint)
+		<< LogMessageSpec.CategoryPointer(Category)
+		<< LogMessageSpec.Line(Line)
+		<< LogMessageSpec.FileNameLength(FileNameLength)
+		<< LogMessageSpec.FormatStringLength(FormatStringLength)
+		<< LogMessageSpec.Verbosity(Verbosity)
+		<< LogMessageSpec.Attachment(StringCopyFunc);
+}
+
+void FLogTrace::OutputLogMessageInternal(const void* LogPoint, uint16 EncodedFormatArgsSize, uint8* EncodedFormatArgs)
+{
+	UE_TRACE_LOG(Logging, LogMessage, EncodedFormatArgsSize)
+		<< LogMessage.LogPoint(LogPoint)
+		<< LogMessage.Cycle(FPlatformTime::Cycles64())
+		<< LogMessage.ThreadId(FPlatformTLS::GetCurrentThreadId())
+		<< LogMessage.Attachment(EncodedFormatArgs, EncodedFormatArgsSize);
+}
+
+#endif
diff --git a/Engine/Source/Runtime/Core/Private/ProfilingDebugging/CpuProfilerTrace.cpp b/Engine/Source/Runtime/Core/Private/ProfilingDebugging/CpuProfilerTrace.cpp
new file mode 100644
--- /dev/null
+++ b/Engine/Source/Runtime/Core/Private/ProfilingDebugging/CpuProfilerTrace.cpp
@@ -0,0 +1,240 @@
+// Copyright 1998-2019 Epic Games, Inc. All Rights Reserved.
+#include "ProfilingDebugging/CpuProfilerTrace.h"
+#include "Templates/UnrealTemplate.h"
+#include "Trace/Trace.h"
+#include "HAL/PlatformTime.h"
+#include "HAL/PlatformTLS.h"
+#include "Containers/Map.h"
+#include "ProfilingDebugging/MiscTrace.h"
+
+#if CPUPROFILERTRACE_ENABLED
+
+UE_TRACE_EVENT_BEGIN(CpuProfiler, EventSpec)
+	UE_TRACE_EVENT_FIELD(uint16, Id)
+	UE_TRACE_EVENT_FIELD(uint16, Group)
+	UE_TRACE_EVENT_FIELD(uint16, NameSize)
+	//UE_TRACE_EVENT_FIELD(uint8[], Name)
+UE_TRACE_EVENT_END()
+
+UE_TRACE_EVENT_BEGIN(CpuProfiler, DynamicEventName)
+	UE_TRACE_EVENT_FIELD(uint16, Id)
+	UE_TRACE_EVENT_FIELD(uint16, NameSize)
+	//UE_TRACE_EVENT_FIELD(uint8[], Name)
+UE_TRACE_EVENT_END()
+
+UE_TRACE_EVENT_BEGIN(CpuProfiler, EventBatch)
+	UE_TRACE_EVENT_FIELD(uint64, CycleBase)
+	UE_TRACE_EVENT_FIELD(uint32, ThreadId)
+	UE_TRACE_EVENT_FIELD(uint8, BufferSize)
+	//UE_TRACE_EVENT_FIELD(uint8[], Buffer)
+UE_TRACE_EVENT_END()
+
+struct FCpuProfilerTraceInternal
+{
+public:
+	enum
+	{
+		MaxBufferSize = 255,
+		MaxEncodedEventSize = 11, // 9 + 2
+		MaxEncodedDynamicEventSize = 13, // 9 + 2 + 2,
+		FullBufferThreshold = MaxBufferSize - MaxEncodedEventSize,
+		DynamicEventFullBufferThreshold = MaxBufferSize - MaxEncodedDynamicEventSize
+	};
+
+	struct FThreadState
+	{
+		TMap<uint32, uint16> KnownDynamicStrings;
+		uint64 CycleBase;
+		uint64 LastCycle;
+		uint32 ThreadId;
+		uint32 BufferSize;
+		uint8 Buffer[MaxBufferSize];
+	};
+
+	FORCENOINLINE static FThreadState* InitThreadState(uint64 CycleBase);
+	FORCENOINLINE static void FlushThreadBuffer(FThreadState* ThreadState, uint64 NewCycleBase);
+	FORCENOINLINE static uint16 NewDynamicEventName(const TCHAR* Name);
+	static void EncodeSpecId(uint16 SpecId, uint8*& BufferPtr);
+
+	static uint32 TlsSlot;
+
+private:
+	friend struct FCpuProfilerEventSpec;
+
+	FCpuProfilerTraceInternal();
+
+	TAtomic<uint32> NextSpecId;
+	TAtomic<uint16> NextDynamicEventNameId;
+
+	static FCpuProfilerTraceInternal& Instance();
+};
+
+
+uint32 FCpuProfilerTraceInternal::TlsSlot;
+
+FCpuProfilerTraceInternal::FCpuProfilerTraceInternal()
+	: NextSpecId(1)
+	, NextDynamicEventNameId(0)
+{
+	TlsSlot = FPlatformTLS::AllocTlsSlot();
+}
+
+void FCpuProfilerTraceInternal::EncodeSpecId(uint16 SpecId, uint8*& BufferPtr)
+{
+	*BufferPtr++ = uint8(SpecId);
+	*BufferPtr++ = uint8(SpecId >> 8);
+}
+
+FCpuProfilerTraceInternal::FThreadState* FCpuProfilerTraceInternal::InitThreadState(uint64 CycleBase)
+{
+	FThreadState* ThreadState = new FThreadState();
+	ThreadState->BufferSize = 0;
+	ThreadState->ThreadId = FPlatformTLS::GetCurrentThreadId();
+	ThreadState->CycleBase = CycleBase;
+	ThreadState->LastCycle = CycleBase;
+	FPlatformTLS::SetTlsValue(TlsSlot, ThreadState);
+	return ThreadState;
+}
+
+void FCpuProfilerTraceInternal::FlushThreadBuffer(FThreadState* ThreadState, uint64 NewCycleBase)
+{
+	UE_TRACE_LOG(CpuProfiler, EventBatch, ThreadState->BufferSize)
+		<< EventBatch.CycleBase(ThreadState->CycleBase)
+		<< EventBatch.ThreadId(ThreadState->ThreadId)
+		<< EventBatch.BufferSize(ThreadState->BufferSize)
+		<< EventBatch.Attachment(ThreadState->Buffer, ThreadState->BufferSize);
+
+	ThreadState->BufferSize = 0;
+	ThreadState->CycleBase = NewCycleBase;
+	ThreadState->LastCycle = NewCycleBase;
+}
+
+uint16 FCpuProfilerTraceInternal::NewDynamicEventName(const TCHAR* Name)
+{
+	uint16 NameLength = FCString::Strlen(Name);
+	uint16 Id = FCpuProfilerTraceInternal::Instance().NextDynamicEventNameId++;
+	auto NameCopy = [NameLength, Name](uint8* Out) {
+		int32 i = 0;
+		for (; i < NameLength; ++i)
+		{
+			Out[i] = Name[i] & 0x7f;
+		}
+		Out[i] = 0;
+	};
+	UE_TRACE_LOG(CpuProfiler, DynamicEventName, NameLength + 1)
+		<< DynamicEventName.Id(Id)
+		<< DynamicEventName.NameSize(NameLength)
+		<< DynamicEventName.Attachment(NameCopy);
+	return Id;
+}
+
+FCpuProfilerTraceInternal& FCpuProfilerTraceInternal::Instance()
+{
+	static FCpuProfilerTraceInternal I;
+	return I;
+}
+
+void FCpuProfilerTrace::OutputBeginEvent(uint16 SpecId)
+{
+	uint64 Cycle = FPlatformTime::Cycles64();
+	FCpuProfilerTraceInternal::FThreadState* ThreadState = (FCpuProfilerTraceInternal::FThreadState*)FPlatformTLS::GetTlsValue(FCpuProfilerTraceInternal::TlsSlot);
+	if (!ThreadState)
+	{
+		ThreadState = FCpuProfilerTraceInternal::InitThreadState(Cycle);
+	}
+	uint64 CycleDiff = Cycle - ThreadState->LastCycle;
+	ThreadState->LastCycle = Cycle;
+	if (ThreadState->BufferSize >= FCpuProfilerTraceInternal::FullBufferThreshold)
+	{
+		FCpuProfilerTraceInternal::FlushThreadBuffer(ThreadState, Cycle);
+		CycleDiff = 0;
+	}
+
+	uint8* BufferPtr = ThreadState->Buffer + ThreadState->BufferSize;
+	FTraceUtils::Encode7bit((CycleDiff << 1) | 1ull, BufferPtr);
+	FCpuProfilerTraceInternal::EncodeSpecId(SpecId, BufferPtr);
+	ThreadState->BufferSize = BufferPtr - ThreadState->Buffer;
+}
+
+void FCpuProfilerTrace::OutputBeginDynamicEvent(const TCHAR* Name)
+{
+	uint64 Cycle = FPlatformTime::Cycles64();
+	FCpuProfilerTraceInternal::FThreadState* ThreadState = (FCpuProfilerTraceInternal::FThreadState*)FPlatformTLS::GetTlsValue(FCpuProfilerTraceInternal::TlsSlot);
+	if (!ThreadState)
+	{
+		ThreadState = FCpuProfilerTraceInternal::InitThreadState(Cycle);
+	}
+	uint64 CycleDiff = Cycle - ThreadState->LastCycle;
+	ThreadState->LastCycle = Cycle;
+	if (ThreadState->BufferSize >= FCpuProfilerTraceInternal::DynamicEventFullBufferThreshold)
+	{
+		FCpuProfilerTraceInternal::FlushThreadBuffer(ThreadState, Cycle);
+		CycleDiff = 0;
+	}
+
+	uint8* BufferPtr = ThreadState->Buffer + ThreadState->BufferSize;
+	FTraceUtils::Encode7bit((CycleDiff << 1) | 1ull, BufferPtr);
+	FCpuProfilerTraceInternal::EncodeSpecId(0, BufferPtr);
+
+	uint32 NameHash = FCrc::StrCrc32(Name);
+	uint16* NameId = ThreadState->KnownDynamicStrings.Find(NameHash);
+	if (!NameId)
+	{
+		uint16 NewNameId = FCpuProfilerTraceInternal::NewDynamicEventName(Name);
+		FCpuProfilerTraceInternal::EncodeSpecId(NewNameId, BufferPtr);
+		ThreadState->KnownDynamicStrings.Add(NameHash, NewNameId);
+	}
+	else
+	{
+		FCpuProfilerTraceInternal::EncodeSpecId(*NameId, BufferPtr);
+	}
+	ThreadState->BufferSize = BufferPtr - ThreadState->Buffer;
+}
+
+void FCpuProfilerTrace::OutputEndEvent()
+{
+	uint64 Cycle = FPlatformTime::Cycles64();
+	FCpuProfilerTraceInternal::FThreadState* ThreadState = (FCpuProfilerTraceInternal::FThreadState*)FPlatformTLS::GetTlsValue(FCpuProfilerTraceInternal::TlsSlot);
+	uint64 CycleDiff = Cycle - ThreadState->LastCycle;
+	ThreadState->LastCycle = Cycle;
+	if (ThreadState->BufferSize >= FCpuProfilerTraceInternal::FullBufferThreshold)
+	{
+		FCpuProfilerTraceInternal::FlushThreadBuffer(ThreadState, Cycle);
+		CycleDiff = 0;
+	}
+	uint8* BufferPtr = ThreadState->Buffer + ThreadState->BufferSize;
+	FTraceUtils::Encode7bit(CycleDiff << 1, BufferPtr);
+	ThreadState->BufferSize = BufferPtr - ThreadState->Buffer;
+}
+
+uint64 FCpuProfilerTrace::GetCurrentCycle()
+{
+	return FPlatformTime::Cycles64();
+}
+
+uint16 FCpuProfilerEventSpec::AssignId(const TCHAR* Name, ECpuProfilerGroup Group)
+{
+	uint32 WideId = FCpuProfilerTraceInternal::Instance().NextSpecId++;
+	check(WideId <= 0xFFFF);
+	uint16 Id = uint16(WideId);
+	uint16 NameLength = FCString::Strlen(Name);
+
+	auto NameCopy = [NameLength, Name](uint8* Out) {
+		int32 i = 0;
+		for (; i < NameLength; ++i)
+		{
+			Out[i] = Name[i] & 0x7f;
+		}
+		Out[i] = 0;
+	};
+
+	UE_TRACE_LOG(CpuProfiler, EventSpec, NameLength + 1)
+		<< EventSpec.Id(Id)
+		<< EventSpec.Group(uint16(Group))
+		<< EventSpec.NameSize(NameLength)
+		<< EventSpec.Attachment(NameCopy);
+	return Id;
+}
+
+
+#endif
\ No newline at end of file
diff --git a/Engine/Source/Runtime/Core/Private/ProfilingDebugging/MiscTrace.cpp b/Engine/Source/Runtime/Core/Private/ProfilingDebugging/MiscTrace.cpp
new file mode 100644
--- /dev/null
+++ b/Engine/Source/Runtime/Core/Private/ProfilingDebugging/MiscTrace.cpp
@@ -0,0 +1,197 @@
+// Copyright 1998-2019 Epic Games, Inc. All Rights Reserved.
+#include "ProfilingDebugging/MiscTrace.h"
+
+#if MISCTRACE_ENABLED
+
+#include "Trace/Trace.h"
+#include "Misc/CString.h"
+#include "HAL/PlatformTLS.h"
+#include "HAL/PlatformTime.h"
+
+UE_TRACE_EVENT_BEGIN(Misc, RegisterGameThread)
+	UE_TRACE_EVENT_FIELD(uint32, ThreadId)
+UE_TRACE_EVENT_END()
+
+UE_TRACE_EVENT_BEGIN(Misc, CreateThread)
+	UE_TRACE_EVENT_FIELD(uint32, CurrentThreadId)
+	UE_TRACE_EVENT_FIELD(uint32, CreatedThreadId)
+	UE_TRACE_EVENT_FIELD(uint32, Priority)
+	UE_TRACE_EVENT_FIELD(uint16, NameSize)
+UE_TRACE_EVENT_END()
+
+UE_TRACE_EVENT_BEGIN(Misc, SetThreadGroup)
+	UE_TRACE_EVENT_FIELD(uint32, ThreadId)
+	UE_TRACE_EVENT_FIELD(uint16, GroupNameLength)
+UE_TRACE_EVENT_END()
+
+UE_TRACE_EVENT_BEGIN(Misc, BeginThreadGroupScope)
+	UE_TRACE_EVENT_FIELD(uint32, CurrentThreadId)
+	UE_TRACE_EVENT_FIELD(uint16, GroupNameLength)
+UE_TRACE_EVENT_END()
+
+UE_TRACE_EVENT_BEGIN(Misc, EndThreadGroupScope)
+	UE_TRACE_EVENT_FIELD(uint32, CurrentThreadId)
+UE_TRACE_EVENT_END()
+
+UE_TRACE_EVENT_BEGIN(Misc, BookmarkSpec)
+	UE_TRACE_EVENT_FIELD(const void*, BookmarkPoint)
+	UE_TRACE_EVENT_FIELD(int32, Line)
+	UE_TRACE_EVENT_FIELD(uint16, FileNameLength)
+	UE_TRACE_EVENT_FIELD(uint16, FormatStringLength)
+UE_TRACE_EVENT_END()
+
+UE_TRACE_EVENT_BEGIN(Misc, Bookmark)
+	UE_TRACE_EVENT_FIELD(uint64, Cycle)
+	UE_TRACE_EVENT_FIELD(const void*, BookmarkPoint)
+UE_TRACE_EVENT_END()
+
+UE_TRACE_EVENT_BEGIN(Misc, BeginFrame)
+	UE_TRACE_EVENT_FIELD(uint64, Cycle)
+	UE_TRACE_EVENT_FIELD(uint8, FrameType)
+UE_TRACE_EVENT_END()
+
+UE_TRACE_EVENT_BEGIN(Misc, EndFrame)
+	UE_TRACE_EVENT_FIELD(uint64, Cycle)
+	UE_TRACE_EVENT_FIELD(uint8, FrameType)
+UE_TRACE_EVENT_END()
+
+void FMiscTrace::OutputRegisterGameThread(uint32 Id)
+{
+	UE_TRACE_LOG(Misc, RegisterGameThread)
+		<< RegisterGameThread.ThreadId(FPlatformTLS::GetCurrentThreadId());
+}
+
+void FMiscTrace::OutputCreateThread(uint32 Id, const TCHAR* Name, uint32 Priority)
+{
+	uint16 NameLength = FCString::Strlen(Name);
+	auto NameCopy = [NameLength, Name](uint8* Out) {
+		int32 i = 0;
+		for (; i < NameLength; ++i)
+		{
+			Out[i] = Name[i] & 0x7f;
+		}
+		Out[i] = 0;
+	};
+	UE_TRACE_LOG(Misc, CreateThread, NameLength + 1)
+		<< CreateThread.CurrentThreadId(FPlatformTLS::GetCurrentThreadId())
+		<< CreateThread.CreatedThreadId(Id)
+		<< CreateThread.Priority(Priority)
+		<< CreateThread.NameSize(NameLength)
+		<< CreateThread.Attachment(NameCopy);
+}
+
+static const char* GetThreadGroupName(ETraceThreadGroup Group)
+{
+	switch (Group)
+	{
+	case TraceThreadGroup_Render:
+		return "Render";
+	case TraceThreadGroup_TaskGraphHigh:
+		return "TaskGraphHigh";
+	case TraceThreadGroup_TaskGraphNormal:
+		return "TaskGraphNormal";
+	case TraceThreadGroup_TaskGraphLow:
+		return "TaskGraphLow";
+	case TraceThreadGroup_ThreadPool:
+		return "ThreadPool";
+	case TraceThreadGroup_BackgroundThreadPool:
+		return "BackgroundThreadPool";
+	case TraceThreadGroup_LargeThreadPool:
+		return "LargeThreadPool";
+	case TraceThreadGroup_AsyncLoading:
+		return "AsyncLoading";
+	case TraceThreadGroup_IOThreadPool:
+		return "IOThreadPool";
+	}
+	check(false);
+	return nullptr;
+}
+
+void FMiscTrace::OutputSetThreadGroup(uint32 Id, ETraceThreadGroup Group)
+{
+	const char* GroupName = GetThreadGroupName(Group);
+	uint16 NameLength = strlen(GroupName);
+	auto NameCopy = [NameLength, GroupName](uint8* Out) {
+		int32 i = 0;
+		for (; i < NameLength; ++i)
+		{
+			Out[i] = GroupName[i];
+		}
+		Out[i] = 0;
+	};
+	UE_TRACE_LOG(Misc, SetThreadGroup, NameLength + 1)
+		<< SetThreadGroup.ThreadId(Id)
+		<< SetThreadGroup.GroupNameLength(NameLength)
+		<< SetThreadGroup.Attachment(NameCopy);
+}
+
+void FMiscTrace::OutputBeginThreadGroupScope(ETraceThreadGroup Group)
+{
+	const char* GroupName = GetThreadGroupName(Group);
+	uint16 NameLength = strlen(GroupName);
+	auto NameCopy = [NameLength, GroupName](uint8* Out) {
+		int32 i = 0;
+		for (; i < NameLength; ++i)
+		{
+			Out[i] = GroupName[i];
+		}
+		Out[i] = 0;
+	};
+	UE_TRACE_LOG(Misc, BeginThreadGroupScope, NameLength + 1)
+		<< BeginThreadGroupScope.CurrentThreadId(FPlatformTLS::GetCurrentThreadId())
+		<< BeginThreadGroupScope.Attachment(NameCopy);
+}
+
+void FMiscTrace::OutputEndThreadGroupScope()
+{
+	UE_TRACE_LOG(Misc, EndThreadGroupScope)
+		<< EndThreadGroupScope.CurrentThreadId(FPlatformTLS::GetCurrentThreadId());
+}
+
+void FMiscTrace::OutputBookmarkSpec(const void* BookmarkPoint, const ANSICHAR* File, int32 Line, const TCHAR* Format)
+{
+	uint16 FileNameLength = strlen(File);
+	int32 FormatStringLength = FCString::Strlen(Format);
+	auto StringCopyFunc = [FileNameLength, FormatStringLength, File, Format](uint8* Out) {
+		ANSICHAR* OutAnsiChar = reinterpret_cast<ANSICHAR*>(Out);
+		for (int i = 0; i < FileNameLength; ++i)
+		{
+			*OutAnsiChar++ = File[i];
+		}
+		TCHAR* OutTChar = reinterpret_cast<TCHAR*>(OutAnsiChar);
+		for (int i = 0; i < FormatStringLength; ++i)
+		{
+			*OutTChar++ = Format[i];
+		}
+	};
+	UE_TRACE_LOG(Misc, BookmarkSpec, FileNameLength + FormatStringLength * 2)
+		<< BookmarkSpec.BookmarkPoint(BookmarkPoint)
+		<< BookmarkSpec.Line(Line)
+		<< BookmarkSpec.FileNameLength(FileNameLength)
+		<< BookmarkSpec.FormatStringLength(FormatStringLength)
+		<< BookmarkSpec.Attachment(StringCopyFunc);
+}
+
+void FMiscTrace::OutputBookmarkInternal(const void* BookmarkPoint, uint16 EncodedFormatArgsSize, uint8* EncodedFormatArgs)
+{
+	UE_TRACE_LOG(Misc, Bookmark, EncodedFormatArgsSize)
+		<< Bookmark.Cycle(FPlatformTime::Cycles64())
+		<< Bookmark.BookmarkPoint(BookmarkPoint)
+		<< Bookmark.Attachment(EncodedFormatArgs, EncodedFormatArgsSize);
+}
+
+void FMiscTrace::OutputBeginFrame(ETraceFrameType FrameType)
+{
+	UE_TRACE_LOG(Misc, BeginFrame)
+		<< BeginFrame.Cycle(FPlatformTime::Cycles64())
+		<< BeginFrame.FrameType(uint8(FrameType));
+}
+
+void FMiscTrace::OutputEndFrame(ETraceFrameType FrameType)
+{
+	UE_TRACE_LOG(Misc, EndFrame)
+		<< EndFrame.Cycle(FPlatformTime::Cycles64())
+		<< EndFrame.FrameType(uint8(FrameType));
+}
+
+#endif
