Fixed TraceServices' compile errors when compiling the vanilla editor target, with PCH disabled, and no unity files.

diff --git a/Engine/Source/Developer/TraceAnalysis/Private/DataStream.h b/Engine/Source/Developer/TraceAnalysis/Private/DataStream.h
--- a/Engine/Source/Developer/TraceAnalysis/Private/DataStream.h
+++ b/Engine/Source/Developer/TraceAnalysis/Private/DataStream.h
@@ -5,8 +5,12 @@
 #pragma once
 
 #include "CoreTypes.h"
+#include "Logging/LogMacros.h"
+#include "Containers/UnrealString.h"
 #include "Trace/DataStream.h"
 
+#include <memory.h>
+
 namespace Trace
 {
 
diff --git a/Engine/Source/Developer/TraceServices/Private/TraceServicesModule.cpp b/Engine/Source/Developer/TraceServices/Private/TraceServicesModule.cpp
--- a/Engine/Source/Developer/TraceServices/Private/TraceServicesModule.cpp
+++ b/Engine/Source/Developer/TraceServices/Private/TraceServicesModule.cpp
@@ -1,5 +1,6 @@
 // Copyright 1998-2019 Epic Games, Inc. All Rights Reserved.
 
+#include "Modules/ModuleManager.h"
 #include "TraceServices/ITraceServicesModule.h"
 #include "SessionServicePrivate.h"
 #include "AnalysisServicePrivate.h"
@@ -42,4 +43,4 @@ TSharedPtr<Trace::IAnalysisService> FTraceServicesModule::GetAnalysisService()
 	return AnalysisService;
 }
 
-IMPLEMENT_MODULE(FTraceServicesModule, TraceServices);
\ No newline at end of file
+IMPLEMENT_MODULE(FTraceServicesModule, TraceServices)
diff --git a/Engine/Source/Developer/TraceServices/Public/TraceServices/AnalysisService.h b/Engine/Source/Developer/TraceServices/Public/TraceServices/AnalysisService.h
--- a/Engine/Source/Developer/TraceServices/Public/TraceServices/AnalysisService.h
+++ b/Engine/Source/Developer/TraceServices/Public/TraceServices/AnalysisService.h
@@ -4,11 +4,13 @@
 
 #include "CoreTypes.h"
 #include "Containers/UnrealString.h"
+#include "ProfilingDebugging/MiscTrace.h"
 #include "Templates/SharedPointer.h"
 #include "Trace/Store.h"
 #include "Delegates/Delegate.h"
 #include "GenericPlatform/GenericPlatformAffinity.h"
 #include "Serialization/LoadTimeTrace.h"
+#include "Templates/UniquePtr.h"
 
 namespace Trace
 {
@@ -281,4 +283,4 @@ public:
 	virtual FAnalysisFinishedEvent& OnAnalysisFinished() = 0;
 };
 
-}
\ No newline at end of file
+}
diff --git a/Engine/Source/Runtime/CoreUObject/Public/Serialization/LoadTimeTrace.h b/Engine/Source/Runtime/CoreUObject/Public/Serialization/LoadTimeTrace.h
--- a/Engine/Source/Runtime/CoreUObject/Public/Serialization/LoadTimeTrace.h
+++ b/Engine/Source/Runtime/CoreUObject/Public/Serialization/LoadTimeTrace.h
@@ -2,6 +2,8 @@
 
 #pragma once
 
+#include "CoreTypes.h"
+
 #if !IS_PROGRAM && !UE_BUILD_SHIPPING
 #define LOADTIMEPROFILERTRACE_ENABLED 0
 #else
@@ -86,4 +88,4 @@ struct FLoadTimeProfilerTrace
 #define TRACE_LOADTIME_END_LOAD_STREAMABLE_HANDLE(...)
 #define TRACE_LOADTIME_STREAMABLE_HANDLE_REQUEST_ASSOCIATION(...)
 #define TRACE_LOADTIME_WAIT_FOR_STREAMABLE_HANDLE_SCOPE(...)
-#endif
\ No newline at end of file
+#endif
