Fixed Linux/Clang compile errors

diff --git a/Engine/Source/Developer/TraceAnalysis/Private/Analysis/Engine.cpp b/Engine/Source/Developer/TraceAnalysis/Private/Analysis/Engine.cpp
--- a/Engine/Source/Developer/TraceAnalysis/Private/Analysis/Engine.cpp
+++ b/Engine/Source/Developer/TraceAnalysis/Private/Analysis/Engine.cpp
@@ -12,7 +12,6 @@
 #include "Trace/Private/DataDecoder.h"
 #endif // 0
 #include "Trace/Private/Event.h"
-#include "Trace/Private/Field.h"
 
 namespace Trace
 {
diff --git a/Engine/Source/Developer/TraceAnalysis/Private/Recorder/Recorder.cpp b/Engine/Source/Developer/TraceAnalysis/Private/Recorder/Recorder.cpp
--- a/Engine/Source/Developer/TraceAnalysis/Private/Recorder/Recorder.cpp
+++ b/Engine/Source/Developer/TraceAnalysis/Private/Recorder/Recorder.cpp
@@ -2,8 +2,10 @@
 // Copyright 1998-2019 Epic Games, Inc. All Rights Reserved.
 //
 
-#if PLATFORM_WINDOWS
 #include "Trace/Recorder.h"
+#include "CoreTypes.h"
+
+#if PLATFORM_WINDOWS
 #include "Trace/Store.h"
 #include "DataStream.h"
 #include "Templates/UniquePtr.h"
diff --git a/Engine/Source/Developer/TraceAnalysis/Public/Trace/Analyzer.h b/Engine/Source/Developer/TraceAnalysis/Public/Trace/Analyzer.h
--- a/Engine/Source/Developer/TraceAnalysis/Public/Trace/Analyzer.h
+++ b/Engine/Source/Developer/TraceAnalysis/Public/Trace/Analyzer.h
@@ -5,6 +5,10 @@
 #pragma once
 
 #include "CoreTypes.h"
+#include "Logging/LogMacros.h"
+#include "Trace/Private/Field.h"
+
+#include <memory.h>
 
 namespace Trace
 {
diff --git a/Engine/Source/Developer/TraceServices/Private/Model/Threads.h b/Engine/Source/Developer/TraceServices/Private/Model/Threads.h
--- a/Engine/Source/Developer/TraceServices/Private/Model/Threads.h
+++ b/Engine/Source/Developer/TraceServices/Private/Model/Threads.h
@@ -27,9 +27,9 @@ private:
 	{
 		uint32 Id = 0;
 		FString Name;
-		uint32 GroupSortOrder = ~0ul;
-		uint32 PrioritySortOrder = ~0ul;
-		uint32 FallbackSortOrder = ~0ul;
+		uint32 GroupSortOrder = ~0u;
+		uint32 PrioritySortOrder = ~0u;
+		uint32 FallbackSortOrder = ~0u;
 		const TCHAR* GroupName = nullptr;
 		bool IsGameThread = false;
 
