Compile fixes

diff --git a/Engine/Source/Developer/TraceServices/Private/Model/MonotonicTimeline.h b/Engine/Source/Developer/TraceServices/Private/Model/MonotonicTimeline.h
--- a/Engine/Source/Developer/TraceServices/Private/Model/MonotonicTimeline.h
+++ b/Engine/Source/Developer/TraceServices/Private/Model/MonotonicTimeline.h
@@ -49,7 +49,7 @@ public:
 
 		AddScopeEntry(DetailLevels[0], StartTime, true);
 		AddEvent(DetailLevels[0], Event);
-		FDetailLevelInsertionState::FDepthState& Lod0DepthState = DetailLevels[0].InsertionState.DepthStates[CurrentDepth];
+		FDetailLevelDepthState& Lod0DepthState = DetailLevels[0].InsertionState.DepthStates[CurrentDepth];
 		Lod0DepthState.EnterTime = StartTime;
 		Lod0DepthState.DominatingEvent = Event;
 		//Lod0DepthState.DebugDominatingEventType = Owner.EventTypes[TypeId];
@@ -57,13 +57,13 @@ public:
 		for (int32 DetailLevelIndex = 1; DetailLevelIndex < SettingsType::DetailLevelsCount; ++DetailLevelIndex)
 		{
 			FDetailLevel& DetailLevel = DetailLevels[DetailLevelIndex];
-			FDetailLevelInsertionState::FDepthState& CurrentDepthState = DetailLevel.InsertionState.DepthStates[CurrentDepth];
+			FDetailLevelDepthState& CurrentDepthState = DetailLevel.InsertionState.DepthStates[CurrentDepth];
 
 			if (CurrentDepthState.PendingScopeEnterIndex < 0 || StartTime >= CurrentDepthState.EnterTime + DetailLevel.Resolution)
 			{
 				for (int32 Depth = DetailLevel.InsertionState.PendingDepth; Depth >= CurrentDepth; --Depth)
 				{
-					FDetailLevelInsertionState::FDepthState& DepthState = DetailLevel.InsertionState.DepthStates[Depth];
+					FDetailLevelDepthState& DepthState = DetailLevel.InsertionState.DepthStates[Depth];
 					check(DepthState.PendingScopeEnterIndex >= 0);
 					AddScopeEntry(DetailLevel, DepthState.ExitTime, false);
 
@@ -323,28 +323,28 @@ private:
 		uint16 InitialStackCount = 0;
 	};
 
-	struct FDetailLevelInsertionState
+	struct FDetailLevelDepthState
 	{
-		struct FDepthState
-		{
-			int64 PendingScopeEnterIndex = -1;
-			int64 PendingEventIndex = -1;
+		int64 PendingScopeEnterIndex = -1;
+		int64 PendingEventIndex = -1;
 
-			EventType DominatingEvent;
-			double DominatingEventStartTime = 0.0;
-			double DominatingEventEndTime = 0.0;
-			double DominatingEventDuration = 0.0;
+		EventType DominatingEvent;
+		double DominatingEventStartTime = 0.0;
+		double DominatingEventEndTime = 0.0;
+		double DominatingEventDuration = 0.0;
 
-			double EnterTime = 0.0;
-			double ExitTime = 0.0;
+		double EnterTime = 0.0;
+		double ExitTime = 0.0;
 
-			//const FTimelineEventType* DebugDominatingEventType;
-		};
+		//const FTimelineEventType* DebugDominatingEventType;
+	};
 
+	struct FDetailLevelInsertionState
+	{
 		double LastTime = -1.0;
 		uint16 CurrentDepth = 0;
 		int32 PendingDepth = -1;
-		FDepthState DepthStates[SettingsType::MaxDepth];
+		FDetailLevelDepthState DepthStates[SettingsType::MaxDepth];
 		FEventStackEntry EventStack[SettingsType::MaxDepth];
 		FEventScopeEntryPage* CurrentScopeEntryPage = nullptr;
 	};
@@ -368,9 +368,9 @@ private:
 
 	void UpdateDominatingEvent(FDetailLevel& DetailLevel, int32 Depth, double CurrentTime)
 	{
-		FDetailLevelInsertionState::FDepthState& Lod0DepthState = DetailLevels[0].InsertionState.DepthStates[Depth];
+		FDetailLevelDepthState& Lod0DepthState = DetailLevels[0].InsertionState.DepthStates[Depth];
 		double Lod0EventDuration = CurrentTime - Lod0DepthState.EnterTime;
-		FDetailLevelInsertionState::FDepthState& CurrentDepthState = DetailLevel.InsertionState.DepthStates[Depth];
+		FDetailLevelDepthState& CurrentDepthState = DetailLevel.InsertionState.DepthStates[Depth];
 		if (Lod0EventDuration > CurrentDepthState.DominatingEventDuration)
 		{
 			check(CurrentDepthState.PendingScopeEnterIndex >= 0);
diff --git a/Engine/Source/Runtime/Core/Public/ProfilingDebugging/CpuProfilerTrace.h b/Engine/Source/Runtime/Core/Public/ProfilingDebugging/CpuProfilerTrace.h
--- a/Engine/Source/Runtime/Core/Public/ProfilingDebugging/CpuProfilerTrace.h
+++ b/Engine/Source/Runtime/Core/Public/ProfilingDebugging/CpuProfilerTrace.h
@@ -59,16 +59,6 @@ struct FCpuProfilerDynamicEventScope
 	}
 };
 
-#define TRACE_CPUPROFILER_EVENT_SCOPE_GROUP(Name, Group) \
-	static uint16 PREPROCESSOR_JOIN(__CpuProfilerEventSpecId, __LINE__); \
-	if (PREPROCESSOR_JOIN(__CpuProfilerEventSpecId, __LINE__) == 0) { \
-		PREPROCESSOR_JOIN(__CpuProfilerEventSpecId, __LINE__) = FCpuProfilerEventSpec::AssignId(L###Name, Group); \
-	} \
-	FCpuProfilerEventScope PREPROCESSOR_JOIN(__CpuProfilerEventScope, __LINE__)(PREPROCESSOR_JOIN(__CpuProfilerEventSpecId, __LINE__));
-
-#define TRACE_CPUPROFILER_EVENT_SCOPE(Name) \
-	TRACE_CPUPROFILER_EVENT_SCOPE_GROUP(Name, CpuProfilerGroup_Default);
-
 #define TRACE_CPUPROFILER_EVENT_SCOPE_TEXT_GROUP(Name, Group) \
 	static uint16 PREPROCESSOR_JOIN(__CpuProfilerEventSpecId, __LINE__); \
 	if (PREPROCESSOR_JOIN(__CpuProfilerEventSpecId, __LINE__) == 0) { \
@@ -79,6 +69,12 @@ struct FCpuProfilerDynamicEventScope
 #define TRACE_CPUPROFILER_EVENT_SCOPE_TEXT(Name) \
 	TRACE_CPUPROFILER_EVENT_SCOPE_TEXT_GROUP(Name, CpuProfilerGroup_Default)
 
+#define TRACE_CPUPROFILER_EVENT_SCOPE_GROUP(Name, Group) \
+	TRACE_CPUPROFILER_EVENT_SCOPE_TEXT_GROUP(TEXT(#Name), Group);
+
+#define TRACE_CPUPROFILER_EVENT_SCOPE(Name) \
+	TRACE_CPUPROFILER_EVENT_SCOPE_TEXT_GROUP(TEXT(#Name), CpuProfilerGroup_Default);
+
 #define TRACE_CPUPROFILER_DYNAMIC_EVENT_SCOPE(Name) \
 	FCpuProfilerDynamicEventScope ANONYMOUS_VARIABLE(__Scope)(Name);
 
diff --git a/Engine/Source/Runtime/CoreUObject/Private/Serialization/LoadTimeTracePrivate.h b/Engine/Source/Runtime/CoreUObject/Private/Serialization/LoadTimeTracePrivate.h
--- a/Engine/Source/Runtime/CoreUObject/Private/Serialization/LoadTimeTracePrivate.h
+++ b/Engine/Source/Runtime/CoreUObject/Private/Serialization/LoadTimeTracePrivate.h
@@ -120,7 +120,7 @@ struct FLoadTimeProfilerTracePrivate
 	FLoadTimeProfilerTracePrivate::FFlushAsyncLoadingScope __LoadTimeTraceFlushAsyncLoadingScope(RequestId);
 
 #define TRACE_LOADTIME_CLASS_INFO(Class, Name) \
-	FMiscTrace::OutputClassInfo(Class, Name);
+	FLoadTimeProfilerTracePrivate::OutputClassInfo(Class, Name);
 
 #else
 
