Use a more reasonable amount of stack space to appease PVS analysis.

diff --git a/Engine/Source/Developer/TraceServices/Private/Model/Bookmarks.cpp b/Engine/Source/Developer/TraceServices/Private/Model/Bookmarks.cpp
--- a/Engine/Source/Developer/TraceServices/Private/Model/Bookmarks.cpp
+++ b/Engine/Source/Developer/TraceServices/Private/Model/Bookmarks.cpp
@@ -34,8 +34,8 @@ void FBookmarkProvider::AppendBookmark(double Time, uint64 BookmarkPoint, uint16
 	FBookmarkSpec& Spec = GetSpec(BookmarkPoint);
 	TSharedRef<FBookmarkInternal> Bookmark = MakeShared<FBookmarkInternal>();
 	Bookmark->Time = Time;
-	TCHAR Buffer[65535];
-	FFormatArgsHelper::Format(Buffer, 65534, *Spec.FormatString, FormatArgs);
+	TCHAR Buffer[4096];
+	FFormatArgsHelper::Format(Buffer, 4096, *Spec.FormatString, FormatArgs);
 	Bookmark->Text = Buffer;
 	Bookmarks.Add(Bookmark);
 }
