Make UE_TRACE_TABLE_COLUMN and UE_TRACE_TABLE_PROJECTED_COLUMN only defined for windows again until Stefan and Ionut can investigate Mac compile errors
#rnx

diff --git a/Engine/Source/Developer/TraceServices/Private/Model/Tables.h b/Engine/Source/Developer/TraceServices/Private/Model/Tables.h
--- a/Engine/Source/Developer/TraceServices/Private/Model/Tables.h
+++ b/Engine/Source/Developer/TraceServices/Private/Model/Tables.h
@@ -23,6 +23,8 @@
 		::Trace::FColumnValueContainer GetColumnValueInternal(const RowType& Row) const { return 0; } \
 		::Trace::TColumnDeclaration<0, RowType
 
+#if PLATFORM_WINDOWS
+
 #define UE_TRACE_TABLE_COLUMN(Name, DisplayName) \
 		> PREPROCESSOR_JOIN(__Column__, __LINE__); \
 		template<> \
@@ -44,6 +46,18 @@
 		::Trace::FColumnValueContainer GetColumnValueInternal<decltype(PREPROCESSOR_JOIN(__Column__, __LINE__))::Index>(const RowType& Row) const { return PREPROCESSOR_JOIN(__ColumnProjector__, __LINE__)(Row); } \
 		::Trace::TColumnDeclaration<decltype(PREPROCESSOR_JOIN(__Column__, __LINE__))::Index + 1, RowType
 
+#else
+
+#define UE_TRACE_TABLE_COLUMN(Name, DisplayName) \
+		> PREPROCESSOR_JOIN(__Column__, __LINE__); \
+		::Trace::TColumnDeclaration<decltype(PREPROCESSOR_JOIN(__Column__, __LINE__))::Index + 1, RowType
+
+#define UE_TRACE_TABLE_PROJECTED_COLUMN(ColumnType, DisplayName, ProjectionFunc) \
+		> PREPROCESSOR_JOIN(__Column__, __LINE__); \
+		::Trace::TColumnDeclaration<decltype(PREPROCESSOR_JOIN(__Column__, __LINE__))::Index + 1, RowType
+
+#endif
+
 #define UE_TRACE_TABLE_LAYOUT_END() \
 		> LastColumn;\
 	};
