UnrealInsights:
Table API now compiles on all target platforms

diff --git a/Engine/Source/Developer/TraceInsights/Private/Insights/ViewModels/TimerNode.cpp b/Engine/Source/Developer/TraceInsights/Private/Insights/ViewModels/TimerNode.cpp
--- a/Engine/Source/Developer/TraceInsights/Private/Insights/ViewModels/TimerNode.cpp
+++ b/Engine/Source/Developer/TraceInsights/Private/Insights/ViewModels/TimerNode.cpp
@@ -10,12 +10,12 @@ const FName FTimerNode::TypeName(TEXT("FTimerNode"));
 
 void FTimerNode::ResetAggregatedStats()
 {
-	AggregatedStats = Trace::FAggregatedTimingStats();
+	AggregatedStats = Trace::FTimingProfilerAggregatedStats();
 }
 
 ////////////////////////////////////////////////////////////////////////////////////////////////////
 
-void FTimerNode::SetAggregatedStats(const Trace::FAggregatedTimingStats& InAggregatedStats)
+void FTimerNode::SetAggregatedStats(const Trace::FTimingProfilerAggregatedStats& InAggregatedStats)
 {
 	AggregatedStats = InAggregatedStats;
 }
diff --git a/Engine/Source/Developer/TraceInsights/Private/Insights/ViewModels/TimerNode.h b/Engine/Source/Developer/TraceInsights/Private/Insights/ViewModels/TimerNode.h
--- a/Engine/Source/Developer/TraceInsights/Private/Insights/ViewModels/TimerNode.h
+++ b/Engine/Source/Developer/TraceInsights/Private/Insights/ViewModels/TimerNode.h
@@ -4,6 +4,7 @@
 
 #include "CoreMinimal.h"
 #include "TraceServices/Containers/Timelines.h"
+#include "TraceServices/Model/TimingProfiler.h"
 
 // Insights
 #include "Insights/Table/ViewModels/BaseTreeNode.h"
@@ -89,10 +90,10 @@ public:
 	/**
 	 * @return the aggregated stats for this timer.
 	 */
-	const Trace::FAggregatedTimingStats& GetAggregatedStats() const { return AggregatedStats; }
+	const Trace::FTimingProfilerAggregatedStats& GetAggregatedStats() const { return AggregatedStats; }
 
 	void ResetAggregatedStats();
-	void SetAggregatedStats(const Trace::FAggregatedTimingStats& AggregatedStats);
+	void SetAggregatedStats(const Trace::FTimingProfilerAggregatedStats& AggregatedStats);
 
 	bool IsHotPath() const { return bIsHotPath; }
 	void SetIsHotPath(bool bOnOff) { bIsHotPath = bOnOff; }
@@ -105,7 +106,7 @@ private:
 	const ETimerNodeType Type;
 
 	/** Aggregated stats. */
-	Trace::FAggregatedTimingStats AggregatedStats;
+	Trace::FTimingProfilerAggregatedStats AggregatedStats;
 
 	/** True if this tree node is on the hot path. */
 	bool bIsHotPath;
diff --git a/Engine/Source/Developer/TraceInsights/Private/Insights/Widgets/STimerTreeView.cpp b/Engine/Source/Developer/TraceInsights/Private/Insights/Widgets/STimerTreeView.cpp
--- a/Engine/Source/Developer/TraceInsights/Private/Insights/Widgets/STimerTreeView.cpp
+++ b/Engine/Source/Developer/TraceInsights/Private/Insights/Widgets/STimerTreeView.cpp
@@ -957,7 +957,7 @@ FTimerNodePtr STimerTreeView::CreateTimerNodeRec(const Trace::FTimingProfilerBut
 
 	FTimerNodePtr TimerNodePtr = MakeShareable(new FTimerNode(Id, Name, Group, Type));
 
-	Trace::FAggregatedTimingStats AggregatedStats;
+	Trace::FTimingProfilerAggregatedStats AggregatedStats;
 	AggregatedStats.InstanceCount = Node.Count;
 	AggregatedStats.TotalInclusiveTime = Node.InclusiveTime;
 	AggregatedStats.TotalExclusiveTime = Node.ExclusiveTime;
diff --git a/Engine/Source/Developer/TraceInsights/Private/Insights/Widgets/STimersViewTooltip.cpp b/Engine/Source/Developer/TraceInsights/Private/Insights/Widgets/STimersViewTooltip.cpp
--- a/Engine/Source/Developer/TraceInsights/Private/Insights/Widgets/STimersViewTooltip.cpp
+++ b/Engine/Source/Developer/TraceInsights/Private/Insights/Widgets/STimersViewTooltip.cpp
@@ -90,7 +90,7 @@ TSharedPtr<SToolTip> STimersViewTooltip::GetColumnTooltip(const Insights::FTable
 
 TSharedPtr<SToolTip> STimersViewTooltip::GetRowTooltip(const TSharedPtr<FTimerNode> TimerNodePtr)
 {
-	const Trace::FAggregatedTimingStats& Stats = TimerNodePtr->GetAggregatedStats();
+	const Trace::FTimingProfilerAggregatedStats& Stats = TimerNodePtr->GetAggregatedStats();
 
 	const FText InstanceCountText = FText::AsNumber(Stats.InstanceCount);
 
diff --git a/Engine/Source/Developer/TraceServices/Private/Common/TimelineStatistics.h b/Engine/Source/Developer/TraceServices/Private/Common/TimelineStatistics.h
--- a/Engine/Source/Developer/TraceServices/Private/Common/TimelineStatistics.h
+++ b/Engine/Source/Developer/TraceServices/Private/Common/TimelineStatistics.h
@@ -5,6 +5,21 @@
 namespace Trace
 {
 
+struct FAggregatedTimingStats
+{
+	uint64 InstanceCount = 0;
+	double TotalInclusiveTime = 0.0;
+	double MinInclusiveTime = DBL_MAX;
+	double MaxInclusiveTime = -DBL_MAX;
+	double AverageInclusiveTime = 0.0;
+	double MedianInclusiveTime = 0.0;
+	double TotalExclusiveTime = 0.0;
+	double MinExclusiveTime = DBL_MAX;
+	double MaxExclusiveTime = -DBL_MAX;
+	double AverageExclusiveTime = 0.0;
+	double MedianExclusiveTime = 0.0;
+};
+
 class FTimelineStatistics
 {
 public:
diff --git a/Engine/Source/Developer/TraceServices/Private/Model/FileActivity.cpp b/Engine/Source/Developer/TraceServices/Private/Model/FileActivity.cpp
--- a/Engine/Source/Developer/TraceServices/Private/Model/FileActivity.cpp
+++ b/Engine/Source/Developer/TraceServices/Private/Model/FileActivity.cpp
@@ -14,6 +14,23 @@ FFileActivityProvider::FFileActivityProvider(IAnalysisSession& InSession)
 	, FileActivities(InSession.GetLinearAllocator(), 1024)
 	, FileActivityTable(FileActivities)
 {
+	FileActivityTable.EditLayout().
+		AddColumn<const TCHAR*>([](const FFileActivity& Row)
+			{
+				return Row.File ? Row.File->Path : TEXT("N/A");
+			},
+			TEXT("File")).
+		AddColumn(&FFileActivity::StartTime, TEXT("StartTime")).
+		AddColumn(&FFileActivity::EndTime, TEXT("EndTime")).
+		AddColumn(&FFileActivity::ThreadId, TEXT("ThreadId")).
+		AddColumn<const TCHAR*>([](const FFileActivity& Row)
+			{
+				return GetFileActivityTypeString(Row.ActivityType);
+			},
+			TEXT("Type")).
+		AddColumn(&FFileActivity::Offset, TEXT("Offset")).
+		AddColumn(&FFileActivity::Size, TEXT("Size")).
+		AddColumn(&FFileActivity::Failed, TEXT("Failed"));
 }
 
 void FFileActivityProvider::EnumerateFileActivity(TFunctionRef<bool(const FFileInfo&, const Timeline&)> Callback) const
diff --git a/Engine/Source/Developer/TraceServices/Private/Model/FileActivity.h b/Engine/Source/Developer/TraceServices/Private/Model/FileActivity.h
--- a/Engine/Source/Developer/TraceServices/Private/Model/FileActivity.h
+++ b/Engine/Source/Developer/TraceServices/Private/Model/FileActivity.h
@@ -44,21 +44,10 @@ private:
 		TSharedPtr<TimelineInternal> ActivityTimeline;
 	};
 
-	UE_TRACE_TABLE_LAYOUT_BEGIN(FFileActivityTableLayout, FFileActivity)
-		UE_TRACE_TABLE_PROJECTED_COLUMN(TableColumnType_CString, TEXT("File"), [](const FFileActivity& Row) { return Row.File ? Row.File->Path : TEXT("N/A"); })
-		UE_TRACE_TABLE_COLUMN(StartTime, TEXT("StartTime"))
-		UE_TRACE_TABLE_COLUMN(EndTime, TEXT("EndTime"))
-		UE_TRACE_TABLE_COLUMN(ThreadId, TEXT("ThreadId"))
-		UE_TRACE_TABLE_PROJECTED_COLUMN(TableColumnType_CString, TEXT("Type"), [](const FFileActivity& Row) { return GetFileActivityTypeString(Row.ActivityType); })
-		UE_TRACE_TABLE_COLUMN(Offset, TEXT("Offset"))
-		UE_TRACE_TABLE_COLUMN(Size, TEXT("Size"))
-		UE_TRACE_TABLE_COLUMN(Failed, TEXT("Failed"))
-	UE_TRACE_TABLE_LAYOUT_END()
-
 	IAnalysisSession& Session;
 	TPagedArray<FFileInfoInternal> Files;
 	TPagedArray<FFileActivity> FileActivities;
-	TTableView<FFileActivityTableLayout> FileActivityTable;
+	TTableView<FFileActivity> FileActivityTable;
 };
 
 }
diff --git a/Engine/Source/Developer/TraceServices/Private/Model/LoadTimeProfiler.cpp b/Engine/Source/Developer/TraceServices/Private/Model/LoadTimeProfiler.cpp
--- a/Engine/Source/Developer/TraceServices/Private/Model/LoadTimeProfiler.cpp
+++ b/Engine/Source/Developer/TraceServices/Private/Model/LoadTimeProfiler.cpp
@@ -13,12 +13,75 @@ FLoadTimeProfilerProvider::FLoadTimeProfilerProvider(IAnalysisSession& InSession
 	, ClassInfos(Session.GetLinearAllocator(), 4096)
 	, Requests(Session.GetLinearAllocator(), 4096)
 	, Packages(Session.GetLinearAllocator(), 4096)
-	, Exports(Session.GetLinearAllocator(), 4090)
+	, Exports(Session.GetLinearAllocator(), 4096)
 	, MainThreadCpuTimeline(MakeShared<CpuTimelineInternal>(Session.GetLinearAllocator()))
 	, AsyncLoadingThreadCpuTimeline(MakeShared<CpuTimelineInternal>(Session.GetLinearAllocator()))
 	, RequestsTable(Requests)
 {
-	
+	RequestsTable.EditLayout().
+		AddColumn(&FLoadRequest::Name, TEXT("Name")).
+		AddColumn(&FLoadRequest::ThreadId, TEXT("ThreadId")).
+		AddColumn(&FLoadRequest::StartTime, TEXT("StartTime")).
+		AddColumn(&FLoadRequest::EndTime, TEXT("EndTime")).
+		AddColumn<int32>(
+			[](const FLoadRequest& Row)
+			{
+				return Row.Packages.Num();
+			}, 
+			TEXT("PackageCount")).
+		AddColumn(&FLoadTimeProfilerProvider::PackageSizeSum, TEXT("Size")).
+		AddColumn<const TCHAR*>(
+			[](const FLoadRequest& Row)
+			{
+				return Row.Packages.Num() ? Row.Packages[0]->Name : TEXT("N/A");
+			},
+			TEXT("FirstPackage"));
+
+	AggregatedStatsTableLayout.
+		AddColumn(&FLoadTimeProfilerAggregatedStats::Name, TEXT("Name")).
+		AddColumn(&FLoadTimeProfilerAggregatedStats::Count, TEXT("Count")).
+		AddColumn(&FLoadTimeProfilerAggregatedStats::Total, TEXT("Total")).
+		AddColumn(&FLoadTimeProfilerAggregatedStats::Min, TEXT("Min")).
+		AddColumn(&FLoadTimeProfilerAggregatedStats::Max, TEXT("Max")).
+		AddColumn(&FLoadTimeProfilerAggregatedStats::Average, TEXT("Avg")).
+		AddColumn(&FLoadTimeProfilerAggregatedStats::Median, TEXT("Med"));
+
+	PackagesTableLayout.
+		AddColumn<const TCHAR*>([](const FPackagesTableRow& Row)
+			{
+				return Row.PackageInfo->Name;
+			},
+			TEXT("Package")).
+		AddColumn<const TCHAR*>([](const FPackagesTableRow& Row)
+			{
+				return GetLoadTimeProfilerPackageEventTypeString(Row.EventType);
+			},
+			TEXT("EventType")).
+		AddColumn(&FPackagesTableRow::SerializedHeaderSize, TEXT("SerializedHeaderSize")).
+		AddColumn(&FPackagesTableRow::SerializedExportsCount, TEXT("SerializedExportsCount")).
+		AddColumn(&FPackagesTableRow::SerializedExportsSize, TEXT("SerializedExportsSize")).
+		AddColumn(&FPackagesTableRow::MainThreadTime, TEXT("MainThreadTime")).
+		AddColumn(&FPackagesTableRow::AsyncLoadingThreadTime, TEXT("AsyncLoadingThreadTime"));
+
+	ExportsTableLayout.
+		AddColumn<const TCHAR*>([](const FExportsTableRow& Row)
+			{
+				return Row.ExportInfo->Package ? Row.ExportInfo->Package->Name : TEXT("[unknown]");
+			},
+			TEXT("Package")).
+		AddColumn<const TCHAR*>([](const FExportsTableRow& Row)
+			{
+				return Row.ExportInfo->Class ? Row.ExportInfo->Class->Name : TEXT("[unknown]");
+			},
+			TEXT("Class")).
+		AddColumn<const TCHAR*>([](const FExportsTableRow& Row)
+			{
+				return GetLoadTimeProfilerObjectEventTypeString(Row.EventType);
+			},
+			TEXT("EventType")).
+		AddColumn(&FExportsTableRow::SerializedSize, TEXT("SerializedSize")).
+		AddColumn(&FExportsTableRow::MainThreadTime, TEXT("MainThreadTime")).
+		AddColumn(&FExportsTableRow::AsyncLoadingThreadTime, TEXT("AsyncLoadingThreadTime"));
 }
 
 void FLoadTimeProfilerProvider::ReadMainThreadCpuTimeline(TFunctionRef<void(const CpuTimeline &)> Callback) const
@@ -45,7 +108,7 @@ ITable<FLoadTimeProfilerAggregatedStats>* FLoadTimeProfilerProvider::CreateEvent
 	};
 	TMap<ELoadTimeProfilerPackageEventType, FAggregatedTimingStats> Aggregation;
 	FTimelineStatistics::CreateAggregation(Timelines, BucketMapper, IntervalStart, IntervalEnd, Aggregation);
-	TTable<FAggregatedStatsTableLayout>* Table = new TTable<FAggregatedStatsTableLayout>();
+	TTable<FLoadTimeProfilerAggregatedStats>* Table = new TTable<FLoadTimeProfilerAggregatedStats>(AggregatedStatsTableLayout);
 	for (const auto& KV : Aggregation)
 	{
 		FLoadTimeProfilerAggregatedStats& Row = Table->AddRow();
@@ -73,7 +136,7 @@ ITable<FLoadTimeProfilerAggregatedStats>* FLoadTimeProfilerProvider::CreateObjec
 	};
 	TMap<const Trace::FClassInfo*, FAggregatedTimingStats> Aggregation;
 	FTimelineStatistics::CreateAggregation(Timelines, BucketMapper, IntervalStart, IntervalEnd, Aggregation);
-	TTable<FAggregatedStatsTableLayout>* Table = new TTable<FAggregatedStatsTableLayout>();
+	TTable<FLoadTimeProfilerAggregatedStats>* Table = new TTable<FLoadTimeProfilerAggregatedStats>(AggregatedStatsTableLayout);
 	for (const auto& KV : Aggregation)
 	{
 		const FClassInfo* ClassInfo = KV.Key;
@@ -95,7 +158,7 @@ ITable<FLoadTimeProfilerAggregatedStats>* FLoadTimeProfilerProvider::CreateObjec
 
 ITable<FPackagesTableRow>* FLoadTimeProfilerProvider::CreatePackageDetailsTable(double IntervalStart, double IntervalEnd) const
 {
-	TTable<FPackagesTableLayout>* Table = new TTable<FPackagesTableLayout>();
+	TTable<FPackagesTableRow>* Table = new TTable<FPackagesTableRow>(PackagesTableLayout);
 
 	TMap<TTuple<const FPackageInfo*, ELoadTimeProfilerPackageEventType>, FPackagesTableRow*> PackagesMap;
 
@@ -159,7 +222,7 @@ ITable<FPackagesTableRow>* FLoadTimeProfilerProvider::CreatePackageDetailsTable(
 
 ITable<FExportsTableRow>* FLoadTimeProfilerProvider::CreateExportDetailsTable(double IntervalStart, double IntervalEnd) const
 {
-	TTable<FExportsTableLayout>* Table = new TTable<FExportsTableLayout>();
+	TTable<FExportsTableRow>* Table = new TTable<FExportsTableRow>(ExportsTableLayout);
 
 	TMap<TTuple<const FPackageExportInfo*, ELoadTimeProfilerObjectEventType>, FExportsTableRow*> ExportsMap;
 
diff --git a/Engine/Source/Developer/TraceServices/Private/Model/LoadTimeProfilerPrivate.h b/Engine/Source/Developer/TraceServices/Private/Model/LoadTimeProfilerPrivate.h
--- a/Engine/Source/Developer/TraceServices/Private/Model/LoadTimeProfilerPrivate.h
+++ b/Engine/Source/Developer/TraceServices/Private/Model/LoadTimeProfilerPrivate.h
@@ -44,45 +44,6 @@ public:
 private:
 	static uint64 PackageSizeSum(const FLoadRequest& Row);
 
-	UE_TRACE_TABLE_LAYOUT_BEGIN(FAggregatedStatsTableLayout, FLoadTimeProfilerAggregatedStats)
-		UE_TRACE_TABLE_COLUMN(Name, TEXT("Name"))
-		UE_TRACE_TABLE_COLUMN(Count, TEXT("Count"))
-		UE_TRACE_TABLE_COLUMN(Total, TEXT("Total"))
-		UE_TRACE_TABLE_COLUMN(Min, TEXT("Min"))
-		UE_TRACE_TABLE_COLUMN(Max, TEXT("Max"))
-		UE_TRACE_TABLE_COLUMN(Average, TEXT("Avg"))
-		UE_TRACE_TABLE_COLUMN(Median, TEXT("Med"))
-	UE_TRACE_TABLE_LAYOUT_END()
-
-	UE_TRACE_TABLE_LAYOUT_BEGIN(FPackagesTableLayout, FPackagesTableRow)
-		UE_TRACE_TABLE_PROJECTED_COLUMN(TableColumnType_CString, TEXT("Package"), [](const FPackagesTableRow& Row) { return Row.PackageInfo->Name; })
-		UE_TRACE_TABLE_PROJECTED_COLUMN(TableColumnType_CString, TEXT("EventType"), [](const FPackagesTableRow& Row) { return GetLoadTimeProfilerPackageEventTypeString(Row.EventType); })
-		UE_TRACE_TABLE_COLUMN(SerializedHeaderSize, TEXT("SerializedHeaderSize"))
-		UE_TRACE_TABLE_COLUMN(SerializedExportsCount, TEXT("SerializedExportsCount"))
-		UE_TRACE_TABLE_COLUMN(SerializedExportsSize, TEXT("SerializedExportsSize"))
-		UE_TRACE_TABLE_COLUMN(MainThreadTime, TEXT("MainThreadTime"))
-		UE_TRACE_TABLE_COLUMN(AsyncLoadingThreadTime, TEXT("AsyncLoadingThreadTime"))
-	UE_TRACE_TABLE_LAYOUT_END()
-
-	UE_TRACE_TABLE_LAYOUT_BEGIN(FExportsTableLayout, FExportsTableRow)
-		UE_TRACE_TABLE_PROJECTED_COLUMN(TableColumnType_CString, TEXT("Package"), [](const FExportsTableRow& Row) { return Row.ExportInfo->Package ? Row.ExportInfo->Package->Name : TEXT("[unknown]"); })
-		UE_TRACE_TABLE_PROJECTED_COLUMN(TableColumnType_CString, TEXT("Class"), [](const FExportsTableRow& Row) { return Row.ExportInfo->Class ? Row.ExportInfo->Class->Name : TEXT("[unknown]"); })
-		UE_TRACE_TABLE_PROJECTED_COLUMN(TableColumnType_CString, TEXT("EventType"), [](const FExportsTableRow& Row) { return GetLoadTimeProfilerObjectEventTypeString(Row.EventType); })
-		UE_TRACE_TABLE_COLUMN(SerializedSize, TEXT("SerializedSize"))
-		UE_TRACE_TABLE_COLUMN(MainThreadTime, TEXT("MainThreadTime"))
-		UE_TRACE_TABLE_COLUMN(AsyncLoadingThreadTime, TEXT("AsyncLoadingThreadTime"))
-	UE_TRACE_TABLE_LAYOUT_END()
-
-	UE_TRACE_TABLE_LAYOUT_BEGIN(FRequestsTableLayout, FLoadRequest)
-		UE_TRACE_TABLE_COLUMN(Name, TEXT("Name"))
-		UE_TRACE_TABLE_COLUMN(ThreadId, TEXT("ThreadId"))
-		UE_TRACE_TABLE_COLUMN(StartTime, TEXT("StartTime"))
-		UE_TRACE_TABLE_COLUMN(EndTime, TEXT("EndTime"))
-		UE_TRACE_TABLE_PROJECTED_COLUMN(TableColumnType_Int, TEXT("PackageCount"), [](const FLoadRequest& Row) { return Row.Packages.Num(); })
-		UE_TRACE_TABLE_PROJECTED_COLUMN(TableColumnType_Int, TEXT("Size"), PackageSizeSum)
-		UE_TRACE_TABLE_PROJECTED_COLUMN(TableColumnType_CString, TEXT("FirstPackage"), [](const FLoadRequest& Row) { return Row.Packages.Num() ? Row.Packages[0]->Name : TEXT("N/A"); })
-	UE_TRACE_TABLE_LAYOUT_END()
-
 	IAnalysisSession& Session;
 	TPagedArray<FClassInfo> ClassInfos;
 	TPagedArray<FLoadRequest> Requests;
@@ -93,7 +54,10 @@ private:
 	TMap<uint32, TSharedRef<CpuTimelineInternal>> AdditionalCpuTimelinesMap;
 	uint32 MainThreadId = uint32(-1);
 	uint32 AsyncLoadingThreadId = uint32(-1);
-	TTableView<FRequestsTableLayout> RequestsTable;
+	TTableView<FLoadRequest> RequestsTable;
+	TTableLayout<FLoadTimeProfilerAggregatedStats> AggregatedStatsTableLayout;
+	TTableLayout<FPackagesTableRow> PackagesTableLayout;
+	TTableLayout<FExportsTableRow> ExportsTableLayout;
 };
 
 }
diff --git a/Engine/Source/Developer/TraceServices/Private/Model/Log.cpp b/Engine/Source/Developer/TraceServices/Private/Model/Log.cpp
--- a/Engine/Source/Developer/TraceServices/Private/Model/Log.cpp
+++ b/Engine/Source/Developer/TraceServices/Private/Model/Log.cpp
@@ -17,7 +17,29 @@ FLogProvider::FLogProvider(IAnalysisSession& InSession)
 	, Messages(InSession.GetLinearAllocator(), 1024)
 	, MessagesTable(Messages)
 {
-
+	MessagesTable.EditLayout().
+		AddColumn(&FLogMessageInternal::Time, TEXT("Time")).
+		AddColumn<const TCHAR*>([](const FLogMessageInternal& Message)
+			{
+				return FOutputDeviceHelper::VerbosityToString(Message.Spec->Verbosity);
+			},
+			TEXT("Verbosity")).
+		AddColumn<const TCHAR*>([](const FLogMessageInternal& Message)
+			{
+				return Message.Spec->Category->Name;
+			},
+			TEXT("Category")).
+		AddColumn<const TCHAR*>([](const FLogMessageInternal& Message)
+			{
+				return Message.Spec->File;
+			},
+			TEXT("File")).
+		AddColumn<int32>([](const FLogMessageInternal& Message)
+			{
+				return Message.Spec->Line;
+			},
+			TEXT("Line")).
+		AddColumn(&FLogMessageInternal::Message, TEXT("Message"));
 }
 
 FLogCategory& FLogProvider::GetCategory(uint64 CategoryPointer)
diff --git a/Engine/Source/Developer/TraceServices/Private/Model/LogPrivate.h b/Engine/Source/Developer/TraceServices/Private/Model/LogPrivate.h
--- a/Engine/Source/Developer/TraceServices/Private/Model/LogPrivate.h
+++ b/Engine/Source/Developer/TraceServices/Private/Model/LogPrivate.h
@@ -58,15 +58,6 @@ public:
 private:
 	void ConstructMessage(uint64 Id, TFunctionRef<void(const FLogMessage &)> Callback) const;
 
-	UE_TRACE_TABLE_LAYOUT_BEGIN(FMessagesTableLayout, FLogMessageInternal)
-		UE_TRACE_TABLE_COLUMN(Time, TEXT("Time"))
-		UE_TRACE_TABLE_PROJECTED_COLUMN(TableColumnType_CString, TEXT("Verbosity"), [](const FLogMessageInternal& Message) { return FOutputDeviceHelper::VerbosityToString(Message.Spec->Verbosity); })
-		UE_TRACE_TABLE_PROJECTED_COLUMN(TableColumnType_CString, TEXT("Category"), [](const FLogMessageInternal& Message) { return Message.Spec->Category->Name; })
-		UE_TRACE_TABLE_PROJECTED_COLUMN(TableColumnType_CString, TEXT("File"), [](const FLogMessageInternal& Message) { return Message.Spec->File; })
-		UE_TRACE_TABLE_PROJECTED_COLUMN(TableColumnType_Int, TEXT("Line"), [](const FLogMessageInternal& Message) { return Message.Spec->Line; })
-		UE_TRACE_TABLE_COLUMN(Message, TEXT("Message"))
-	UE_TRACE_TABLE_LAYOUT_END()
-
 	enum
 	{
 		FormatBufferSize = 65536
@@ -80,7 +71,7 @@ private:
 	TPagedArray<FLogMessageInternal> Messages;
 	TCHAR FormatBuffer[FormatBufferSize];
 	TCHAR TempBuffer[FormatBufferSize];
-	TTableView<FMessagesTableLayout> MessagesTable;
+	TTableView<FLogMessageInternal> MessagesTable;
 };
 
 }
\ No newline at end of file
diff --git a/Engine/Source/Developer/TraceServices/Private/Model/NetProfilerProvider.cpp b/Engine/Source/Developer/TraceServices/Private/Model/NetProfilerProvider.cpp
--- a/Engine/Source/Developer/TraceServices/Private/Model/NetProfilerProvider.cpp
+++ b/Engine/Source/Developer/TraceServices/Private/Model/NetProfilerProvider.cpp
@@ -18,6 +18,15 @@ FNetProfilerProvider::FNetProfilerProvider(IAnalysisSession& InSession)
 {
 	// Use name index 0 to indicate that we do not know the name
 	AddNetProfilerName(TEXT("N/A"));
+
+	AggregatedStatsTableLayout.
+		AddColumn(&FNetProfilerAggregatedStats::EventTypeIndex, TEXT("EventTypeIndex")).
+		AddColumn(&FNetProfilerAggregatedStats::InstanceCount, TEXT("Count")).
+		AddColumn(&FNetProfilerAggregatedStats::TotalInclusive, TEXT("Incl")).
+		AddColumn(&FNetProfilerAggregatedStats::MaxInclusive, TEXT("I.Max")).
+		AddColumn(&FNetProfilerAggregatedStats::AverageInclusive, TEXT("I.Avg")).
+		AddColumn(&FNetProfilerAggregatedStats::TotalExclusive, TEXT("Excl")).
+		AddColumn(&FNetProfilerAggregatedStats::MaxExclusive, TEXT("E.Max"));
 }
 
 FNetProfilerProvider::~FNetProfilerProvider()
@@ -557,7 +566,7 @@ ITable<FNetProfilerAggregatedStats>* FNetProfilerProvider::CreateAggregation(uin
 	}
 
 	// Calculate averages and populate table
-	TTable<FAggregatedStatsTableLayout>* Table = new TTable<FAggregatedStatsTableLayout>();
+	TTable<FNetProfilerAggregatedStats>* Table = new TTable<FNetProfilerAggregatedStats>(AggregatedStatsTableLayout);
 	for (const auto& KV : AggregatedStats)
 	{
 		FNetProfilerAggregatedStats& Row = Table->AddRow();
diff --git a/Engine/Source/Developer/TraceServices/Private/Model/NetProfilerProvider.h b/Engine/Source/Developer/TraceServices/Private/Model/NetProfilerProvider.h
--- a/Engine/Source/Developer/TraceServices/Private/Model/NetProfilerProvider.h
+++ b/Engine/Source/Developer/TraceServices/Private/Model/NetProfilerProvider.h
@@ -115,17 +115,6 @@ public:
 	virtual ITable<FNetProfilerAggregatedStats>* CreateAggregation(uint32 ConnectionIndex, ENetProfilerConnectionMode Mode, uint32 PacketIndexIntervalStart, uint32 PacketIndexIntervalEnd, uint32 StartPosition, uint32 EndPosition) const override;
 
 private:
-
-	UE_TRACE_TABLE_LAYOUT_BEGIN(FAggregatedStatsTableLayout, FNetProfilerAggregatedStats)
-		UE_TRACE_TABLE_COLUMN(EventTypeIndex, TEXT("EventTypeIndex"))
-		UE_TRACE_TABLE_COLUMN(InstanceCount, TEXT("Count"))
-		UE_TRACE_TABLE_COLUMN(TotalInclusive, TEXT("Incl"))
-		UE_TRACE_TABLE_COLUMN(MaxInclusive, TEXT("I.Max"))
-		UE_TRACE_TABLE_COLUMN(AverageInclusive, TEXT("I.Avg"))
-		UE_TRACE_TABLE_COLUMN(TotalExclusive, TEXT("Excl"))
-		UE_TRACE_TABLE_COLUMN(MaxExclusive, TEXT("E.Max"))
-	UE_TRACE_TABLE_LAYOUT_END()
-
 	const FNetProfilerName* GetNetProfilerName(uint32 ProfilerNameId) const;
 	const FNetProfilerEventType* GetNetProfilerEventType(uint32 ProfilerEventTypeId) const;
 
@@ -143,6 +132,7 @@ private:
 	// All connections we have seen throughout the session
 	TPagedArray<FNetProfilerConnectionInternal> Connections;
 	uint32 ConnectionChangeCount;
+	TTableLayout<FNetProfilerAggregatedStats> AggregatedStatsTableLayout;
 };
 
 }
diff --git a/Engine/Source/Developer/TraceServices/Private/Model/Tables.h b/Engine/Source/Developer/TraceServices/Private/Model/Tables.h
--- a/Engine/Source/Developer/TraceServices/Private/Model/Tables.h
+++ b/Engine/Source/Developer/TraceServices/Private/Model/Tables.h
@@ -8,60 +8,6 @@
 #include "Common/PagedArray.h"
 #include "Common/SlabAllocator.h"
 
-#define UE_TRACE_TABLE_LAYOUT_BEGIN(Name, InRowType) \
-	class Name \
-		: public ::Trace::TTableLayoutBase<Name, InRowType> \
-	{ \
-	public: \
-		typedef InRowType RowType; \
-		virtual uint64 GetColumnCount() const override { return decltype(LastColumn)::Index; } \
-		template<uint8 IndexValue> \
-		constexpr ::Trace::ETableColumnType GetColumnTypeInternal() const { return TableColumnType_Invalid; } \
-		template<uint8 IndexValue> \
-		constexpr const TCHAR* GetColumnNameInternal() const { return nullptr; } \
-		template<uint8 IndexValue> \
-		::Trace::FColumnValueContainer GetColumnValueInternal(const RowType& Row) const { return 0; } \
-		::Trace::TColumnDeclaration<0, RowType
-
-#if PLATFORM_WINDOWS
-
-#define UE_TRACE_TABLE_COLUMN(Name, DisplayName) \
-		> PREPROCESSOR_JOIN(__Column__, __LINE__); \
-		template<> \
-		constexpr ::Trace::ETableColumnType GetColumnTypeInternal<decltype(PREPROCESSOR_JOIN(__Column__, __LINE__))::Index>() const { return ::Trace::GetColumnTypeFromNativeType<decltype(GetMemberTypeHelper(&RowType::Name))>(); } \
-		template<> \
-		constexpr const TCHAR* GetColumnNameInternal<decltype(PREPROCESSOR_JOIN(__Column__, __LINE__))::Index>() const { return DisplayName; } \
-		template<> \
-		::Trace::FColumnValueContainer GetColumnValueInternal<decltype(PREPROCESSOR_JOIN(__Column__, __LINE__))::Index>(const RowType& Row) const { return Row.Name; } \
-		::Trace::TColumnDeclaration<decltype(PREPROCESSOR_JOIN(__Column__, __LINE__))::Index + 1, RowType
-		
-#define UE_TRACE_TABLE_PROJECTED_COLUMN(ColumnType, DisplayName, ProjectionFunc) \
-		> PREPROCESSOR_JOIN(__Column__, __LINE__); \
-		TFunction<FNativeTypeFromColumnType<::Trace::ColumnType>::Type(const RowType&)> PREPROCESSOR_JOIN(__ColumnProjector__, __LINE__) = ProjectionFunc; \
-		template<> \
-		constexpr ::Trace::ETableColumnType GetColumnTypeInternal<decltype(PREPROCESSOR_JOIN(__Column__, __LINE__))::Index>() const { return ::Trace::ColumnType; } \
-		template<> \
-		constexpr const TCHAR* GetColumnNameInternal<decltype(PREPROCESSOR_JOIN(__Column__, __LINE__))::Index>() const { return DisplayName; } \
-		template<> \
-		::Trace::FColumnValueContainer GetColumnValueInternal<decltype(PREPROCESSOR_JOIN(__Column__, __LINE__))::Index>(const RowType& Row) const { return PREPROCESSOR_JOIN(__ColumnProjector__, __LINE__)(Row); } \
-		::Trace::TColumnDeclaration<decltype(PREPROCESSOR_JOIN(__Column__, __LINE__))::Index + 1, RowType
-
-#else
-
-#define UE_TRACE_TABLE_COLUMN(Name, DisplayName) \
-		> PREPROCESSOR_JOIN(__Column__, __LINE__); \
-		::Trace::TColumnDeclaration<decltype(PREPROCESSOR_JOIN(__Column__, __LINE__))::Index + 1, RowType
-
-#define UE_TRACE_TABLE_PROJECTED_COLUMN(ColumnType, DisplayName, ProjectionFunc) \
-		> PREPROCESSOR_JOIN(__Column__, __LINE__); \
-		::Trace::TColumnDeclaration<decltype(PREPROCESSOR_JOIN(__Column__, __LINE__))::Index + 1, RowType
-
-#endif
-
-#define UE_TRACE_TABLE_LAYOUT_END() \
-		> LastColumn;\
-	};
-
 namespace Trace
 {
 
@@ -140,45 +86,6 @@ constexpr ETableColumnType GetColumnTypeFromNativeType<const TCHAR*>()
 	return TableColumnType_CString;
 }
 
-template<ETableColumnType ColumnType>
-struct FNativeTypeFromColumnType
-{
-	
-};
-
-template<>
-struct FNativeTypeFromColumnType<TableColumnType_Bool>
-{
-	using Type = bool;
-};
-
-template<>
-struct FNativeTypeFromColumnType<TableColumnType_Int>
-{
-	using Type = int64;
-};
-
-template<>
-struct FNativeTypeFromColumnType<TableColumnType_Float>
-{
-	using Type = float;
-};
-
-template<>
-struct FNativeTypeFromColumnType<TableColumnType_Double>
-{
-	using Type = double;
-};
-
-template<>
-struct FNativeTypeFromColumnType<TableColumnType_CString>
-{
-	using Type = const TCHAR*;
-};
-
-template<typename MemberType, typename ObjectType>
-MemberType GetMemberTypeHelper(MemberType ObjectType::*);
-
 struct FColumnValueContainer
 {
 	FColumnValueContainer(bool Value)
@@ -228,7 +135,7 @@ struct FColumnValueContainer
 
 	FColumnValueContainer(const TCHAR* Value)
 	{
-		CStringValue = Value;
+		StringValue = Value;
 	}
 
 	union
@@ -237,29 +144,94 @@ struct FColumnValueContainer
 		int64 IntValue;
 		float FloatValue;
 		double DoubleValue;
-		const TCHAR* CStringValue;
+		const TCHAR* StringValue;
 	};
 };
 
-template<uint8 IndexValue, typename RowType>
-struct TColumnDeclaration
+template<typename RowType>
+class TTableLayout
+	: public ITableLayout
 {
-	enum
+public:
+	template<typename ColumnNativeType>
+	TTableLayout<RowType>& AddColumn(ColumnNativeType RowType::* MemberVariableColumn, const TCHAR* ColumnName)
+	{
+		Columns.Add({
+			ColumnName,
+			GetColumnTypeFromNativeType<ColumnNativeType>(),
+			[MemberVariableColumn](const RowType& Row) -> FColumnValueContainer
+			{
+				return FColumnValueContainer(Row.*MemberVariableColumn);
+			}
+			});
+		return *this;
+	}
+
+	template<typename ColumnNativeType>
+	TTableLayout<RowType>& AddColumn(ColumnNativeType(RowType::* MemberFunctionColumn)() const, const TCHAR* ColumnName)
+	{
+		Columns.Add({
+			ColumnName,
+			GetColumnTypeFromNativeType<ColumnNativeType>(),
+			[MemberFunctionColumn](const RowType& Row) -> FColumnValueContainer
+			{
+				return FColumnValueContainer((Row.*MemberFunctionColumn)());
+			}
+			});
+		return *this;
+	}
+
+	template<typename ColumnNativeType>
+	TTableLayout<RowType>& AddColumn(ColumnNativeType(*FunctionColumn)(const RowType&), const TCHAR* ColumnName)
+	{
+		Columns.Add({
+			ColumnName,
+			GetColumnTypeFromNativeType<ColumnNativeType>(),
+			[FunctionColumn](const RowType& Row) -> FColumnValueContainer
+			{
+				return FColumnValueContainer(FunctionColumn(Row));
+			}
+			});
+		return *this;
+	}
+
+	uint64 GetColumnCount() const override
+	{
+		return Columns.Num();
+	}
+
+	const TCHAR* GetColumnName(uint64 ColumnIndex) const override
 	{
-		Index = IndexValue
+		return *Columns[ColumnIndex].Name;
+	}
+
+	ETableColumnType GetColumnType(uint64 ColumnIndex) const override
+	{
+		return Columns[ColumnIndex].Type;
+	}
+
+	FColumnValueContainer GetColumnValue(const RowType& Row, uint64 ColumnIndex) const
+	{
+		return Columns[ColumnIndex].Projector(Row);
+	}
+
+private:
+	struct FColumnDeclaration
+	{
+		FString Name;
+		ETableColumnType Type;
+		TFunction<FColumnValueContainer(const RowType&)> Projector;
 	};
 
-	static constexpr ETableColumnType GetType() { return ::Trace::TableColumnType_Invalid; }
-	static constexpr const TCHAR* GetName() { return nullptr; }
-	static const void* GetValue(const RowType& Row) { return nullptr; }
+	TArray<FColumnDeclaration> Columns;
 };
 
-template<typename LayoutType>
+template<typename RowType>
 class TTableReader
-	: public ITableReader<typename LayoutType::RowType>
+	: public ITableReader<RowType>
 {
 public:
-	TTableReader(const LayoutType& InLayout, const TPagedArray<typename LayoutType::RowType>& InRows)
+	TTableReader(const TTableLayout<RowType>& InLayout, const TPagedArray<RowType>& InRows)
 		: Layout(InLayout)
 		, Iterator(InRows.GetIteratorFromItem(0))
 	{
@@ -281,7 +253,7 @@ public:
 		CurrentRow = Iterator.SetPosition(RowIndex);
 	}
 
-	virtual const typename LayoutType::RowType* GetCurrentRow() const override
+	virtual const RowType* GetCurrentRow() const override
 	{
 		return CurrentRow;
 	}
@@ -292,7 +264,7 @@ public:
 		{
 			return false;
 		}
-		ETableColumnType ColumnType = Layout.GetColumnTypeConstExpr(ColumnIndex);
+		ETableColumnType ColumnType = Layout.GetColumnType(ColumnIndex);
 		switch (ColumnType)
 		{
 		case TableColumnType_Bool:
@@ -313,7 +285,7 @@ public:
 		{
 			return 0;
 		}
-		ETableColumnType ColumnType = Layout.GetColumnTypeConstExpr(ColumnIndex);
+		ETableColumnType ColumnType = Layout.GetColumnType(ColumnIndex);
 		switch (ColumnType)
 		{
 		case TableColumnType_Bool:
@@ -334,7 +306,7 @@ public:
 		{
 			return 0.0;
 		}
-		ETableColumnType ColumnType = Layout.GetColumnTypeConstExpr(ColumnIndex);
+		ETableColumnType ColumnType = Layout.GetColumnType(ColumnIndex);
 		switch (ColumnType)
 		{
 		case TableColumnType_Bool:
@@ -355,7 +327,7 @@ public:
 		{
 			return 0.0;
 		}
-		ETableColumnType ColumnType = Layout.GetColumnTypeConstExpr(ColumnIndex);
+		ETableColumnType ColumnType = Layout.GetColumnType(ColumnIndex);
 		switch (ColumnType)
 		{
 		case TableColumnType_Bool:
@@ -376,159 +348,84 @@ public:
 		{
 			return TEXT("");
 		}
-		ETableColumnType ColumnType = Layout.GetColumnTypeConstExpr(ColumnIndex);
+		ETableColumnType ColumnType = Layout.GetColumnType(ColumnIndex);
 		if (ColumnType == TableColumnType_CString)
 		{
-			const TCHAR* CStringValue = Layout.GetColumnValue(*CurrentRow, ColumnIndex).CStringValue;
+			const TCHAR* CStringValue = Layout.GetColumnValue(*CurrentRow, ColumnIndex).StringValue;
 			return CStringValue ? CStringValue : TEXT("");
 		}
 		return TEXT("");
 	}
 
 private:
-	const LayoutType& Layout;
-	typename TPagedArray<typename LayoutType::RowType>::TIterator Iterator;
-	const typename LayoutType::RowType* CurrentRow;
+	const TTableLayout<RowType>& Layout;
+	typename TPagedArray<RowType>::TIterator Iterator;
+	const RowType* CurrentRow;
 };
 
-template<typename LayoutType, typename RowType>
-class TTableLayoutBase
-	: public ITableLayout
+template<typename RowType>
+class TTableBase
+	: public ITable<RowType>
 {
 public:
-	constexpr ETableColumnType GetColumnTypeConstExpr(uint64 ColumnIndex) const
-	{
-		const LayoutType* This = static_cast<const LayoutType*>(this);
-		switch (ColumnIndex)
-		{
-		case 0: return This->template GetColumnTypeInternal<0>();
-		case 1: return This->template GetColumnTypeInternal<1>();
-		case 2: return This->template GetColumnTypeInternal<2>();
-		case 3: return This->template GetColumnTypeInternal<3>();
-		case 4: return This->template GetColumnTypeInternal<4>();
-		case 5: return This->template GetColumnTypeInternal<5>();
-		case 6: return This->template GetColumnTypeInternal<6>();
-		case 7: return This->template GetColumnTypeInternal<7>();
-		case 8: return This->template GetColumnTypeInternal<8>();
-		case 9: return This->template GetColumnTypeInternal<9>();
-		case 10: return This->template GetColumnTypeInternal<10>();
-		case 11: return This->template GetColumnTypeInternal<11>();
-		case 12: return This->template GetColumnTypeInternal<12>();
-		case 13: return This->template GetColumnTypeInternal<13>();
-		case 14: return This->template GetColumnTypeInternal<14>();
-		case 15: return This->template GetColumnTypeInternal<15>();
-		}
-		return ::Trace::TableColumnType_Invalid;
-	}
+	TTableBase() = default;
 
-	virtual ETableColumnType GetColumnType(uint64 ColumnIndex) const override
+	TTableBase(TTableLayout<RowType> InLayout)
+		: Layout(InLayout)
 	{
-		return GetColumnTypeConstExpr(ColumnIndex);
-	}
 
-	constexpr const TCHAR* GetColumnNameConstExpr(uint64 ColumnIndex) const
-	{
-		const LayoutType* This = static_cast<const LayoutType*>(this);
-		switch (ColumnIndex)
-		{
-		case 0: return This->template GetColumnNameInternal<0>();
-		case 1: return This->template GetColumnNameInternal<1>();
-		case 2: return This->template GetColumnNameInternal<2>();
-		case 3: return This->template GetColumnNameInternal<3>();
-		case 4: return This->template GetColumnNameInternal<4>();
-		case 5: return This->template GetColumnNameInternal<5>();
-		case 6: return This->template GetColumnNameInternal<6>();
-		case 7: return This->template GetColumnNameInternal<7>();
-		case 8: return This->template GetColumnNameInternal<8>();
-		case 9: return This->template GetColumnNameInternal<9>();
-		case 10: return This->template GetColumnNameInternal<10>();
-		case 11: return This->template GetColumnNameInternal<11>();
-		case 12: return This->template GetColumnNameInternal<12>();
-		case 13: return This->template GetColumnNameInternal<13>();
-		case 14: return This->template GetColumnNameInternal<14>();
-		case 15: return This->template GetColumnNameInternal<15>();
-		}
-		return nullptr;
 	}
 
-	virtual const TCHAR* GetColumnName(uint64 ColumnIndex) const override
-	{
-		return GetColumnNameConstExpr(ColumnIndex);
-	}
+	virtual ~TTableBase() = default;
 
-	FColumnValueContainer GetColumnValue(const RowType& Row, uint64 ColumnIndex) const
+	const ITableLayout& GetLayout() const override
 	{
-		const LayoutType* This = static_cast<const LayoutType*>(this);
-		switch (ColumnIndex)
-		{
-		case 0: return This->template GetColumnValueInternal<0>(Row);
-		case 1: return This->template GetColumnValueInternal<1>(Row);
-		case 2: return This->template GetColumnValueInternal<2>(Row);
-		case 3: return This->template GetColumnValueInternal<3>(Row);
-		case 4: return This->template GetColumnValueInternal<4>(Row);
-		case 5: return This->template GetColumnValueInternal<5>(Row);
-		case 6: return This->template GetColumnValueInternal<6>(Row);
-		case 7: return This->template GetColumnValueInternal<7>(Row);
-		case 8: return This->template GetColumnValueInternal<8>(Row);
-		case 9: return This->template GetColumnValueInternal<9>(Row);
-		case 10: return This->template GetColumnValueInternal<10>(Row);
-		case 11: return This->template GetColumnValueInternal<11>(Row);
-		case 12: return This->template GetColumnValueInternal<12>(Row);
-		case 13: return This->template GetColumnValueInternal<13>(Row);
-		case 14: return This->template GetColumnValueInternal<14>(Row);
-		case 15: return This->template GetColumnValueInternal<15>(Row);
-		}
-		return nullptr;
+		return Layout;
 	}
-};
 
-template<typename LayoutType>
-class TTableBase
-	: public ITable<typename LayoutType::RowType>
-{
-	virtual const ITableLayout& GetLayout() const override
+	uint64 GetRowCount() const override
 	{
-		return Layout;
+		return GetRows().Num();
 	}
 
-	virtual uint64 GetRowCount() const override
+	ITableReader<RowType>* CreateReader() const override
 	{
-		return GetRows().Num();
+		return new TTableReader<RowType>(Layout, GetRows());
 	}
 
-	virtual ITableReader<typename LayoutType::RowType>* CreateReader() const override
+	TTableLayout<RowType>& EditLayout()
 	{
-		return new TTableReader<LayoutType>(Layout, GetRows());
+		return Layout;
 	}
 
 private:
-	virtual const TPagedArray<typename LayoutType::RowType>& GetRows() const = 0;
+	virtual const TPagedArray<RowType>& GetRows() const = 0;
 
-	LayoutType Layout;
+	TTableLayout<RowType> Layout;
 };
 
-template<typename LayoutType>
+template<typename RowType>
 class TTableView
-	: public TTableBase<LayoutType>
+	: public TTableBase<RowType>
 {
 public:
-	TTableView(const TPagedArray<typename LayoutType::RowType>& InRows)
+	TTableView(const TPagedArray<RowType>& InRows)
 		: Rows(InRows)
 	{
 	}
 
 private:
-	virtual const TPagedArray<typename LayoutType::RowType>& GetRows() const override
+	virtual const TPagedArray<RowType>& GetRows() const override
 	{
 		return Rows;
 	}
 
-	const TPagedArray<typename LayoutType::RowType>& Rows;
+	const TPagedArray<RowType>& Rows;
 };
 
-template<typename LayoutType>
+template<typename RowType>
 class TTable
-	: public TTableBase<LayoutType>
+	: public TTableBase<RowType>
 {
 public:
 	TTable()
@@ -538,19 +435,27 @@ public:
 
 	}
 
-	typename LayoutType::RowType& AddRow()
+	TTable(TTableLayout<RowType> Layout)
+		: TTableBase<RowType>(Layout)
+		, Allocator(2 << 20)
+		, Rows(Allocator, 1024)
+	{
+
+	}
+
+	RowType& AddRow()
 	{
 		return Rows.PushBack();
 	}
 
 private:
-	virtual const TPagedArray<typename LayoutType::RowType>& GetRows() const override
+	virtual const TPagedArray<RowType>& GetRows() const override
 	{
 		return Rows;
 	}
 
 	FSlabAllocator Allocator;
-	TPagedArray<typename LayoutType::RowType> Rows;
+	TPagedArray<RowType> Rows;
 };
 
 }
\ No newline at end of file
diff --git a/Engine/Source/Developer/TraceServices/Private/Model/TimingProfiler.cpp b/Engine/Source/Developer/TraceServices/Private/Model/TimingProfiler.cpp
--- a/Engine/Source/Developer/TraceServices/Private/Model/TimingProfiler.cpp
+++ b/Engine/Source/Developer/TraceServices/Private/Model/TimingProfiler.cpp
@@ -52,6 +52,24 @@ FTimingProfilerProvider::FTimingProfilerProvider(IAnalysisSession& InSession)
 	: Session(InSession)
 {
 	Timelines.Add(MakeShared<TimelineInternal>(Session.GetLinearAllocator()));
+
+	AggregatedStatsTableLayout.
+		AddColumn<const TCHAR*>([](const FTimingProfilerAggregatedStats& Row)
+			{
+				return Row.Timer->Name;
+			},
+			TEXT("Name")).
+		AddColumn(&FTimingProfilerAggregatedStats::InstanceCount, TEXT("Count")).
+		AddColumn(&FTimingProfilerAggregatedStats::TotalInclusiveTime, TEXT("Incl")).
+		AddColumn(&FTimingProfilerAggregatedStats::MinInclusiveTime, TEXT("I.Min")).
+		AddColumn(&FTimingProfilerAggregatedStats::MaxInclusiveTime, TEXT("I.Max")).
+		AddColumn(&FTimingProfilerAggregatedStats::AverageInclusiveTime, TEXT("I.Avg")).
+		AddColumn(&FTimingProfilerAggregatedStats::MedianInclusiveTime, TEXT("I.Med")).
+		AddColumn(&FTimingProfilerAggregatedStats::TotalExclusiveTime, TEXT("Excl")).
+		AddColumn(&FTimingProfilerAggregatedStats::MinExclusiveTime, TEXT("E.Min")).
+		AddColumn(&FTimingProfilerAggregatedStats::MaxExclusiveTime, TEXT("E.Max")).
+		AddColumn(&FTimingProfilerAggregatedStats::AverageExclusiveTime, TEXT("E.Avg")).
+		AddColumn(&FTimingProfilerAggregatedStats::MedianExclusiveTime, TEXT("E.Med"));
 }
 
 FTimingProfilerProvider::~FTimingProfilerProvider()
@@ -205,12 +223,23 @@ ITable<FTimingProfilerAggregatedStats>* FTimingProfilerProvider::CreateAggregati
 
 	TMap<const FTimingProfilerTimer*, FAggregatedTimingStats> Aggregation;
 	FTimelineStatistics::CreateAggregation(IncludedTimelines, BucketMappingFunc, IntervalStart, IntervalEnd, Aggregation);
-	TTable<FAggregatedStatsTableLayout>* Table = new TTable<FAggregatedStatsTableLayout>();
+	TTable<FTimingProfilerAggregatedStats>* Table = new TTable<FTimingProfilerAggregatedStats>(AggregatedStatsTableLayout);
 	for (const auto& KV : Aggregation)
 	{
 		FTimingProfilerAggregatedStats& Row = Table->AddRow();
-		*static_cast<FAggregatedTimingStats*>(&Row) = KV.Value;
 		Row.Timer = KV.Key;
+		const FAggregatedTimingStats& Stats = KV.Value;
+		Row.InstanceCount = Stats.InstanceCount;
+		Row.TotalInclusiveTime = Stats.TotalInclusiveTime;
+		Row.MinInclusiveTime = Stats.MinInclusiveTime;
+		Row.MaxInclusiveTime = Stats.MaxInclusiveTime;
+		Row.AverageInclusiveTime = Stats.AverageInclusiveTime;
+		Row.MedianInclusiveTime = Stats.MedianInclusiveTime;
+		Row.TotalExclusiveTime = Stats.TotalExclusiveTime;
+		Row.MinExclusiveTime = Stats.MinExclusiveTime;
+		Row.MaxExclusiveTime = Stats.MaxExclusiveTime;
+		Row.AverageExclusiveTime = Stats.AverageExclusiveTime;
+		Row.MedianExclusiveTime = Stats.MedianExclusiveTime;
 	}
 	return Table;
 }
diff --git a/Engine/Source/Developer/TraceServices/Private/Model/TimingProfilerPrivate.h b/Engine/Source/Developer/TraceServices/Private/Model/TimingProfilerPrivate.h
--- a/Engine/Source/Developer/TraceServices/Private/Model/TimingProfilerPrivate.h
+++ b/Engine/Source/Developer/TraceServices/Private/Model/TimingProfilerPrivate.h
@@ -43,21 +43,7 @@ private:
 	TArray<TSharedRef<TimelineInternal>> Timelines;
 	TMap<uint32, uint32> CpuThreadTimelineIndexMap;
 	uint32 GpuTimelineIndex = 0;
-
-	UE_TRACE_TABLE_LAYOUT_BEGIN(FAggregatedStatsTableLayout, FTimingProfilerAggregatedStats)
-		UE_TRACE_TABLE_PROJECTED_COLUMN(TableColumnType_CString, TEXT("Name"), [](const FTimingProfilerAggregatedStats& Row) { return Row.Timer->Name; })
-		UE_TRACE_TABLE_COLUMN(InstanceCount, TEXT("Count"))
-		UE_TRACE_TABLE_COLUMN(TotalInclusiveTime, TEXT("Incl"))
-		UE_TRACE_TABLE_COLUMN(MinInclusiveTime, TEXT("I.Min"))
-		UE_TRACE_TABLE_COLUMN(MaxInclusiveTime, TEXT("I.Max"))
-		UE_TRACE_TABLE_COLUMN(AverageInclusiveTime, TEXT("I.Avg"))
-		UE_TRACE_TABLE_COLUMN(MedianInclusiveTime, TEXT("I.Med"))
-		UE_TRACE_TABLE_COLUMN(TotalExclusiveTime, TEXT("Excl"))
-		UE_TRACE_TABLE_COLUMN(MinExclusiveTime, TEXT("E.Min"))
-		UE_TRACE_TABLE_COLUMN(MaxExclusiveTime, TEXT("E.Max"))
-		UE_TRACE_TABLE_COLUMN(AverageExclusiveTime, TEXT("E.Avg"))
-		UE_TRACE_TABLE_COLUMN(MedianExclusiveTime, TEXT("E.Med"))
-	UE_TRACE_TABLE_LAYOUT_END()
+	TTableLayout<FTimingProfilerAggregatedStats> AggregatedStatsTableLayout;
 };
 
 }
\ No newline at end of file
diff --git a/Engine/Source/Developer/TraceServices/Public/TraceServices/Containers/Timelines.h b/Engine/Source/Developer/TraceServices/Public/TraceServices/Containers/Timelines.h
--- a/Engine/Source/Developer/TraceServices/Public/TraceServices/Containers/Timelines.h
+++ b/Engine/Source/Developer/TraceServices/Public/TraceServices/Containers/Timelines.h
@@ -22,19 +22,4 @@ public:
 	virtual void EnumerateEvents(double IntervalStart, double IntervalEnd, TFunctionRef<void(double, double, uint32, const EventType&)> Callback) const = 0;
 };
 
-struct FAggregatedTimingStats
-{
-	uint64 InstanceCount = 0;
-	double TotalInclusiveTime = 0.0;
-	double MinInclusiveTime = DBL_MAX;
-	double MaxInclusiveTime = -DBL_MAX;
-	double AverageInclusiveTime = 0.0;
-	double MedianInclusiveTime = 0.0;
-	double TotalExclusiveTime = 0.0;
-	double MinExclusiveTime = DBL_MAX;
-	double MaxExclusiveTime = -DBL_MAX;
-	double AverageExclusiveTime = 0.0;
-	double MedianExclusiveTime = 0.0;
-};
-
 }
\ No newline at end of file
diff --git a/Engine/Source/Developer/TraceServices/Public/TraceServices/Model/TimingProfiler.h b/Engine/Source/Developer/TraceServices/Public/TraceServices/Model/TimingProfiler.h
--- a/Engine/Source/Developer/TraceServices/Public/TraceServices/Model/TimingProfiler.h
+++ b/Engine/Source/Developer/TraceServices/Public/TraceServices/Model/TimingProfiler.h
@@ -24,9 +24,19 @@ struct FTimingProfilerEvent
 };
 
 struct FTimingProfilerAggregatedStats
-	: public FAggregatedTimingStats
 {
 	const FTimingProfilerTimer* Timer = nullptr;
+	uint64 InstanceCount = 0;
+	double TotalInclusiveTime = 0.0;
+	double MinInclusiveTime = DBL_MAX;
+	double MaxInclusiveTime = -DBL_MAX;
+	double AverageInclusiveTime = 0.0;
+	double MedianInclusiveTime = 0.0;
+	double TotalExclusiveTime = 0.0;
+	double MinExclusiveTime = DBL_MAX;
+	double MaxExclusiveTime = -DBL_MAX;
+	double AverageExclusiveTime = 0.0;
+	double MedianExclusiveTime = 0.0;
 };
 
 struct FTimingProfilerButterflyNode
