UnrealInsights:
Added missing provider row and modification counts to enable selective UI updates
Streamlined tracing of strings for log messages and bookmarks
Changed it so that formatted log message strings are always stored in the analysis session instead of being dynamically resolved
Enabled log trace

diff --git a/Engine/Source/Developer/TraceInsights/Private/Insights/ViewModels/LogMessage.cpp b/Engine/Source/Developer/TraceInsights/Private/Insights/ViewModels/LogMessage.cpp
--- a/Engine/Source/Developer/TraceInsights/Private/Insights/ViewModels/LogMessage.cpp
+++ b/Engine/Source/Developer/TraceInsights/Private/Insights/ViewModels/LogMessage.cpp
@@ -153,8 +153,7 @@ FLogMessageRecord& FLogMessageCache::Get(uint64 Index)
 				FScopeLock Lock(&CriticalSection);
 				FLogMessageRecord Entry(Message);
 				Map.Add(Index, Entry);
-			},
-				true); // resolve format string
+			});
 		});
 	}
 
@@ -186,8 +185,7 @@ TSharedPtr<FLogMessageRecord> FLogMessageCache::GetUncached(uint64 Index, bool b
 				[&EntryPtr](const Trace::FLogMessage& Message)
 			{
 				EntryPtr = MakeShareable(new FLogMessageRecord(Message));
-			},
-				bResolveFormatString);
+			});
 		});
 	}
 
diff --git a/Engine/Source/Developer/TraceInsights/Private/Insights/ViewModels/MarkersTimingTrack.cpp b/Engine/Source/Developer/TraceInsights/Private/Insights/ViewModels/MarkersTimingTrack.cpp
--- a/Engine/Source/Developer/TraceInsights/Private/Insights/ViewModels/MarkersTimingTrack.cpp
+++ b/Engine/Source/Developer/TraceInsights/Private/Insights/ViewModels/MarkersTimingTrack.cpp
@@ -139,8 +139,7 @@ void FMarkersTimingTrack::Update(const FTimingTrackViewport& InViewport)
 			LogProvider.EnumerateMessages(
 				CB.GetViewport().StartTime,
 				CB.GetViewport().EndTime,
-				[&CB](const Trace::FLogMessage& Message) { CB.AddLogMessage(Message); },
-				false);
+				[&CB](const Trace::FLogMessage& Message) { CB.AddLogMessage(Message); });
 
 			CB.EndLog();
 		});
@@ -342,18 +341,17 @@ void FTimeMarkerTrackBuilder::AddLogMessage(const Trace::FLogMessage& Message)
 		//TODO: Search API like: LogProviderPtr->SearchMessage(StartIndex, ESearchDirection::Backward, LambdaPredicate, bResolveFormatString);
 		LogProviderPtr->ReadMessage(
 			Message.Index - 1,
-			[this](const Trace::FLogMessage& Message) { AddLogMessage(Message); },
-			false); // do not resolve format string
+			[this](const Trace::FLogMessage& Message) { AddLogMessage(Message); });
 	}
 
-	if (!Track.bUseOnlyBookmarks || FCString::Strcmp(*Message.Category->Name, TEXT("LogBookmark")) == 0)
+	if (!Track.bUseOnlyBookmarks || FCString::Strcmp(Message.Category->Name, TEXT("LogBookmark")) == 0)
 	{
 		float X = Viewport.TimeToSlateUnitsRounded(Message.Time);
 		if (X < 0)
 		{
 			X = -1.0f;
 		}
-		AddTimeMarker(X, *Message.Category->Name, Message.Verbosity, Message.Index);
+		AddTimeMarker(X, Message.Category->Name, Message.Verbosity, Message.Index);
 	}
 }
 
@@ -460,8 +458,7 @@ void FTimeMarkerTrackBuilder::Flush(float AvailableTextW)
 						TextInfo.Message.AppendChars(Message.Message, LastWholeCharacterIndexMessage + 1);
 					}
 				}
-			},
-				true); // resolve format string
+			});
 		}
 	}
 }
diff --git a/Engine/Source/Developer/TraceInsights/Private/Insights/Widgets/SLogView.cpp b/Engine/Source/Developer/TraceInsights/Private/Insights/Widgets/SLogView.cpp
--- a/Engine/Source/Developer/TraceInsights/Private/Insights/Widgets/SLogView.cpp
+++ b/Engine/Source/Developer/TraceInsights/Private/Insights/Widgets/SLogView.cpp
@@ -732,7 +732,7 @@ void SLogView::Tick(const FGeometry& AllottedGeometry, const double InCurrentTim
 						CategoryStr = CategoryStr.RightChop(3);
 					if (Categories.Contains(FName(*CategoryStr)))
 					{
-						UE_LOG(TimingProfiler, Log, TEXT("[LogView] Duplicated Log Category: \"%s\""), *Category.Name);
+						UE_LOG(TimingProfiler, Log, TEXT("[LogView] Duplicated Log Category: \"%s\""), Category.Name);
 					}
 					Categories.Add(FName(*CategoryStr));
 				});
diff --git a/Engine/Source/Developer/TraceServices/Private/AnalysisService.cpp b/Engine/Source/Developer/TraceServices/Private/AnalysisService.cpp
--- a/Engine/Source/Developer/TraceServices/Private/AnalysisService.cpp
+++ b/Engine/Source/Developer/TraceServices/Private/AnalysisService.cpp
@@ -70,9 +70,10 @@ void FAnalysisSessionLock::EndEdit()
 FAnalysisSession::FAnalysisSession(const TCHAR* SessionName)
 	: DurationSeconds(0.0)
 	, Allocator(32 << 20)
+	, StringStore(Allocator)
 	, ClassInfos(Allocator, 4096)
-	, BookmarkProvider(MakeShared<FBookmarkProvider>(Lock))
-	, LogProvider(MakeShared<FLogProvider>(Allocator, Lock))
+	, BookmarkProvider(MakeShared<FBookmarkProvider>(Lock, StringStore))
+	, LogProvider(MakeShared<FLogProvider>(Allocator, Lock, StringStore))
 	, ThreadProvider(MakeShared<FThreadProvider>(Lock))
 	, FramesProvider(MakeShared<FFrameProvider>(Allocator, Lock))
 	, TimingProfilerProvider(MakeShared<FTimingProfilerProvider>(Allocator, Lock))
diff --git a/Engine/Source/Developer/TraceServices/Private/AnalysisServicePrivate.h b/Engine/Source/Developer/TraceServices/Private/AnalysisServicePrivate.h
--- a/Engine/Source/Developer/TraceServices/Private/AnalysisServicePrivate.h
+++ b/Engine/Source/Developer/TraceServices/Private/AnalysisServicePrivate.h
@@ -7,6 +7,7 @@
 #include "Containers/HashTable.h"
 #include "Misc/ScopeLock.h"
 #include "Model/TimingProfiler.h"
+#include "Common/StringStore.h"
 
 struct FRandomStream;
 
@@ -78,6 +79,8 @@ public:
 	TSharedRef<FCountersProvider> EditCountersProvider() { return CountersProvider; }
 	virtual void ReadCountersProvider(TFunctionRef<void(const ICountersProvider &)> Callback) const override;
 
+	const TCHAR* StoreString(const TCHAR* String) { return StringStore.Store(String); }
+
 	virtual void BeginRead() const override { Lock.BeginRead(); }
 	virtual void EndRead() const override { Lock.EndRead(); }
 
@@ -91,6 +94,7 @@ private:
 	bool IsComplete = false;
 	double DurationSeconds = 0.0;
 	FSlabAllocator Allocator;
+	FStringStore StringStore;
 	TPagedArray<FClassInfo> ClassInfos;
 	TSharedRef<FBookmarkProvider> BookmarkProvider;
 	TSharedRef<FLogProvider> LogProvider;
diff --git a/Engine/Source/Developer/TraceServices/Private/Analyzers/LogTraceAnalysis.cpp b/Engine/Source/Developer/TraceServices/Private/Analyzers/LogTraceAnalysis.cpp
--- a/Engine/Source/Developer/TraceServices/Private/Analyzers/LogTraceAnalysis.cpp
+++ b/Engine/Source/Developer/TraceServices/Private/Analyzers/LogTraceAnalysis.cpp
@@ -31,8 +31,7 @@ void FLogTraceAnalyzer::OnEvent(uint16 RouteId, const FOnEventContext& Context)
 	{
 		uint64 CategoryPointer = EventData.GetValue("CategoryPointer").As<uint64>();
 		Trace::FLogCategory& Category = LogProvider->GetCategory(CategoryPointer);
-		uint16 NameLength = EventData.GetValue("NameLength").As<uint16>();
-		Category.Name = FString(NameLength, reinterpret_cast<const TCHAR*>(EventData.GetAttachment()));
+		Category.Name = Session->StoreString(reinterpret_cast<const TCHAR*>(EventData.GetAttachment()));
 		Category.DefaultVerbosity = static_cast<ELogVerbosity::Type>(EventData.GetValue("DefaultVerbosity").As<uint8>());
 		break;
 	}
@@ -45,17 +44,16 @@ void FLogTraceAnalyzer::OnEvent(uint16 RouteId, const FOnEventContext& Context)
 		Spec.Category = &Category;
 		Spec.Line = EventData.GetValue("Line").As<int32>();
 		Spec.Verbosity = static_cast<ELogVerbosity::Type>(EventData.GetValue("Verbosity").As<uint8>());
-		uint16 FileNameLength = EventData.GetValue("FileNameLength").As<uint16>();
-		Spec.File = FString(FileNameLength, reinterpret_cast<const ANSICHAR*>(EventData.GetAttachment()));
-		uint16 FormatStringLength = EventData.GetValue("FormatStringLength").As<uint16>();
-		Spec.FormatString = FString(FormatStringLength, reinterpret_cast<const TCHAR*>(EventData.GetAttachment() + FileNameLength));
+		const ANSICHAR* File = reinterpret_cast<const ANSICHAR*>(EventData.GetAttachment());
+		Spec.File = Session->StoreString(ANSI_TO_TCHAR(File));
+		Spec.FormatString = Session->StoreString(reinterpret_cast<const TCHAR*>(EventData.GetAttachment() + strlen(File) + 1));
 		break;
 	}
 	case RouteId_LogMessage:
 	{
 		uint64 LogPoint = EventData.GetValue("LogPoint").As<uint64>();
 		uint64 Cycle = EventData.GetValue("Cycle").As<uint64>();
-		LogProvider->AppendMessage(LogPoint, Context.SessionContext.TimestampFromCycle(Cycle), EventData.GetAttachmentSize(), EventData.GetAttachment());
+		LogProvider->AppendMessage(LogPoint, Context.SessionContext.TimestampFromCycle(Cycle), EventData.GetAttachment());
 		break;
 	}
 	}
diff --git a/Engine/Source/Developer/TraceServices/Private/Analyzers/MiscTraceAnalysis.cpp b/Engine/Source/Developer/TraceServices/Private/Analyzers/MiscTraceAnalysis.cpp
--- a/Engine/Source/Developer/TraceServices/Private/Analyzers/MiscTraceAnalysis.cpp
+++ b/Engine/Source/Developer/TraceServices/Private/Analyzers/MiscTraceAnalysis.cpp
@@ -134,10 +134,9 @@ void FMiscTraceAnalyzer::OnEvent(uint16 RouteId, const FOnEventContext& Context)
 		uint64 BookmarkPoint = EventData.GetValue("BookmarkPoint").As<uint64>();
 		Trace::FBookmarkSpec& Spec = BookmarkProvider->GetSpec(BookmarkPoint);
 		Spec.Line = EventData.GetValue("Line").As<int32>();
-		uint16 FileNameLength = EventData.GetValue("FileNameLength").As<uint16>();
-		Spec.File = FString(FileNameLength, reinterpret_cast<const ANSICHAR*>(EventData.GetAttachment()));
-		uint16 FormatStringLength = EventData.GetValue("FormatStringLength").As<uint16>();
-		Spec.FormatString = FString(FormatStringLength, reinterpret_cast<const TCHAR*>(EventData.GetAttachment() + FileNameLength));
+		const ANSICHAR* File = reinterpret_cast<const ANSICHAR*>(EventData.GetAttachment());
+		Spec.File = Session->StoreString(ANSI_TO_TCHAR(File));
+		Spec.FormatString = Session->StoreString(reinterpret_cast<const TCHAR*>(EventData.GetAttachment() + strlen(File) + 1));
 
 		Trace::FLogMessageSpec& LogMessageSpec = LogProvider->GetMessageSpec(BookmarkPoint);
 		LogMessageSpec.Category = &LogProvider->GetCategory(Trace::FLogProvider::ReservedLogCategory_Bookmark);
@@ -151,11 +150,9 @@ void FMiscTraceAnalyzer::OnEvent(uint16 RouteId, const FOnEventContext& Context)
 	{
 		uint64 BookmarkPoint = EventData.GetValue("BookmarkPoint").As<uint64>();
 		uint64 Cycle = EventData.GetValue("Cycle").As<uint64>();
-		uint16 FormatArgsSize = EventData.GetAttachmentSize();
 		double Timestamp = Context.SessionContext.TimestampFromCycle(Cycle);
-		BookmarkProvider->AppendBookmark(Timestamp, BookmarkPoint, FormatArgsSize, EventData.GetAttachment());
-
-		LogProvider->AppendMessage(BookmarkPoint, Timestamp, FormatArgsSize, EventData.GetAttachment());
+		BookmarkProvider->AppendBookmark(Timestamp, BookmarkPoint, EventData.GetAttachment());
+		LogProvider->AppendMessage(BookmarkPoint, Timestamp, EventData.GetAttachment());
 		break;
 	}
 	case RouteId_BeginFrame:
diff --git a/Engine/Source/Developer/TraceServices/Private/Common/StringStore.cpp b/Engine/Source/Developer/TraceServices/Private/Common/StringStore.cpp
new file mode 100644
--- /dev/null
+++ b/Engine/Source/Developer/TraceServices/Private/Common/StringStore.cpp
@@ -0,0 +1,41 @@
+// Copyright 1998-2019 Epic Games, Inc. All Rights Reserved.
+
+#include "Common/StringStore.h"
+
+namespace Trace
+{
+
+FStringStore::FStringStore(FSlabAllocator& InAllocator)
+	: Allocator(InAllocator)
+{
+
+}
+
+const TCHAR* FStringStore::Store(const TCHAR* String)
+{
+	uint32 Hash = GetTypeHash(String);
+	const TCHAR** AlreadyStored = StoredStrings.Find(Hash);
+	if (AlreadyStored && !FCString::Strcmp(String, *AlreadyStored))
+	{
+		return *AlreadyStored;
+	}
+	
+	int32 StringLength = FCString::Strlen(String) + 1;
+	if (BufferLeft < StringLength)
+	{
+		BufferPtr = Allocator.Allocate<TCHAR>(BlockSize);
+		++BlockCount;
+		BufferLeft = BlockSize;
+	}
+	const TCHAR* Stored = BufferPtr;
+	memcpy(BufferPtr, String, StringLength * sizeof(TCHAR));
+	BufferLeft -= StringLength;
+	BufferPtr += StringLength;
+	if (!AlreadyStored)
+	{
+		StoredStrings.Add(Hash, Stored);
+	}
+	return Stored;
+}
+
+}
diff --git a/Engine/Source/Developer/TraceServices/Private/Common/StringStore.h b/Engine/Source/Developer/TraceServices/Private/Common/StringStore.h
new file mode 100644
--- /dev/null
+++ b/Engine/Source/Developer/TraceServices/Private/Common/StringStore.h
@@ -0,0 +1,32 @@
+// Copyright 1998-2019 Epic Games, Inc. All Rights Reserved.
+
+#pragma once
+
+#include "CoreTypes.h"
+#include "Common/SlabAllocator.h"
+#include "Common/PagedArray.h"
+#include "Containers/Map.h"
+
+namespace Trace
+{
+
+class FStringStore
+{
+public:
+	FStringStore(FSlabAllocator& Allocator);
+	const TCHAR* Store(const TCHAR* String);
+
+private:
+	enum
+	{
+		BlockSize = 4 << 20
+	};
+	FSlabAllocator& Allocator;
+	TMap<uint32, const TCHAR*> StoredStrings;
+	TCHAR* BufferPtr = nullptr;
+	uint64 BufferLeft = 0;
+	uint64 BlockCount = 0;
+
+};
+
+}
\ No newline at end of file
diff --git a/Engine/Source/Developer/TraceServices/Private/Model/Bookmarks.cpp b/Engine/Source/Developer/TraceServices/Private/Model/Bookmarks.cpp
--- a/Engine/Source/Developer/TraceServices/Private/Model/Bookmarks.cpp
+++ b/Engine/Source/Developer/TraceServices/Private/Model/Bookmarks.cpp
@@ -7,8 +7,9 @@
 namespace Trace
 {
 
-FBookmarkProvider::FBookmarkProvider(const FAnalysisSessionLock& InSessionLock)
+FBookmarkProvider::FBookmarkProvider(const FAnalysisSessionLock& InSessionLock, FStringStore& InStringStore)
 	: SessionLock(InSessionLock)
+	, StringStore(InStringStore)
 {
 
 }
@@ -28,15 +29,14 @@ FBookmarkSpec& FBookmarkProvider::GetSpec(uint64 BookmarkPoint)
 	}
 }
 
-void FBookmarkProvider::AppendBookmark(double Time, uint64 BookmarkPoint, uint16 FormatArgsSize, const uint8* FormatArgs)
+void FBookmarkProvider::AppendBookmark(double Time, uint64 BookmarkPoint, const uint8* FormatArgs)
 {
 	SessionLock.WriteAccessCheck();
 	FBookmarkSpec& Spec = GetSpec(BookmarkPoint);
 	TSharedRef<FBookmarkInternal> Bookmark = MakeShared<FBookmarkInternal>();
 	Bookmark->Time = Time;
-	TCHAR Buffer[4096];
-	FFormatArgsHelper::Format(Buffer, 4096, *Spec.FormatString, FormatArgs);
-	Bookmark->Text = Buffer;
+	FFormatArgsHelper::Format(FormatBuffer, FormatBufferSize - 1, Spec.FormatString, FormatArgs);
+	Bookmark->Text = StringStore.Store(FormatBuffer);
 	Bookmarks.Add(Bookmark);
 }
 
@@ -70,7 +70,7 @@ void FBookmarkProvider::EnumerateBookmarks(double IntervalStart, double Interval
 		const FBookmarkInternal& InternalBookmark = Bookmarks[Index].Get();
 		FBookmark Bookmark;
 		Bookmark.Time = InternalBookmark.Time;
-		Bookmark.Text = *InternalBookmark.Text;
+		Bookmark.Text = InternalBookmark.Text;
 		Callback(Bookmark);
 	}
 }
diff --git a/Engine/Source/Developer/TraceServices/Private/Model/Bookmarks.h b/Engine/Source/Developer/TraceServices/Private/Model/Bookmarks.h
--- a/Engine/Source/Developer/TraceServices/Private/Model/Bookmarks.h
+++ b/Engine/Source/Developer/TraceServices/Private/Model/Bookmarks.h
@@ -9,34 +9,43 @@ namespace Trace
 {
 
 class FAnalysisSessionLock;
+class FStringStore;
 
 struct FBookmarkSpec
 {
-	FString File;
-	FString FormatString;
+	const TCHAR* File;
+	const TCHAR* FormatString;
 	int32 Line;
 };
 
 struct FBookmarkInternal
 {
 	double Time;
-	FString Text;
+	const TCHAR* Text;
 };
 
 class FBookmarkProvider
 	: public IBookmarkProvider
 {
 public:
-	FBookmarkProvider(const FAnalysisSessionLock& SessionLock);
+	FBookmarkProvider(const FAnalysisSessionLock& SessionLock, FStringStore& StringStore);
 
 	FBookmarkSpec& GetSpec(uint64 BookmarkPoint);
-	void AppendBookmark(double Time, uint64 BookmarkPoint, uint16 FormatArgsSize, const uint8* FormatArgs);
+	virtual uint64 GetBookmarkCount() const override { return Bookmarks.Num(); }
+	void AppendBookmark(double Time, uint64 BookmarkPoint, const uint8* FormatArgs);
 	virtual void EnumerateBookmarks(double IntervalStart, double IntervalEnd, TFunctionRef<void(const FBookmark&)> Callback) const override;
 
 private:
+	enum
+	{
+		FormatBufferSize = 65536
+	};
+
 	const FAnalysisSessionLock& SessionLock;
+	FStringStore& StringStore;
 	TMap<uint64, TSharedPtr<FBookmarkSpec>> SpecMap;
 	TArray<TSharedRef<FBookmarkInternal>> Bookmarks;
+	TCHAR FormatBuffer[FormatBufferSize];
 };
 
 }
\ No newline at end of file
diff --git a/Engine/Source/Developer/TraceServices/Private/Model/Counters.h b/Engine/Source/Developer/TraceServices/Private/Model/Counters.h
--- a/Engine/Source/Developer/TraceServices/Private/Model/Counters.h
+++ b/Engine/Source/Developer/TraceServices/Private/Model/Counters.h
@@ -54,6 +54,7 @@ class FCountersProvider
 public:
 	FCountersProvider(FSlabAllocator& Allocator, FAnalysisSessionLock& SessionLock);
 	virtual ~FCountersProvider();
+	virtual uint64 GetCounterCount() const override { return Counters.Num(); }
 	virtual void EnumerateCounters(TFunctionRef<void(const ICounter&)> Callback) const override;
 
 	FCounterInternal* CreateCounter(const TCHAR* Name, const TCHAR* Description, ECounterDisplayHint DisplayHint);
diff --git a/Engine/Source/Developer/TraceServices/Private/Model/IntervalTimeline.h b/Engine/Source/Developer/TraceServices/Private/Model/IntervalTimeline.h
--- a/Engine/Source/Developer/TraceServices/Private/Model/IntervalTimeline.h
+++ b/Engine/Source/Developer/TraceServices/Private/Model/IntervalTimeline.h
@@ -28,6 +28,7 @@ public:
 	}
 
 	virtual ~TIntervalTimeline() = default;
+	virtual uint64 GetModCount() const override { return ModCount; }
 	virtual uint64 GetEventCount() const override { return Events.Num(); }
 	
 	virtual void EnumerateEventsDownSampled(double IntervalStart, double IntervalEnd, double Resolution, TFunctionRef<void(bool, double, const EventType&)> Callback) const override { check(false); }
@@ -87,6 +88,8 @@ public:
 
 		LastPage->EndEventIndex = Index + 1;
 
+		++ModCount;
+
 		return Index;
 	}
 
@@ -99,6 +102,9 @@ public:
 		EventInternal.EndTime = EndTime;
 		FEventPage* Page = Events.GetItemPage(EventIndex);
 		Page->EndTime = FMath::Max(Page->EndTime, EndTime);
+
+		++ModCount;
+
 		return EventInternal.Event;
 	}
 
@@ -121,6 +127,7 @@ private:
 	TPagedArray<FEventInternal, FEventPage> Events;
 	FEventPage* CurrentPage = nullptr;
 	double LastTime = 0.0;
+	uint64 ModCount = 0;
 };
 
 }
\ No newline at end of file
diff --git a/Engine/Source/Developer/TraceServices/Private/Model/LoadTimeProfiler.h b/Engine/Source/Developer/TraceServices/Private/Model/LoadTimeProfiler.h
--- a/Engine/Source/Developer/TraceServices/Private/Model/LoadTimeProfiler.h
+++ b/Engine/Source/Developer/TraceServices/Private/Model/LoadTimeProfiler.h
@@ -19,6 +19,7 @@ public:
 	typedef TMonotonicTimeline<FLoadTimeProfilerCpuEvent> CpuTimelineInternal;
 
 	FLoadTimeProfilerProvider(FSlabAllocator& Allocator, FAnalysisSessionLock& InSessionLock);
+	virtual uint64 GetPackageCount() const override { return Packages.Num(); }
 	virtual void EnumeratePackages(TFunctionRef<void(const FPackageInfo&)> Callback) const override;
 	virtual void ReadMainThreadCpuTimeline(TFunctionRef<void(const CpuTimeline&)> Callback) const override;
 	virtual void ReadAsyncLoadingThreadCpuTimeline(TFunctionRef<void(const CpuTimeline&)> Callback) const override;
diff --git a/Engine/Source/Developer/TraceServices/Private/Model/Log.cpp b/Engine/Source/Developer/TraceServices/Private/Model/Log.cpp
--- a/Engine/Source/Developer/TraceServices/Private/Model/Log.cpp
+++ b/Engine/Source/Developer/TraceServices/Private/Model/Log.cpp
@@ -7,9 +7,10 @@
 namespace Trace
 {
 
-FLogProvider::FLogProvider(FSlabAllocator& InAllocator, FAnalysisSessionLock& InSessionLock)
+FLogProvider::FLogProvider(FSlabAllocator& InAllocator, FAnalysisSessionLock& InSessionLock, FStringStore& InStringStore)
 	: Allocator(InAllocator)
 	, SessionLock(InSessionLock)
+	, StringStore(InStringStore)
 	, Categories(Allocator, 128)
 	, MessageSpecs(Allocator, 1024)
 {
@@ -46,16 +47,15 @@ FLogMessageSpec& FLogProvider::GetMessageSpec(uint64 LogPoint)
 	}
 }
 
-void FLogProvider::AppendMessage(uint64 LogPoint, double Time, uint16 FormatArgsSize, const uint8* FormatArgs)
+void FLogProvider::AppendMessage(uint64 LogPoint, double Time, const uint8* FormatArgs)
 {
 	SessionLock.WriteAccessCheck();
 	check(SpecMap.Contains(LogPoint));
 	FLogMessageInternal& InternalMessage = Messages.AddDefaulted_GetRef();
 	InternalMessage.Time = Time;
 	InternalMessage.Spec = SpecMap[LogPoint];
-	InternalMessage.FormatArgsOffset = FormatArgsMemory.Num();
-	FormatArgsMemory.AddUninitialized(FormatArgsSize);
-	memcpy(FormatArgsMemory.GetData() + InternalMessage.FormatArgsOffset, FormatArgs, FormatArgsSize);
+	FFormatArgsHelper::Format(FormatBuffer, FormatBufferSize - 1, InternalMessage.Spec->FormatString, FormatArgs);
+	InternalMessage.Message = StringStore.Store(FormatBuffer);
 }
 
 uint64 FLogProvider::GetMessageCount() const
@@ -64,49 +64,36 @@ uint64 FLogProvider::GetMessageCount() const
 	return Messages.Num();
 }
 
-bool FLogProvider::ReadMessage(uint64 Index, TFunctionRef<void(const FLogMessage &)> Callback, bool ResolveFormatString) const
+bool FLogProvider::ReadMessage(uint64 Index, TFunctionRef<void(const FLogMessage &)> Callback) const
 {
 	SessionLock.ReadAccessCheck();
 	if (Index >= Messages.Num())
 	{
 		return false;
 	}
-	ConstructMessage(Index, Callback, ResolveFormatString);
+	ConstructMessage(Index, Callback);
 	return true;
 }
 
-void FLogProvider::EnumerateMessages(double IntervalStart, double IntervalEnd, TFunctionRef<void(const FLogMessage&)> Callback, bool ResolveFormatString) const
+void FLogProvider::EnumerateMessages(double IntervalStart, double IntervalEnd, TFunctionRef<void(const FLogMessage&)> Callback) const
 {
 	SessionLock.ReadAccessCheck();
 	if (IntervalStart > IntervalEnd)
 	{
 		return;
 	}
-	uint64 FirstMessageIndex = Algo::LowerBoundBy(Messages, IntervalStart, [](const FLogMessageInternal& M)
-	{
-		return M.Time;
-	});
 	uint64 MessageCount = Messages.Num();
-	if (FirstMessageIndex >= MessageCount)
+	for (uint64 Index = 0; Index < MessageCount; ++Index)
 	{
-		return;
-	}
-	uint64 LastMessageIndex = Algo::UpperBoundBy(Messages, IntervalEnd, [](const FLogMessageInternal& M)
-	{
-		return M.Time;
-	});
-	if (LastMessageIndex == 0)
-	{
-		return;
-	}
-	--LastMessageIndex;
-	for (uint64 Index = FirstMessageIndex; Index <= LastMessageIndex; ++Index)
-	{
-		ConstructMessage(Index, Callback, ResolveFormatString);
+		double Time = Messages[Index].Time;
+		if (IntervalStart <= Time && Time <= IntervalEnd)
+		{
+			ConstructMessage(Index, Callback);
+		}
 	}
 }
 
-void FLogProvider::EnumerateMessagesByIndex(uint64 Start, uint64 End, TFunctionRef<void(const FLogMessage &)> Callback, bool ResolveFormatString) const
+void FLogProvider::EnumerateMessagesByIndex(uint64 Start, uint64 End, TFunctionRef<void(const FLogMessage &)> Callback) const
 {
 	SessionLock.ReadAccessCheck();
 	uint64 Count = Messages.Num();
@@ -124,30 +111,21 @@ void FLogProvider::EnumerateMessagesByIndex(uint64 Start, uint64 End, TFunctionR
 	}
 	for (uint64 Index = Start; Index < End; ++Index)
 	{
-		ConstructMessage(Index, Callback, ResolveFormatString);
+		ConstructMessage(Index, Callback);
 	}
 }
 
-void FLogProvider::ConstructMessage(uint64 Index, TFunctionRef<void(const FLogMessage &)> Callback, bool ResolveFormatString) const
+void FLogProvider::ConstructMessage(uint64 Index, TFunctionRef<void(const FLogMessage &)> Callback) const
 {
 	const FLogMessageInternal& InternalMessage = Messages[Index];
 	FLogMessage Message;
 	Message.Index = Index;
 	Message.Time = InternalMessage.Time;
 	Message.Category = InternalMessage.Spec->Category;
-	Message.File = *InternalMessage.Spec->File;
+	Message.File = InternalMessage.Spec->File;
 	Message.Line = InternalMessage.Spec->Line;
 	Message.Verbosity = InternalMessage.Spec->Verbosity;
-	if (ResolveFormatString)
-	{
-		TCHAR Buffer[65535];
-		FFormatArgsHelper::Format(Buffer, 65534, *InternalMessage.Spec->FormatString, FormatArgsMemory.GetData() + InternalMessage.FormatArgsOffset);
-		Message.Message = Buffer;
-	}
-	else
-	{
-		Message.Message = *InternalMessage.Spec->FormatString;
-	}
+	Message.Message = InternalMessage.Message;
 	Callback(Message);
 }
 
diff --git a/Engine/Source/Developer/TraceServices/Private/Model/Log.h b/Engine/Source/Developer/TraceServices/Private/Model/Log.h
--- a/Engine/Source/Developer/TraceServices/Private/Model/Log.h
+++ b/Engine/Source/Developer/TraceServices/Private/Model/Log.h
@@ -10,21 +10,22 @@ namespace Trace
 {
 
 class FAnalysisSessionLock;
+class FStringStore;
 
 struct FLogMessageSpec
 {
-	FLogCategory* Category;
-	FString File;
-	FString FormatString;
+	FLogCategory* Category = nullptr;
+	const TCHAR* File = nullptr;
+	const TCHAR* FormatString = nullptr;
 	int32 Line;
 	ELogVerbosity::Type Verbosity;
 };
 
 struct FLogMessageInternal
 {
-	FLogMessageSpec* Spec;
+	FLogMessageSpec* Spec = nullptr;
 	double Time;
-	uint64 FormatArgsOffset;
+	const TCHAR* Message = nullptr;
 };
 
 class FLogProvider :
@@ -36,29 +37,36 @@ public:
 		ReservedLogCategory_Bookmark = 0
 	};
 	
-	FLogProvider(FSlabAllocator& Allocator, FAnalysisSessionLock& SessionLock);
+	FLogProvider(FSlabAllocator& Allocator, FAnalysisSessionLock& SessionLock, FStringStore& StringStore);
 
 	FLogCategory& GetCategory(uint64 CategoryPointer);
 	FLogMessageSpec& GetMessageSpec(uint64 LogPoint);
-	void AppendMessage(uint64 LogPoint, double Time, uint16 FormatArgsSize, const uint8* FormatArgs);
+	void AppendMessage(uint64 LogPoint, double Time, const uint8* FormatArgs);
 
 	virtual uint64 GetMessageCount() const override;
-	virtual bool ReadMessage(uint64 Index, TFunctionRef<void(const FLogMessage &)> Callback, bool ResolveFormatString) const override;
-	virtual void EnumerateMessages(double IntervalStart, double IntervalEnd, TFunctionRef<void(const FLogMessage&)> Callback, bool ResolveFormatString) const override;
-	virtual void EnumerateMessagesByIndex(uint64 Start, uint64 End, TFunctionRef<void(const FLogMessage&)> Callback, bool ResolveFormatString) const override;
+	virtual bool ReadMessage(uint64 Index, TFunctionRef<void(const FLogMessage &)> Callback) const override;
+	virtual void EnumerateMessages(double IntervalStart, double IntervalEnd, TFunctionRef<void(const FLogMessage&)> Callback) const override;
+	virtual void EnumerateMessagesByIndex(uint64 Start, uint64 End, TFunctionRef<void(const FLogMessage&)> Callback) const override;
+	virtual uint64 GetCategoryCount() const override { return Categories.Num(); }
 	virtual void EnumerateCategories(TFunctionRef<void(const FLogCategory&)> Callback) const override;
 
 private:
-	void ConstructMessage(uint64 Id, TFunctionRef<void(const FLogMessage &)> Callback, bool ResolveFormatString) const;
+	void ConstructMessage(uint64 Id, TFunctionRef<void(const FLogMessage &)> Callback) const;
 	
+	enum
+	{
+		FormatBufferSize = 65536
+	};
+
 	FSlabAllocator& Allocator;
 	FAnalysisSessionLock& SessionLock;
+	FStringStore& StringStore;
 	TMap<uint64, FLogCategory*> CategoryMap;
 	TMap<uint64, FLogMessageSpec*> SpecMap;
 	TPagedArray<FLogCategory> Categories;
 	TPagedArray<FLogMessageSpec> MessageSpecs;
 	TArray<FLogMessageInternal> Messages;
-	TArray<uint8> FormatArgsMemory;
+	TCHAR FormatBuffer[FormatBufferSize];
 };
 
 }
\ No newline at end of file
diff --git a/Engine/Source/Developer/TraceServices/Private/Model/MonotonicTimeline.h b/Engine/Source/Developer/TraceServices/Private/Model/MonotonicTimeline.h
--- a/Engine/Source/Developer/TraceServices/Private/Model/MonotonicTimeline.h
+++ b/Engine/Source/Developer/TraceServices/Private/Model/MonotonicTimeline.h
@@ -93,6 +93,7 @@ public:
 			}
 			SetEvent(DetailLevel, CurrentDepthState.PendingEventIndex, Event);
 		}
+		++ModCount;
 	}
 
 	void AppendEndEvent(double EndTime)
@@ -108,6 +109,12 @@ public:
 
 			UpdateDominatingEvent(DetailLevel, CurrentDepth, EndTime);
 		}
+		++ModCount;
+	}
+
+	virtual uint64 GetModCount() const override
+	{
+		return ModCount;
 	}
 
 	virtual uint64 GetEventCount() const override
@@ -455,6 +462,7 @@ private:
 
 	FSlabAllocator& Allocator;
 	TArray<FDetailLevel> DetailLevels;
+	uint64 ModCount = 0;
 };
 
 }
\ No newline at end of file
diff --git a/Engine/Source/Developer/TraceServices/Private/Model/Threads.cpp b/Engine/Source/Developer/TraceServices/Private/Model/Threads.cpp
--- a/Engine/Source/Developer/TraceServices/Private/Model/Threads.cpp
+++ b/Engine/Source/Developer/TraceServices/Private/Model/Threads.cpp
@@ -25,6 +25,7 @@ void FThreadProvider::EnsureThreadExists(uint32 Id)
 		SortedThreads.Add(ThreadInfo);
 		ThreadMap.Add(Id, ThreadInfo);
 		SortThreads();
+		++ModCount;
 	}
 }
 
@@ -41,6 +42,7 @@ void FThreadProvider::AddGameThread(uint32 Id)
 	ThreadInfo->IsGameThread = true;
 	SortedThreads.Add(ThreadInfo);
 	ThreadMap.Add(Id, ThreadInfo);
+	++ModCount;
 }
 
 void FThreadProvider::AddThread(uint32 Id, const FString& Name, EThreadPriority Priority)
@@ -63,6 +65,7 @@ void FThreadProvider::AddThread(uint32 Id, const FString& Name, EThreadPriority
 	ThreadInfo->PrioritySortOrder = GetPrioritySortOrder(Priority);
 	ThreadInfo->Name = Name;
 	SortThreads();
+	++ModCount;
 }
 
 void FThreadProvider::SetThreadPriority(uint32 Id, EThreadPriority Priority)
@@ -73,6 +76,7 @@ void FThreadProvider::SetThreadPriority(uint32 Id, EThreadPriority Priority)
 	TSharedRef<FThreadInfoInternal> ThreadInfo = ThreadMap[Id];
 	ThreadInfo->PrioritySortOrder = GetPrioritySortOrder(Priority);
 	SortThreads();
+	++ModCount;
 }
 
 void FThreadProvider::SetThreadGroup(uint32 Id, ETraceThreadGroup Group)
@@ -110,6 +114,7 @@ void FThreadProvider::SetThreadGroup(uint32 Id, ETraceThreadGroup Group)
 		break;
 	}
 	SortThreads();
+	++ModCount;
 }
 
 void FThreadProvider::EnumerateThreads(TFunctionRef<void(const FThreadInfo &)> Callback) const
diff --git a/Engine/Source/Developer/TraceServices/Private/Model/Threads.h b/Engine/Source/Developer/TraceServices/Private/Model/Threads.h
--- a/Engine/Source/Developer/TraceServices/Private/Model/Threads.h
+++ b/Engine/Source/Developer/TraceServices/Private/Model/Threads.h
@@ -19,7 +19,7 @@ public:
 	void AddThread(uint32 Id, const FString& Name, EThreadPriority Priority);
 	void SetThreadPriority(uint32 Id, EThreadPriority Priority);
 	void SetThreadGroup(uint32 Id, ETraceThreadGroup Group);
-
+	virtual uint64 GetModCount() const override { return ModCount; }
 	virtual void EnumerateThreads(TFunctionRef<void(const FThreadInfo&)> Callback) const override;
 
 private:
@@ -41,6 +41,7 @@ private:
 	static uint32 GetGroupSortOrder(ETraceThreadGroup ThreadGroup);
 
 	const FAnalysisSessionLock& SessionLock;
+	uint64 ModCount = 0;
 	TMap<uint32, TSharedRef<FThreadInfoInternal>> ThreadMap;
 	TArray<TSharedRef<FThreadInfoInternal>> SortedThreads;
 };
diff --git a/Engine/Source/Developer/TraceServices/Private/Model/TimingProfiler.h b/Engine/Source/Developer/TraceServices/Private/Model/TimingProfiler.h
--- a/Engine/Source/Developer/TraceServices/Private/Model/TimingProfiler.h
+++ b/Engine/Source/Developer/TraceServices/Private/Model/TimingProfiler.h
@@ -26,6 +26,7 @@ public:
 	virtual bool GetCpuThreadTimelineIndex(uint32 ThreadId, uint32& OutTimelineIndex) const override;
 	virtual bool GetGpuTimelineIndex(uint32& OutTimelineIndex) const override;
 	virtual bool ReadTimeline(uint32 Index, TFunctionRef<void(const Timeline&)> Callback) const override;
+	virtual uint64 GetTimelineCount() const override { return Timelines.Num(); }
 	virtual void EnumerateTimelines(TFunctionRef<void(const Timeline&)> Callback) const override;
 	virtual void ReadTimers(TFunctionRef<void(const FTimingProfilerTimer*, uint64)> Callback) const override;
 
diff --git a/Engine/Source/Developer/TraceServices/Public/TraceServices/AnalysisService.h b/Engine/Source/Developer/TraceServices/Public/TraceServices/AnalysisService.h
--- a/Engine/Source/Developer/TraceServices/Public/TraceServices/AnalysisService.h
+++ b/Engine/Source/Developer/TraceServices/Public/TraceServices/AnalysisService.h
@@ -25,12 +25,13 @@ class IBookmarkProvider
 {
 public:
 	virtual ~IBookmarkProvider() = default;
+	virtual uint64 GetBookmarkCount() const = 0;
 	virtual void EnumerateBookmarks(double IntervalStart, double IntervalEnd, TFunctionRef<void(const FBookmark&)> Callback) const = 0;
 };
 
 struct FLogCategory
 {
-	FString Name;
+	const TCHAR* Name = nullptr;
 	ELogVerbosity::Type DefaultVerbosity;
 };
 
@@ -38,9 +39,9 @@ struct FLogMessage
 {
 	uint64 Index;
 	double Time;
-	const FLogCategory* Category;
-	const TCHAR* File;
-	const TCHAR* Message;
+	const FLogCategory* Category = nullptr;
+	const TCHAR* File = nullptr;
+	const TCHAR* Message = nullptr;
 	int32 Line;
 	ELogVerbosity::Type Verbosity;
 };
@@ -50,9 +51,10 @@ class ILogProvider
 public:
 	virtual ~ILogProvider() = default;
 	virtual uint64 GetMessageCount() const = 0;
-	virtual void EnumerateMessages(double IntervalStart, double IntervalEnd, TFunctionRef<void(const FLogMessage&)> Callback, bool ResolveFormatString = true) const = 0;
-	virtual void EnumerateMessagesByIndex(uint64 Start, uint64 End, TFunctionRef<void(const FLogMessage&)> Callback, bool ResolveFormatString = true) const = 0;
-	virtual bool ReadMessage(uint64 Index, TFunctionRef<void(const FLogMessage&)> Callback, bool ResolveFormatString = true) const = 0;
+	virtual void EnumerateMessages(double IntervalStart, double IntervalEnd, TFunctionRef<void(const FLogMessage&)> Callback) const = 0;
+	virtual void EnumerateMessagesByIndex(uint64 Start, uint64 End, TFunctionRef<void(const FLogMessage&)> Callback) const = 0;
+	virtual bool ReadMessage(uint64 Index, TFunctionRef<void(const FLogMessage&)> Callback) const = 0;
+	virtual uint64 GetCategoryCount() const = 0;
 	virtual void EnumerateCategories(TFunctionRef<void(const FLogCategory&)> Callback) const = 0;
 };
 
@@ -67,6 +69,7 @@ class IThreadProvider
 {
 public:
 	virtual ~IThreadProvider() = default;
+	virtual uint64 GetModCount() const = 0;
 	virtual void EnumerateThreads(TFunctionRef<void(const FThreadInfo&)> Callback) const = 0;
 };
 
@@ -75,6 +78,7 @@ class ITimeline
 {
 public:
 	virtual ~ITimeline() = default;
+	virtual uint64 GetModCount() const = 0;
 	virtual uint64 GetEventCount() const = 0;
 	virtual void EnumerateEventsDownSampled(double IntervalStart, double IntervalEnd, double Resolution, TFunctionRef<void(bool, double, const EventType&)> Callback) const = 0;
 	virtual void EnumerateEventsDownSampled(double IntervalStart, double IntervalEnd, double Resolution, TFunctionRef<void(double, double, uint32, const EventType&)> Callback) const = 0;
@@ -104,6 +108,7 @@ public:
 	virtual bool GetCpuThreadTimelineIndex(uint32 ThreadId, uint32& OutTimelineIndex) const = 0;
 	virtual bool GetGpuTimelineIndex(uint32& OutTimelineIndex) const = 0;
 	virtual bool ReadTimeline(uint32 Index, TFunctionRef<void(const Timeline&)> Callback) const = 0;
+	virtual uint64 GetTimelineCount() const = 0;
 	virtual void EnumerateTimelines(TFunctionRef<void(const Timeline&)> Callback) const = 0;
 	virtual void ReadTimers(TFunctionRef<void(const FTimingProfilerTimer*, uint64)> Callback) const = 0;
 };
@@ -200,6 +205,7 @@ public:
 	typedef ITimeline<FLoadTimeProfilerCpuEvent> CpuTimeline;
 
 	virtual ~ILoadTimeProfilerProvider() = default;
+	virtual uint64 GetPackageCount() const = 0;
 	virtual void EnumeratePackages(TFunctionRef<void(const FPackageInfo&)> Callback) const = 0; 
 	virtual void ReadMainThreadCpuTimeline(TFunctionRef<void(const CpuTimeline&)> Callback) const = 0;
 	virtual void ReadAsyncLoadingThreadCpuTimeline(TFunctionRef<void(const CpuTimeline&)> Callback) const = 0;
@@ -228,6 +234,7 @@ class ICountersProvider
 {
 public:
 	virtual ~ICountersProvider() = default;
+	virtual uint64 GetCounterCount() const = 0;
 	virtual void EnumerateCounters(TFunctionRef<void(const ICounter&)> Callback) const = 0;
 };
 
diff --git a/Engine/Source/Runtime/Core/Private/Logging/LogTrace.cpp b/Engine/Source/Runtime/Core/Private/Logging/LogTrace.cpp
--- a/Engine/Source/Runtime/Core/Private/Logging/LogTrace.cpp
+++ b/Engine/Source/Runtime/Core/Private/Logging/LogTrace.cpp
@@ -10,7 +10,6 @@
 
 UE_TRACE_EVENT_BEGIN(Logging, LogCategory)
 	UE_TRACE_EVENT_FIELD(const void*, CategoryPointer)
-	UE_TRACE_EVENT_FIELD(uint16, NameLength)
 	UE_TRACE_EVENT_FIELD(uint8, DefaultVerbosity)
 UE_TRACE_EVENT_END()
 
@@ -18,8 +17,6 @@ UE_TRACE_EVENT_BEGIN(Logging, LogMessageSpec)
 	UE_TRACE_EVENT_FIELD(const void*, LogPoint)
 	UE_TRACE_EVENT_FIELD(const void*, CategoryPointer)
 	UE_TRACE_EVENT_FIELD(int32, Line)
-	UE_TRACE_EVENT_FIELD(uint16, FileNameLength)
-	UE_TRACE_EVENT_FIELD(uint16, FormatStringLength)
 	UE_TRACE_EVENT_FIELD(uint8, Verbosity)
 UE_TRACE_EVENT_END()
 
@@ -31,43 +28,25 @@ UE_TRACE_EVENT_END()
 
 void FLogTrace::OutputLogCategory(const FLogCategoryBase* Category, const TCHAR* Name, ELogVerbosity::Type DefaultVerbosity)
 {
-	uint16 NameLength = FCString::Strlen(Name);
-	auto NameCopy = [NameLength, Name](uint8* Out) {
-		TCHAR* OutTChar = reinterpret_cast<TCHAR*>(Out);
-		for (int i = 0; i < NameLength; ++i)
-		{
-			*OutTChar++ = Name[i];
-		}
-	};
-	UE_TRACE_LOG(Logging, LogCategory, NameLength * 2)
+	uint16 NameSize = (FCString::Strlen(Name) + 1) * sizeof(TCHAR);
+	UE_TRACE_LOG(Logging, LogCategory, NameSize)
 		<< LogCategory.CategoryPointer(Category)
-		<< LogCategory.NameLength(NameLength)
 		<< LogCategory.DefaultVerbosity(DefaultVerbosity)
-		<< LogCategory.Attachment(NameCopy);
+		<< LogCategory.Attachment(Name, NameSize);
 }
 
 void FLogTrace::OutputLogMessageSpec(const void* LogPoint, const FLogCategoryBase* Category, ELogVerbosity::Type Verbosity, const ANSICHAR* File, int32 Line, const TCHAR* Format)
 {
-	uint16 FileNameLength = strlen(File);
-	uint16 FormatStringLength = FCString::Strlen(Format);
-	auto StringCopyFunc = [FileNameLength, FormatStringLength, File, Format](uint8* Out) {
-		ANSICHAR* OutAnsiChar = reinterpret_cast<ANSICHAR*>(Out);
-		for (int i = 0; i < FileNameLength; ++i)
-		{
-			*OutAnsiChar++ = File[i];
-		}
-		TCHAR* OutTChar = reinterpret_cast<TCHAR*>(OutAnsiChar);
-		for (int i = 0; i < FormatStringLength; ++i)
-		{
-			*OutTChar++ = Format[i];
-		}
+	uint16 FileNameSize = strlen(File) + 1;
+	uint16 FormatStringSize = (FCString::Strlen(Format) + 1) * sizeof(TCHAR);
+	auto StringCopyFunc = [FileNameSize, FormatStringSize, File, Format](uint8* Out) {
+		memcpy(Out, File, FileNameSize);
+		memcpy(Out + FileNameSize, Format, FormatStringSize);
 	};
-	UE_TRACE_LOG(Logging, LogMessageSpec, FileNameLength + FormatStringLength * 2)
+	UE_TRACE_LOG(Logging, LogMessageSpec, FileNameSize + FormatStringSize)
 		<< LogMessageSpec.LogPoint(LogPoint)
 		<< LogMessageSpec.CategoryPointer(Category)
 		<< LogMessageSpec.Line(Line)
-		<< LogMessageSpec.FileNameLength(FileNameLength)
-		<< LogMessageSpec.FormatStringLength(FormatStringLength)
 		<< LogMessageSpec.Verbosity(Verbosity)
 		<< LogMessageSpec.Attachment(StringCopyFunc);
 }
diff --git a/Engine/Source/Runtime/Core/Private/ProfilingDebugging/MiscTrace.cpp b/Engine/Source/Runtime/Core/Private/ProfilingDebugging/MiscTrace.cpp
--- a/Engine/Source/Runtime/Core/Private/ProfilingDebugging/MiscTrace.cpp
+++ b/Engine/Source/Runtime/Core/Private/ProfilingDebugging/MiscTrace.cpp
@@ -36,8 +36,6 @@ UE_TRACE_EVENT_END()
 UE_TRACE_EVENT_BEGIN(Misc, BookmarkSpec)
 	UE_TRACE_EVENT_FIELD(const void*, BookmarkPoint)
 	UE_TRACE_EVENT_FIELD(int32, Line)
-	UE_TRACE_EVENT_FIELD(uint16, FileNameLength)
-	UE_TRACE_EVENT_FIELD(uint16, FormatStringLength)
 UE_TRACE_EVENT_END()
 
 UE_TRACE_EVENT_BEGIN(Misc, Bookmark)
@@ -150,25 +148,15 @@ void FMiscTrace::OutputEndThreadGroupScope()
 
 void FMiscTrace::OutputBookmarkSpec(const void* BookmarkPoint, const ANSICHAR* File, int32 Line, const TCHAR* Format)
 {
-	uint16 FileNameLength = strlen(File);
-	int32 FormatStringLength = FCString::Strlen(Format);
-	auto StringCopyFunc = [FileNameLength, FormatStringLength, File, Format](uint8* Out) {
-		ANSICHAR* OutAnsiChar = reinterpret_cast<ANSICHAR*>(Out);
-		for (int i = 0; i < FileNameLength; ++i)
-		{
-			*OutAnsiChar++ = File[i];
-		}
-		TCHAR* OutTChar = reinterpret_cast<TCHAR*>(OutAnsiChar);
-		for (int i = 0; i < FormatStringLength; ++i)
-		{
-			*OutTChar++ = Format[i];
-		}
+	uint16 FileNameSize = strlen(File) + 1;
+	uint16 FormatStringSize = (FCString::Strlen(Format) + 1) * sizeof(TCHAR);
+	auto StringCopyFunc = [FileNameSize, FormatStringSize, File, Format](uint8* Out) {
+		memcpy(Out, File, FileNameSize);
+		memcpy(Out + FileNameSize, Format, FormatStringSize);
 	};
-	UE_TRACE_LOG(Misc, BookmarkSpec, FileNameLength + FormatStringLength * 2)
+	UE_TRACE_LOG(Misc, BookmarkSpec, FileNameSize + FormatStringSize)
 		<< BookmarkSpec.BookmarkPoint(BookmarkPoint)
 		<< BookmarkSpec.Line(Line)
-		<< BookmarkSpec.FileNameLength(FileNameLength)
-		<< BookmarkSpec.FormatStringLength(FormatStringLength)
 		<< BookmarkSpec.Attachment(StringCopyFunc);
 }
 
diff --git a/Engine/Source/Runtime/Core/Public/Logging/LogTrace.h b/Engine/Source/Runtime/Core/Public/Logging/LogTrace.h
--- a/Engine/Source/Runtime/Core/Public/Logging/LogTrace.h
+++ b/Engine/Source/Runtime/Core/Public/Logging/LogTrace.h
@@ -5,7 +5,7 @@
 #include "CoreTypes.h"
 
 #if !IS_PROGRAM && !UE_BUILD_SHIPPING
-#define LOGTRACE_ENABLED 0
+#define LOGTRACE_ENABLED 1
 #else
 #define LOGTRACE_ENABLED 0
 #endif
