UnrealInsights:
Stat counters trace analysis optimizations

diff --git a/Engine/Source/Developer/TraceServices/Private/Common/PagedArray.h b/Engine/Source/Developer/TraceServices/Private/Common/PagedArray.h
--- a/Engine/Source/Developer/TraceServices/Private/Common/PagedArray.h
+++ b/Engine/Source/Developer/TraceServices/Private/Common/PagedArray.h
@@ -4,30 +4,62 @@
 
 #include "Common/SlabAllocator.h"
 
-struct FNoPageData
+template<typename ItemType>
+struct TPagedArrayPage
 {
+	ItemType* Items = nullptr;
+	uint64 Count = 0;
 };
 
-template<typename ItemType, typename PageDataType = FNoPageData>
-class TPagedArray
+template<typename ItemType>
+inline const ItemType* GetData(const TPagedArrayPage<ItemType>& Page)
+{
+	return Page.Items;
+}
+
+template<typename ItemType>
+inline SIZE_T GetNum(const TPagedArrayPage<ItemType>& Page)
 {
-	struct FPageInternal;
+	return Page.Count;
+}
+
+template<typename ItemType>
+inline const ItemType* GetFirstItem(const TPagedArrayPage<ItemType>& Page)
+{
+	return Page.Items;
+}
+
+template<typename ItemType>
+inline const ItemType* GetLastItem(const TPagedArrayPage<ItemType>& Page)
+{
+	if (Page.Count)
+	{
+		return Page.Items + Page.Count - 1;
+	}
+	else
+	{
+		return nullptr;
+	}
+}
 
+template<typename InItemType, typename InPageType = TPagedArrayPage<InItemType>>
+class TPagedArray
+{
 public:
+	typedef InItemType ItemType;
+	typedef InPageType PageType;
+	
 	class FIterator
 	{
 	public:
-		const PageDataType* GetCurrentPage()
+		uint64 GetCurrentPageIndex()
 		{
-			if (CurrentPage)
-			{
-				uint64 CurrentPageIndex = CurrentPage - FirstPage;
-				return Outer.PageUserDatas.GetData() + CurrentPageIndex;
-			}
-			else
-			{
-				return nullptr;
-			}
+			return CurrentPageIndex;
+		}
+
+		const PageType* GetCurrentPage()
+		{
+			return CurrentPage;
 		}
 
 		const ItemType* GetCurrentItem()
@@ -35,26 +67,28 @@ public:
 			return CurrentItem;
 		}
 
-		const PageDataType* PrevPage()
+		const PageType* PrevPage()
 		{
-			if (CurrentPage == FirstPage)
+			if (CurrentPageIndex == 0)
 			{
 				return nullptr;
 			}
-			--CurrentPage;
+			--CurrentPageIndex;
+			CurrentPage = Outer.FirstPage + CurrentPageIndex;
 			CurrentPageBegin = CurrentPage->Items - 1;
 			CurrentPageEnd = CurrentPage->Items + CurrentPage->Count;
 			CurrentItem = CurrentPageEnd - 1;
 			return GetCurrentPage();
 		}
 
-		const PageDataType* NextPage()
+		const PageType* NextPage()
 		{
-			if (CurrentPage == LastPage)
+			if (CurrentPageIndex == Outer.PagesArray.Num() - 1)
 			{
 				return nullptr;
 			}
-			++CurrentPage;
+			++CurrentPageIndex;
+			CurrentPage = Outer.FirstPage + CurrentPageIndex;
 			CurrentPageBegin = CurrentPage->Items - 1;
 			CurrentPageEnd = CurrentPage->Items + CurrentPage->Count;
 			CurrentItem = CurrentPage->Items;
@@ -84,15 +118,14 @@ public:
 	private:
 		friend class TPagedArray;
 
-		FIterator(const TPagedArray& InOuter, uint64 InitialBankIndex, uint64 InitialItemIndex)
+		FIterator(const TPagedArray& InOuter, uint64 InitialPageIndex, uint64 InitialItemIndex)
 			: Outer(InOuter)
 		{
-			if (Outer.Pages.Num() > 0)
+			if (Outer.PagesArray.Num())
 			{
-				FirstPage = Outer.Pages.GetData();
-				LastPage = Outer.Pages.GetData() + Outer.Pages.Num() - 1;
-				check(InitialBankIndex < Outer.Pages.Num());
-				CurrentPage = Outer.Pages.GetData() + InitialBankIndex;
+				check(InitialPageIndex < Outer.PagesArray.Num());
+				CurrentPageIndex = InitialPageIndex;
+				CurrentPage = Outer.FirstPage + InitialPageIndex;
 				CurrentPageBegin = CurrentPage->Items - 1;
 				CurrentPageEnd = CurrentPage->Items + CurrentPage->Count;
 				check(InitialItemIndex < CurrentPage->Count);
@@ -101,9 +134,8 @@ public:
 		}
 
 		const TPagedArray& Outer;
-		const FPageInternal* CurrentPage = nullptr;
-		const FPageInternal* FirstPage = nullptr;
-		const FPageInternal* LastPage = nullptr;
+		const PageType* CurrentPage = nullptr;
+		uint64 CurrentPageIndex = 0;
 		const ItemType* CurrentItem = nullptr;
 		const ItemType* CurrentPageBegin = nullptr;
 		const ItemType* CurrentPageEnd = nullptr;
@@ -118,7 +150,7 @@ public:
 
 	~TPagedArray()
 	{
-		for (FPageInternal& Page : Pages)
+		for (PageType& Page : PagesArray)
 		{
 			ItemType* PageEnd = Page.Items + Page.Count;
 			for (ItemType* Item = Page.Items; Item != PageEnd; ++Item)
@@ -133,13 +165,23 @@ public:
 		return TotalItemCount;
 	}
 
+	uint64 GetPageSize() const
+	{
+		return PageSize;
+	}
+
+	uint64 NumPages() const
+	{
+		return PagesArray.Num();
+	}
+
 	ItemType& PushBack()
 	{
 		if (!LastPage || LastPage->Count == PageSize)
 		{
-			LastPage = &Pages.AddDefaulted_GetRef();
+			LastPage = &PagesArray.AddDefaulted_GetRef();
+			FirstPage = PagesArray.GetData();
 			LastPage->Items = Allocator.Allocate<ItemType>(PageSize);
-			PageUserDatas.AddDefaulted_GetRef();
 		}
 		++TotalItemCount;
 		ItemType* ItemPtr = LastPage->Items + LastPage->Count;
@@ -150,35 +192,39 @@ public:
 
 	ItemType& Insert(uint64 Index)
 	{
+		if (Index >= TotalItemCount)
+		{
+			return PushBack();
+		}
 		PushBack();
-		for (uint64 CurrentIndex = TotalItemCount - 1; CurrentIndex > Index; --CurrentIndex)
+		uint64 PageIndex = Index / PageSize;
+		uint64 PageItemIndex = Index % PageSize;
+		for (uint64 CurrentPageIndex = PagesArray.Num() - 1; CurrentPageIndex > PageIndex; --CurrentPageIndex)
 		{
-			(*this)[CurrentIndex] = (*this)[CurrentIndex - 1];
+			PageType* CurrentPage = FirstPage + CurrentPageIndex;
+			memmove(CurrentPage->Items + 1, CurrentPage->Items, sizeof(ItemType) * (CurrentPage->Count - 1));
+			PageType* PrevPage = CurrentPage - 1;
+			memcpy(CurrentPage->Items, PrevPage->Items + PrevPage->Count - 1, sizeof(ItemType));
 		}
-		return (*this)[Index];
+		PageType* Page = FirstPage + PageIndex;
+		memmove(Page->Items + PageItemIndex + 1, Page->Items + PageItemIndex, sizeof(ItemType) * (Page->Count - PageItemIndex - 1));
+		return Page->Items[PageItemIndex];
 	}
 
-	PageDataType* GetLastPage()
+	PageType* GetLastPage()
 	{
-		if (PageUserDatas.Num())
-		{
-			return &PageUserDatas.Last();
-		}
-		else
-		{
-			return nullptr;
-		}
+		return LastPage;
 	}
 
-	PageDataType* GetPage(uint64 PageIndex)
+	PageType* GetPage(uint64 PageIndex)
 	{
-		return PageUserDatas.GetData() + PageIndex;
+		return FirstPage + PageIndex;
 	}
 
-	PageDataType* GetItemPage(uint64 ItemIndex)
+	PageType* GetItemPage(uint64 ItemIndex)
 	{
 		uint64 PageIndex = ItemIndex / PageSize;
-		return PageUserDatas.GetData() + PageIndex;
+		return FirstPage + PageIndex;
 	}
 
 	FIterator GetIteratorFromPage(uint64 PageIndex) const
@@ -193,16 +239,16 @@ public:
 		return FIterator(*this, PageIndex, IndexInPage);
 	}
 
-	const TArray<PageDataType> GetPages() const
+	const PageType* GetPages() const
 	{
-		return PageUserDatas;
+		return FirstPage;
 	}
 
 	ItemType& operator[](uint64 Index)
 	{
 		uint64 PageIndex = Index / PageSize;
 		uint64 IndexInPage = Index % PageSize;
-		FPageInternal* Page = Pages.GetData() + PageIndex;
+		PageType* Page = FirstPage + PageIndex;
 		ItemType* Item = Page->Items + IndexInPage;
 		return *Item;
 	}
@@ -213,16 +259,22 @@ public:
 	}
 
 private:
-	struct FPageInternal
-	{
-		ItemType* Items = nullptr;
-		uint64 Count = 0;
-	};
-
 	FSlabAllocator& Allocator;
-	TArray<FPageInternal> Pages;
-	TArray<PageDataType> PageUserDatas;
-	FPageInternal* LastPage = nullptr;
+	TArray<PageType> PagesArray;
+	PageType* FirstPage = nullptr;
+	PageType* LastPage = nullptr;
 	uint64 PageSize;
 	uint64 TotalItemCount = 0;
 };
+
+template<typename ItemType, typename PageType>
+inline const PageType* GetData(const TPagedArray<ItemType, PageType>& PagedArray)
+{
+	return PagedArray.GetPages();
+}
+
+template<typename ItemType, typename PageType>
+inline SIZE_T GetNum(const TPagedArray<ItemType, PageType>& PagedArray)
+{
+	return PagedArray.NumPages();
+}
\ No newline at end of file
diff --git a/Engine/Source/Developer/TraceServices/Private/Model/Counters.cpp b/Engine/Source/Developer/TraceServices/Private/Model/Counters.cpp
--- a/Engine/Source/Developer/TraceServices/Private/Model/Counters.cpp
+++ b/Engine/Source/Developer/TraceServices/Private/Model/Counters.cpp
@@ -129,11 +129,28 @@ template<typename T>
 static void InsertOp(FCounterInternal& Counter, double Time, ECounterOpType Type, T Arg)
 {
 	TPagedArray<double>& Timestamps = Counter.EditTimestamps();
-	uint64 Count = Timestamps.Num();
-	uint64 InsertionIndex = Count;
-	while (InsertionIndex > 0 && Timestamps[InsertionIndex - 1] > Time)
+	uint64 InsertionIndex;
+	if (Timestamps.Num() == 0 || Timestamps[Timestamps.Num() - 1] <= Time)
 	{
-		--InsertionIndex;
+		InsertionIndex = Timestamps.Num();
+	}
+	else
+	{
+		auto TimestampIterator = Timestamps.GetIteratorFromItem(Timestamps.Num() - 1);
+		auto CurrentPage = TimestampIterator.GetCurrentPage();
+		while (CurrentPage && *GetFirstItem(*CurrentPage) > Time)
+		{
+			CurrentPage = TimestampIterator.PrevPage();
+		}
+		if (!CurrentPage)
+		{
+			InsertionIndex = 0;
+		}
+		else
+		{
+			uint64 PageInsertionIndex = Algo::LowerBound(*CurrentPage, Time);
+			InsertionIndex = TimestampIterator.GetCurrentPageIndex() * Timestamps.GetPageSize() + PageInsertionIndex;
+		}
 	}
 	Timestamps.Insert(InsertionIndex) = Time;
 	Counter.EditOps().Insert(InsertionIndex) = Type;
diff --git a/Engine/Source/Developer/TraceServices/Private/Model/IntervalTimeline.h b/Engine/Source/Developer/TraceServices/Private/Model/IntervalTimeline.h
--- a/Engine/Source/Developer/TraceServices/Private/Model/IntervalTimeline.h
+++ b/Engine/Source/Developer/TraceServices/Private/Model/IntervalTimeline.h
@@ -109,19 +109,21 @@ public:
 	}
 
 private:
+	struct FEventInternal
+	{
+		double StartTime;
+		double EndTime;
+		EventType Event;
+	};
+
 	struct FEventPage
 	{
 		double BeginTime = 0.0;
 		double EndTime = 0.0;
 		uint64 BeginEventIndex = 0;
 		uint64 EndEventIndex = 0;
-	};
-
-	struct FEventInternal
-	{
-		double StartTime;
-		double EndTime;
-		EventType Event;
+		FEventInternal* Items = nullptr;
+		uint64 Count = 0;
 	};
 
 	TPagedArray<FEventInternal, FEventPage> Events;
diff --git a/Engine/Source/Developer/TraceServices/Private/Model/MonotonicTimeline.h b/Engine/Source/Developer/TraceServices/Private/Model/MonotonicTimeline.h
--- a/Engine/Source/Developer/TraceServices/Private/Model/MonotonicTimeline.h
+++ b/Engine/Source/Developer/TraceServices/Private/Model/MonotonicTimeline.h
@@ -139,7 +139,7 @@ public:
 			return;
 		}
 
-		uint64 FirstScopePageIndex = Algo::UpperBoundBy(DetailLevel.ScopeEntries.GetPages(), IntervalStart, [](const FEventScopeEntryPage& Page)
+		uint64 FirstScopePageIndex = Algo::UpperBoundBy(DetailLevel.ScopeEntries, IntervalStart, [](const FEventScopeEntryPage& Page)
 		{
 			return Page.BeginTime;
 		});
@@ -322,6 +322,8 @@ private:
 
 	struct FEventScopeEntryPage
 	{
+		FEventScopeEntry* Items = nullptr;
+		uint64 Count = 0;
 		double BeginTime = 0.0;
 		double EndTime = 0.0;
 		uint64 BeginEventIndex = 0;
