UnrealInsights:
Trace files are now stored in the project Saved directory

diff --git a/Engine/Source/Developer/TraceServices/Private/AnalysisService.cpp b/Engine/Source/Developer/TraceServices/Private/AnalysisService.cpp
--- a/Engine/Source/Developer/TraceServices/Private/AnalysisService.cpp
+++ b/Engine/Source/Developer/TraceServices/Private/AnalysisService.cpp
@@ -140,8 +140,7 @@ void FAnalysisSession::ReadCountersProvider(TFunctionRef<void(const ICountersPro
 	Callback(*CountersProvider);
 }
 
-FAnalysisService::FAnalysisService(TSharedRef<IStore> InTraceStore)
-	: TraceStore(InTraceStore)
+FAnalysisService::FAnalysisService()
 {
 
 }
diff --git a/Engine/Source/Developer/TraceServices/Private/AnalysisServicePrivate.h b/Engine/Source/Developer/TraceServices/Private/AnalysisServicePrivate.h
--- a/Engine/Source/Developer/TraceServices/Private/AnalysisServicePrivate.h
+++ b/Engine/Source/Developer/TraceServices/Private/AnalysisServicePrivate.h
@@ -126,7 +126,7 @@ class FAnalysisService
 	: public IAnalysisService
 {
 public:
-	FAnalysisService(TSharedRef<IStore> InTraceStore);
+	FAnalysisService();
 	virtual ~FAnalysisService();
 	virtual TSharedPtr<const IAnalysisSession> Analyze(const TCHAR* SessionName, TUniquePtr<Trace::IInDataStream>&& DataStream) override;
 	virtual TSharedPtr<const IAnalysisSession> StartAnalysis(const TCHAR* SessionName, TUniquePtr<Trace::IInDataStream>&& DataStream) override;
@@ -169,7 +169,6 @@ private:
 
 	FAnalysisStartedEvent AnalysisStartedEvent;
 	FAnalysisFinishedEvent AnalysisFinishedEvent;
-	TSharedRef<IStore> TraceStore;
 	TArray<TSharedPtr<FAsyncTask<FAnalysisWorker>>> Tasks;
 };
 
diff --git a/Engine/Source/Developer/TraceServices/Private/SessionService.cpp b/Engine/Source/Developer/TraceServices/Private/SessionService.cpp
--- a/Engine/Source/Developer/TraceServices/Private/SessionService.cpp
+++ b/Engine/Source/Developer/TraceServices/Private/SessionService.cpp
@@ -5,14 +5,16 @@
 #include "Trace/Store.h"
 #include "Trace/Recorder.h"
 #include "Trace/DataStream.h"
+#include "Misc/Paths.h"
 
 namespace Trace
 {
 
-FSessionService::FSessionService(TSharedRef<IStore> InTraceStore)
-	: TraceStore(InTraceStore)
+FSessionService::FSessionService()
 {
-	TraceRecorder = Recorder_Create(TraceStore);
+	LocalSessionDirectory = FPaths::ProjectSavedDir() / TEXT("TraceSessions");
+	TraceStore = Store_Create(*LocalSessionDirectory);
+	TraceRecorder = Recorder_Create(TraceStore.ToSharedRef());
 }
 
 FSessionService::~FSessionService()
@@ -55,19 +57,4 @@ Trace::IInDataStream* FSessionService::OpenSessionFromFile(const TCHAR* FilePath
 	return Trace::DataStream_ReadFile(FilePath);
 }
 
-/*void FSessionService::GetSessions(TArray<TSharedRef<FSessionInfo>>& OutSessions)
-{
-	OutSessions.Reset(AvailableSessions.Num());
-	for (const auto& KV : AvailableSessions)
-	{
-		OutSessions.Add(KV.Value);
-	}
-}*/
-
-/*void FSessionService::SelectSession(TSharedPtr<FSessionInfo> Session)
-{
-	SelectedSession = Session;
-	OnSessionSelected().Broadcast(SelectedSession);
-}*/
-
 }
diff --git a/Engine/Source/Developer/TraceServices/Private/SessionServicePrivate.h b/Engine/Source/Developer/TraceServices/Private/SessionServicePrivate.h
--- a/Engine/Source/Developer/TraceServices/Private/SessionServicePrivate.h
+++ b/Engine/Source/Developer/TraceServices/Private/SessionServicePrivate.h
@@ -13,18 +13,20 @@ class FSessionService
 	: public ISessionService
 {
 public:
-	FSessionService(TSharedRef<IStore> InTraceStore);
+	FSessionService();
 	virtual ~FSessionService();
 	virtual bool StartRecorderServer() override;
 	virtual bool IsRecorderServerRunning() const override;
 	virtual void StopRecorderServer() override;
+	virtual const TCHAR* GetLocalSessionDirectory() const override { return *LocalSessionDirectory; }
 	virtual void GetAvailableSessions(TArray<Trace::FSessionHandle>& OutSessions) override;
 	virtual bool GetSessionInfo(Trace::FSessionHandle SessionHandle, Trace::FSessionInfo& OutSessionInfo) override;
 	virtual Trace::IInDataStream* OpenSessionStream(Trace::FSessionHandle SessionHandle) override;
 	virtual Trace::IInDataStream* OpenSessionFromFile(const TCHAR* FilePath) override;
 
 private:
-	TSharedRef<IStore> TraceStore;
+	FString LocalSessionDirectory;
+	TSharedPtr<IStore> TraceStore;
 	TSharedPtr<IRecorder> TraceRecorder;
 };
 
diff --git a/Engine/Source/Developer/TraceServices/Private/TraceServicesModule.cpp b/Engine/Source/Developer/TraceServices/Private/TraceServicesModule.cpp
--- a/Engine/Source/Developer/TraceServices/Private/TraceServicesModule.cpp
+++ b/Engine/Source/Developer/TraceServices/Private/TraceServicesModule.cpp
@@ -9,26 +9,19 @@ class FTraceServicesModule
 	: public ITraceServicesModule
 {
 public:
-	virtual void StartupModule() override;
 	virtual TSharedPtr<Trace::ISessionService> GetSessionService() override;
 	virtual TSharedPtr<Trace::IAnalysisService> GetAnalysisService() override;
 
 private:
-	TSharedPtr<Trace::IStore> TraceStore;
 	TSharedPtr<Trace::FSessionService> SessionService;
 	TSharedPtr<Trace::FAnalysisService> AnalysisService;
 };
 
-void FTraceServicesModule::StartupModule()
-{
-	TraceStore = Trace::Store_Create(TEXT("D:\\Trace"));
-}
-
 TSharedPtr<Trace::ISessionService> FTraceServicesModule::GetSessionService()
 {
 	if (!SessionService.IsValid())
 	{
-		SessionService = MakeShareable(new Trace::FSessionService(TraceStore.ToSharedRef()));
+		SessionService = MakeShared<Trace::FSessionService>();
 	}
 	return SessionService;
 }
@@ -37,7 +30,7 @@ TSharedPtr<Trace::IAnalysisService> FTraceServicesModule::GetAnalysisService()
 {
 	if (!AnalysisService.IsValid())
 	{
-		AnalysisService = MakeShareable(new Trace::FAnalysisService(TraceStore.ToSharedRef()));
+		AnalysisService = MakeShared<Trace::FAnalysisService>();
 	}
 	return AnalysisService;
 }
diff --git a/Engine/Source/Developer/TraceServices/Public/TraceServices/SessionService.h b/Engine/Source/Developer/TraceServices/Public/TraceServices/SessionService.h
--- a/Engine/Source/Developer/TraceServices/Public/TraceServices/SessionService.h
+++ b/Engine/Source/Developer/TraceServices/Public/TraceServices/SessionService.h
@@ -17,6 +17,7 @@ public:
 	virtual bool StartRecorderServer() = 0;
 	virtual bool IsRecorderServerRunning() const = 0;
 	virtual void StopRecorderServer() = 0;
+	virtual const TCHAR* GetLocalSessionDirectory() const = 0;
 	virtual void GetAvailableSessions(TArray<Trace::FSessionHandle>& OutSessions) = 0;
 	virtual bool GetSessionInfo(Trace::FSessionHandle SessionHandle, Trace::FSessionInfo& OutSessionInfo) = 0;
 	virtual Trace::IInDataStream* OpenSessionStream(Trace::FSessionHandle SessionHandle) = 0;
