UnrealInsights:
Added table API
Moved timeline statistics calculations to TraceServices
Added some load time profiler statistics

diff --git a/Engine/Source/Developer/TraceInsights/Private/Insights/ViewModels/TimerNode.cpp b/Engine/Source/Developer/TraceInsights/Private/Insights/ViewModels/TimerNode.cpp
--- a/Engine/Source/Developer/TraceInsights/Private/Insights/ViewModels/TimerNode.cpp
+++ b/Engine/Source/Developer/TraceInsights/Private/Insights/ViewModels/TimerNode.cpp
@@ -8,12 +8,12 @@
 
 void FTimerNode::ResetAggregatedStats()
 {
-	AggregatedStats.Reset();
+	AggregatedStats = Trace::FAggregatedTimingStats();
 }
 
 ////////////////////////////////////////////////////////////////////////////////////////////////////
 
-void FTimerNode::SetAggregatedStats(FTimerAggregatedStats& InAggregatedStats)
+void FTimerNode::SetAggregatedStats(const Trace::FAggregatedTimingStats& InAggregatedStats)
 {
 	AggregatedStats = InAggregatedStats;
 }
diff --git a/Engine/Source/Developer/TraceInsights/Private/Insights/ViewModels/TimerNode.h b/Engine/Source/Developer/TraceInsights/Private/Insights/ViewModels/TimerNode.h
--- a/Engine/Source/Developer/TraceInsights/Private/Insights/ViewModels/TimerNode.h
+++ b/Engine/Source/Developer/TraceInsights/Private/Insights/ViewModels/TimerNode.h
@@ -3,6 +3,7 @@
 #pragma once
 
 #include "CoreMinimal.h"
+#include "TraceServices/AnalysisService.h"
 
 class FTimerNode;
 
@@ -36,58 +37,6 @@ enum class ETimerNodeType
 	InvalidOrMax,
 };
 
-struct FTimerAggregatedStats
-{
-	//FAggregatedStats InclusiveTime; /** Aggregated stats for inclusive time, in seconds. */
-	//FAggregatedStats ExclusiveTime; /** Aggregated stats for exclusive time, in seconds. */
-	////FAggregatedStats InstanceCountPerFrame; /** Aggregated stats for variation of number of instances per frame. */
-	//uint64 GetInstanceCount() { return InclusiveTime.Count; } /** Number of instances. */
-
-	double TotalInclusiveTime; /** Total inclusive time, in seconds. */
-	double MinInclusiveTime; /** Min inclusive time, in seconds. */
-	double MaxInclusiveTime; /** Max inclusive time, in seconds. */
-	double AverageInclusiveTime; /** Average inclusive time, in seconds. */
-	double MedianInclusiveTime; /** Median inclusive time, in seconds. */
-
-	double TotalExclusiveTime; /** Total exclusive time, in seconds. */
-	double MinExclusiveTime; /** Min exclusive time, in seconds. */
-	double MaxExclusiveTime; /** Max exclusive time, in seconds. */
-	double AverageExclusiveTime; /** Average exclusive time, in seconds. */
-	double MedianExclusiveTime; /** Median exclusive time, in seconds. */
-
-	uint64 InstanceCount; /** Number of instances. */
-
-	FTimerAggregatedStats()
-		: TotalInclusiveTime(0.0)
-		, MinInclusiveTime(DBL_MAX)
-		, MaxInclusiveTime(-DBL_MAX)
-		, AverageInclusiveTime(0.0)
-		, MedianInclusiveTime(0.0)
-		, TotalExclusiveTime(0.0)
-		, MinExclusiveTime(DBL_MAX)
-		, MaxExclusiveTime(-DBL_MAX)
-		, AverageExclusiveTime(0.0)
-		, MedianExclusiveTime(0.0)
-		, InstanceCount(0)
-	{
-	}
-
-	void Reset()
-	{
-		TotalInclusiveTime = 0.0;
-		MinInclusiveTime = DBL_MAX;
-		MaxInclusiveTime = -DBL_MAX;
-		AverageInclusiveTime = 0.0;
-		MedianInclusiveTime = 0.0;
-		TotalExclusiveTime = 0.0;
-		MinExclusiveTime = DBL_MAX;
-		MaxExclusiveTime = -DBL_MAX;
-		AverageExclusiveTime = 0.0;
-		MedianExclusiveTime = 0.0;
-		InstanceCount = 0;
-	}
-};
-
 /**
  * Class used to store information about a timer node (used in timers' tree view).
  */
@@ -166,14 +115,14 @@ public:
 	/**
 	 * @return the aggregated stats for this timer.
 	 */
-	const FTimerAggregatedStats& GetAggregatedStats() const
+	const Trace::FAggregatedTimingStats& GetAggregatedStats() const
 	{
 		return AggregatedStats;
 	}
 
 	void ResetAggregatedStats();
 
-	void SetAggregatedStats(FTimerAggregatedStats& AggregatedStats);
+	void SetAggregatedStats(const Trace::FAggregatedTimingStats& AggregatedStats);
 
 	/**
 	 * @return a const reference to the child nodes of this group.
@@ -261,7 +210,7 @@ protected:
 	const ETimerNodeType Type;
 
 	/** Aggregated stats. */
-	FTimerAggregatedStats AggregatedStats;
+	Trace::FAggregatedTimingStats AggregatedStats;
 
 	/** Children of this node. */
 	TArray<FTimerNodePtr> Children;
diff --git a/Engine/Source/Developer/TraceInsights/Private/Insights/Widgets/STimersView.cpp b/Engine/Source/Developer/TraceInsights/Private/Insights/Widgets/STimersView.cpp
--- a/Engine/Source/Developer/TraceInsights/Private/Insights/Widgets/STimersView.cpp
+++ b/Engine/Source/Developer/TraceInsights/Private/Insights/Widgets/STimersView.cpp
@@ -1477,70 +1477,6 @@ void STimersView::RebuildTree(bool bResync)
 
 ////////////////////////////////////////////////////////////////////////////////////////////////////
 
-struct FTimeCalculationHelper
-{
-	struct FStackEntry
-	{
-		double StartTime;
-		double AccumulatedTimeInChildren;
-		uint32 EventType;
-	};
-
-	FTimeCalculationHelper(double InIntervalStartTime, double InIntervalEndTime, TMap<uint64, FTimerStatsEx>& InStatsMap)
-		: IntervalStartTime(InIntervalStartTime)
-		, IntervalEndTime(InIntervalEndTime)
-		, StatsMap(InStatsMap)
-	{
-
-	}
-
-	template<typename CallbackType>
-	void UpdateStats(const Trace::ITimingProfilerProvider::Timeline& Timeline, CallbackType Callback)
-	{
-		EventStack.Empty(EventStack.Num());
-		Timeline.EnumerateEvents(IntervalStartTime, IntervalEndTime, [this, Callback](bool IsEnter, double Time, const Trace::FTimingProfilerEvent& Event)
-		{
-			if (IsEnter)
-			{
-				FStackEntry& StackEntry = EventStack.AddDefaulted_GetRef();
-				StackEntry.StartTime = FMath::Max(IntervalStartTime, Time);
-				StackEntry.AccumulatedTimeInChildren = 0.0;
-				StackEntry.EventType = Event.TimerIndex;
-			}
-			else
-			{
-				FStackEntry& StackEntry = EventStack.Last();
-				double EventInclusiveTime = FMath::Min(IntervalEndTime, Time) - StackEntry.StartTime;
-				check(EventInclusiveTime >= 0.0);
-				double EventExclusiveTime = EventInclusiveTime - StackEntry.AccumulatedTimeInChildren;
-				EventStack.Pop(false);
-				if (EventStack.Num())
-				{
-					EventStack.Last().AccumulatedTimeInChildren += EventInclusiveTime;
-				}
-
-				FTimerStatsEx& Stats = StatsMap.FindOrAdd(Event.TimerIndex);
-				double EventNonRecursiveInclusiveTime = EventInclusiveTime;
-				for (const FStackEntry& AncestorStackEntry : EventStack)
-				{
-					if (AncestorStackEntry.EventType == Event.TimerIndex)
-					{
-						EventNonRecursiveInclusiveTime = 0.0;
-					}
-				}
-				Callback(Stats, EventNonRecursiveInclusiveTime, EventExclusiveTime);
-			}
-		});
-	}
-
-	TArray<FStackEntry> EventStack;
-	double IntervalStartTime;
-	double IntervalEndTime;
-	TMap<uint64, FTimerStatsEx>& StatsMap;
-};
-
-////////////////////////////////////////////////////////////////////////////////////////////////////
-
 void STimersView::UpdateStats(double StartTime, double EndTime)
 {
 	for (const FTimerNodePtr& TimerNodePtr : TimerNodes)
@@ -1551,69 +1487,31 @@ void STimersView::UpdateStats(double StartTime, double EndTime)
 	if (Session.IsValid() &&
 		StartTime < EndTime)
 	{
-		struct ComputeStatsContext
-		{
-			double StartTime;
-			double EndTime;
-			TMap<uint64, FTimerStatsEx> Map;
-		} Ctx;
-
-		Ctx.StartTime = StartTime;
-		Ctx.EndTime = EndTime;
-
-		const bool bComputeMedian = true;
-
-		FTimeCalculationHelper CalculationHelper(Ctx.StartTime, Ctx.EndTime, Ctx.Map);
-
+		TUniquePtr<Trace::ITable<Trace::FTimingProfilerAggregatedStats>> AggregationResultTable;
+		Trace::FAnalysisSessionReadScope SessionReadScope(*Session.Get());
+		Session->ReadTimingProfilerProvider([StartTime, EndTime, &AggregationResultTable](const Trace::ITimingProfilerProvider& TimingProfilerProvider)
 		{
-			Trace::FAnalysisSessionReadScope SessionReadScope(*Session.Get());
-
-			// Compute instance count and total/min/max inclusive/exclusive times for each timer.
-			// Iterate through all timing events (all thread timelines).
-			Session->ReadTimingProfilerProvider([&CalculationHelper](const Trace::ITimingProfilerProvider& TimingProfilerProvider)
-			{
-				TimingProfilerProvider.EnumerateTimelines([&CalculationHelper](const Trace::ITimingProfilerProvider::Timeline& Timeline)
-				{
-					CalculationHelper.UpdateStats(Timeline, UpdateTotalMinMaxTimerStats);
-				});
-			});
-
-			// Now, as we know min/max inclusive/exclusive times for each timer, we can compute histogram and median values.
-			if (bComputeMedian)
+			auto ThreadFilter = [](uint32 ThreadId)
 			{
-				// Update bucket size (DT) for computing histogram.
-				for (auto& KV : Ctx.Map)
-				{
-					PreComputeHistogram(KV.Value);
-				}
-
-				// Compute histogram.
-				// Iterate again through all timing events (all timelines).
-				Session->ReadTimingProfilerProvider([&CalculationHelper](const Trace::ITimingProfilerProvider& TimingProfilerProvider)
-				{
-					TimingProfilerProvider.EnumerateTimelines([&CalculationHelper](const Trace::ITimingProfilerProvider::Timeline& Timeline)
-					{
-						CalculationHelper.UpdateStats(Timeline, UpdateHistogramForTimerStats);
-					});
-				});
-			}
-		}
+				return true;
+			};
+			AggregationResultTable.Reset(TimingProfilerProvider.CreateAggregation(StartTime, EndTime, ThreadFilter, true));
+		});
 
-		// Compute average and median inclusive/exclusive times.
-		for (auto& KV : Ctx.Map)
+		TUniquePtr<Trace::ITableReader<Trace::FTimingProfilerAggregatedStats>> AggregationResultTableReader(AggregationResultTable->CreateReader());
+		while (AggregationResultTableReader->IsValid())
 		{
-			PostProcessTimerStats(KV.Value, bComputeMedian);
-
-			// Update the CPU timer node.
-			FTimerNodePtr* TimerNodePtrPtr = TimerNodesIdMap.Find(KV.Key);
+			const Trace::FTimingProfilerAggregatedStats* Row = AggregationResultTableReader->GetCurrentRow();
+			FTimerNodePtr* TimerNodePtrPtr = TimerNodesIdMap.Find(Row->Timer->Id);
 			if (TimerNodePtrPtr != nullptr)
 			{
-				(*TimerNodePtrPtr)->SetAggregatedStats(KV.Value.BaseStats);
+				(*TimerNodePtrPtr)->SetAggregatedStats(*Row);
 			}
+			AggregationResultTableReader->NextRow();
 		}
 
-		StatsStartTime = Ctx.StartTime;
-		StatsEndTime = Ctx.EndTime;
+		StatsStartTime = StartTime;
+		StatsEndTime = EndTime;
 
 		UpdateTree();
 		TreeView->RebuildList();
@@ -1622,155 +1520,6 @@ void STimersView::UpdateStats(double StartTime, double EndTime)
 
 ////////////////////////////////////////////////////////////////////////////////////////////////////
 
-void STimersView::UpdateTotalMinMaxTimerStats(FTimerStatsEx& StatsEx, double InclTime, double ExclTime)
-{
-	FTimerAggregatedStats& Stats = StatsEx.BaseStats;
-
-	Stats.TotalInclusiveTime += InclTime;
-	if (InclTime < Stats.MinInclusiveTime)
-	{
-		Stats.MinInclusiveTime = InclTime;
-	}
-	if (InclTime > Stats.MaxInclusiveTime)
-	{
-		Stats.MaxInclusiveTime = InclTime;
-	}
-
-	Stats.TotalExclusiveTime += ExclTime;
-	if (ExclTime < Stats.MinExclusiveTime)
-	{
-		Stats.MinExclusiveTime = ExclTime;
-	}
-	if (ExclTime > Stats.MaxExclusiveTime)
-	{
-		Stats.MaxExclusiveTime = ExclTime;
-	}
-
-	Stats.InstanceCount++;
-}
-
-////////////////////////////////////////////////////////////////////////////////////////////////////
-
-void STimersView::PreComputeHistogram(FTimerStatsEx& StatsEx)
-{
-	const FTimerAggregatedStats& Stats = StatsEx.BaseStats;
-
-	// Each bucket (Histogram[i]) will be centered on a value.
-	// I.e. First bucket (bucket 0) is centered on Min value: [Min-DT/2, Min+DT/2)
-	// and last bucket (bucket N-1) is centered on Max value: [Max-DT/2, Max+DT/2).
-
-	constexpr double InvHistogramLen = 1.0 / static_cast<double>(FTimerStatsEx::HistogramLen - 1);
-
-	if (Stats.MaxInclusiveTime == Stats.MinInclusiveTime)
-	{
-		StatsEx.InclDT = 1.0; // single large bucket
-	}
-	else
-	{
-		StatsEx.InclDT = (Stats.MaxInclusiveTime - Stats.MinInclusiveTime) * InvHistogramLen;
-	}
-
-	if (Stats.MaxExclusiveTime == Stats.MinExclusiveTime)
-	{
-		StatsEx.ExclDT = 1.0; // single large bucket
-	}
-	else
-	{
-		StatsEx.ExclDT = (Stats.MaxExclusiveTime - Stats.MinExclusiveTime) * InvHistogramLen;
-	}
-}
-
-////////////////////////////////////////////////////////////////////////////////////////////////////
-
-void STimersView::UpdateHistogramForTimerStats(FTimerStatsEx& StatsEx, double InclTime, double ExclTime)
-{
-	const FTimerAggregatedStats& Stats = StatsEx.BaseStats;
-
-	int32 InclIndex = static_cast<int32>((InclTime - Stats.MinInclusiveTime + StatsEx.InclDT / 2) / StatsEx.InclDT);
-	ensure(InclIndex >= 0);
-	if (InclIndex < 0)
-	{
-		InclIndex = 0;
-	}
-	ensure(InclIndex < FTimerStatsEx::HistogramLen);
-	if (InclIndex >= FTimerStatsEx::HistogramLen)
-	{
-		InclIndex = FTimerStatsEx::HistogramLen - 1;
-	}
-	StatsEx.InclHistogram[InclIndex]++;
-
-	int32 ExclIndex = static_cast<int32>((ExclTime - Stats.MinExclusiveTime + StatsEx.ExclDT / 2) / StatsEx.ExclDT);
-	ensure(ExclIndex >= 0);
-	if (ExclIndex < 0)
-	{
-		ExclIndex = 0;
-	}
-	ensure(ExclIndex < FTimerStatsEx::HistogramLen);
-	if (ExclIndex >= FTimerStatsEx::HistogramLen)
-	{
-		ExclIndex = FTimerStatsEx::HistogramLen - 1;
-	}
-	StatsEx.ExclHistogram[ExclIndex]++;
-}
-
-////////////////////////////////////////////////////////////////////////////////////////////////////
-
-void STimersView::PostProcessTimerStats(FTimerStatsEx& StatsEx, bool bComputeMedian)
-{
-	FTimerAggregatedStats& Stats = StatsEx.BaseStats;
-
-	// Compute average inclusive/exclusive times.
-	ensure(Stats.InstanceCount > 0);
-	double InvCount = 1.0f / static_cast<double>(Stats.InstanceCount);
-	Stats.AverageInclusiveTime = Stats.TotalInclusiveTime * InvCount;
-	Stats.AverageExclusiveTime = Stats.TotalExclusiveTime * InvCount;
-
-	if (bComputeMedian)
-	{
-		const int32 HalfCount = Stats.InstanceCount / 2;
-
-		// Compute median inclusive time.
-		int32 InclCount = 0;
-		for (int32 HistogramIndex = 0; HistogramIndex < FTimerStatsEx::HistogramLen; HistogramIndex++)
-		{
-			InclCount += StatsEx.InclHistogram[HistogramIndex];
-			if (InclCount > HalfCount)
-			{
-				Stats.MedianInclusiveTime = Stats.MinInclusiveTime + HistogramIndex * StatsEx.InclDT;
-				if (HistogramIndex > 0 &&
-					Stats.InstanceCount % 2 == 0 &&
-					InclCount - StatsEx.InclHistogram[HistogramIndex] == HalfCount)
-				{
-					const double PrevMedian = Stats.MinInclusiveTime + (HistogramIndex - 1) * StatsEx.InclDT;
-					Stats.MedianInclusiveTime = (Stats.MedianInclusiveTime + PrevMedian) / 2;
-				}
-				break;
-			}
-		}
-
-		// Compute median exclusive time.
-		int32 ExclCount = 0;
-		for (int32 HistogramIndex = 0; HistogramIndex < FTimerStatsEx::HistogramLen; HistogramIndex++)
-		{
-			ExclCount += StatsEx.InclHistogram[HistogramIndex];
-			if (ExclCount > HalfCount)
-			{
-				Stats.MedianExclusiveTime = Stats.MinExclusiveTime + HistogramIndex * StatsEx.ExclDT;
-				if (HistogramIndex > 0 &&
-					Stats.InstanceCount % 2 == 0 &&
-					ExclCount - StatsEx.ExclHistogram[HistogramIndex] == HalfCount)
-				{
-					const double PrevMedian = Stats.MinExclusiveTime + (HistogramIndex - 1) * StatsEx.ExclDT;
-					Stats.MedianExclusiveTime = (Stats.MedianExclusiveTime + PrevMedian) / 2;
-				}
-				break;
-			}
-		}
-	}
-}
-
-////////////////////////////////////////////////////////////////////////////////////////////////////
-
 void STimersView::SelectTimerNode(uint64 Id)
 {
 	FTimerNodePtr* TimerNodePtrPtr = TimerNodesIdMap.Find(Id);
diff --git a/Engine/Source/Developer/TraceInsights/Private/Insights/Widgets/STimersView.h b/Engine/Source/Developer/TraceInsights/Private/Insights/Widgets/STimersView.h
--- a/Engine/Source/Developer/TraceInsights/Private/Insights/Widgets/STimersView.h
+++ b/Engine/Source/Developer/TraceInsights/Private/Insights/Widgets/STimersView.h
@@ -32,32 +32,6 @@ namespace Trace
 
 ////////////////////////////////////////////////////////////////////////////////////////////////////
 
-struct FTimerStatsEx
-{
-	static constexpr int32 HistogramLen = 100; // number of buckets per histogram
-
-	FTimerAggregatedStats BaseStats;
-
-	// Histogram for computing median inclusive time.
-	int32 InclHistogram[HistogramLen];
-	double InclDT; // bucket size
-
-	// Histogram for computing median exclusive time.
-	int32 ExclHistogram[HistogramLen];
-	double ExclDT; // bucket size
-
-	FTimerStatsEx()
-	{
-		FMemory::Memzero(InclHistogram, sizeof(int32) * HistogramLen);
-		InclDT = 1.0;
-
-		FMemory::Memzero(ExclHistogram, sizeof(int32) * HistogramLen);
-		ExclDT = 1.0;
-	}
-};
-
-////////////////////////////////////////////////////////////////////////////////////////////////////
-
 /** The filter collection - used for updating the list of timer nodes. */
 typedef TFilterCollection<const FTimerNodePtr&> FTimerNodeFilterCollection;
 
@@ -103,11 +77,6 @@ protected:
 
 	void UpdateTree();
 
-	static void UpdateTotalMinMaxTimerStats(FTimerStatsEx& Stats, double InclTime, double ExclTime);
-	static void PreComputeHistogram(FTimerStatsEx& StatsEx);
-	static void UpdateHistogramForTimerStats(FTimerStatsEx& StatsEx, double InclTime, double ExclTime);
-	static void PostProcessTimerStats(FTimerStatsEx& StatsEx, bool bComputeMedian);
-
 	/** Called when the analysis session has changed. */
 	void InsightsManager_OnSessionChanged();
 
diff --git a/Engine/Source/Developer/TraceInsights/Private/Insights/Widgets/STimersViewTooltip.cpp b/Engine/Source/Developer/TraceInsights/Private/Insights/Widgets/STimersViewTooltip.cpp
--- a/Engine/Source/Developer/TraceInsights/Private/Insights/Widgets/STimersViewTooltip.cpp
+++ b/Engine/Source/Developer/TraceInsights/Private/Insights/Widgets/STimersViewTooltip.cpp
@@ -67,7 +67,7 @@ TSharedPtr<SToolTip> STimersViewTooltip::GetTableCellTooltip(const TSharedPtr<FT
 	const float Alpha = 0.0f;//TimerNodePtr->_FramePct * 0.01f;
 	const FLinearColor ColorAndOpacity = FMath::Lerp(DefaultColor, ThreadColor,Alpha);
 
-	const FTimerAggregatedStats& Stats = TimerNodePtr->GetAggregatedStats();
+	const Trace::FAggregatedTimingStats& Stats = TimerNodePtr->GetAggregatedStats();
 
 	const int32 NumDigits = 5;
 
diff --git a/Engine/Source/Developer/TraceServices/Private/AnalysisService.cpp b/Engine/Source/Developer/TraceServices/Private/AnalysisService.cpp
--- a/Engine/Source/Developer/TraceServices/Private/AnalysisService.cpp
+++ b/Engine/Source/Developer/TraceServices/Private/AnalysisService.cpp
@@ -71,27 +71,18 @@ FAnalysisSession::FAnalysisSession(const TCHAR* SessionName)
 	: DurationSeconds(0.0)
 	, Allocator(32 << 20)
 	, StringStore(Allocator)
-	, ClassInfos(Allocator, 4096)
 	, BookmarkProvider(MakeShared<FBookmarkProvider>(Lock, StringStore))
 	, LogProvider(MakeShared<FLogProvider>(Allocator, Lock, StringStore))
 	, ThreadProvider(MakeShared<FThreadProvider>(Lock, StringStore))
 	, FramesProvider(MakeShared<FFrameProvider>(Allocator, Lock))
 	, TimingProfilerProvider(MakeShared<FTimingProfilerProvider>(Allocator, Lock, StringStore))
 	, FileActivityProvider(MakeShared<FFileActivityProvider>(Allocator, Lock))
-	, LoadTimeProfilerProvider(MakeShared<FLoadTimeProfilerProvider>(Allocator, Lock))
+	, LoadTimeProfilerProvider(MakeShared<FLoadTimeProfilerProvider>(Allocator, Lock, StringStore))
 	, CountersProvider(MakeShared<FCountersProvider>(Allocator, Lock))
 {
 
 }
 
-const Trace::FClassInfo& FAnalysisSession::AddClassInfo(const TCHAR* ClassName)
-{
-	Lock.WriteAccessCheck();
-	FClassInfo& ClassInfo = ClassInfos.PushBack();
-	ClassInfo.Name = ClassName;
-	return ClassInfo;
-}
-
 void FAnalysisSession::ReadBookmarkProvider(TFunctionRef<void(const IBookmarkProvider&)> Callback) const
 {
 	Lock.ReadAccessCheck();
diff --git a/Engine/Source/Developer/TraceServices/Private/AnalysisServicePrivate.h b/Engine/Source/Developer/TraceServices/Private/AnalysisServicePrivate.h
--- a/Engine/Source/Developer/TraceServices/Private/AnalysisServicePrivate.h
+++ b/Engine/Source/Developer/TraceServices/Private/AnalysisServicePrivate.h
@@ -53,8 +53,6 @@ public:
 	void UpdateDuration(double Seconds) { Lock.WriteAccessCheck(); DurationSeconds = Seconds > DurationSeconds ? Seconds : DurationSeconds; }
 	void SetComplete() { IsComplete = true; }
 
-	const FClassInfo& AddClassInfo(const TCHAR* Name);
-
 	TSharedRef<FBookmarkProvider> EditBookmarkProvider() { return BookmarkProvider; }
 	virtual void ReadBookmarkProvider(TFunctionRef<void(const IBookmarkProvider&)> Callback) const override;
 
@@ -95,7 +93,6 @@ private:
 	double DurationSeconds = 0.0;
 	FSlabAllocator Allocator;
 	FStringStore StringStore;
-	TPagedArray<FClassInfo> ClassInfos;
 	TSharedRef<FBookmarkProvider> BookmarkProvider;
 	TSharedRef<FLogProvider> LogProvider;
 	TSharedRef<FThreadProvider> ThreadProvider;
diff --git a/Engine/Source/Developer/TraceServices/Private/Analyzers/LoadTimeTraceAnalysis.cpp b/Engine/Source/Developer/TraceServices/Private/Analyzers/LoadTimeTraceAnalysis.cpp
--- a/Engine/Source/Developer/TraceServices/Private/Analyzers/LoadTimeTraceAnalysis.cpp
+++ b/Engine/Source/Developer/TraceServices/Private/Analyzers/LoadTimeTraceAnalysis.cpp
@@ -393,7 +393,7 @@ void FAsyncLoadingTraceAnalyzer::OnEvent(uint16 RouteId, const FOnEventContext&
 		TSharedRef<FAsyncPackageState> AsyncPackageState = ActiveAsyncPackagesMap[AsyncPackagePtr];
 		uint32 ThreadId = EventData.GetValue("ThreadId").As<uint32>();
 		TSharedRef<FThreadState> ThreadState = GetThreadState(ThreadId);
-		ELoadTimeProfilePackageEventType EventType = static_cast<ELoadTimeProfilePackageEventType>(EventData.GetValue("EventType").As<uint8>());
+		ELoadTimeProfilerPackageEventType EventType = static_cast<ELoadTimeProfilerPackageEventType>(EventData.GetValue("EventType").As<uint8>());
 		ThreadState->EnterPackageScope(Context.SessionContext.TimestampFromCycle(EventData.GetValue("Cycle").As<uint64>()), AsyncPackageState->PackageInfo, EventType);
 		break;
 	}
@@ -433,8 +433,10 @@ void FAsyncLoadingTraceAnalyzer::OnEvent(uint16 RouteId, const FOnEventContext&
 	}
 	case RouteId_ClassInfo:
 	{
+		Trace::FAnalysisSessionEditScope _(Session.Get());
+
 		uint64 ClassPtr = EventData.GetValue("Class").As<uint64>();
-		const Trace::FClassInfo& ClassInfo = Session->AddClassInfo(reinterpret_cast<const TCHAR*>(EventData.GetAttachment()));
+		const Trace::FClassInfo& ClassInfo = LoadTimeProfilerProvider->AddClassInfo(reinterpret_cast<const TCHAR*>(EventData.GetAttachment()));
 		ClassInfosMap.Add(ClassPtr, &ClassInfo);
 		break;
 	}
diff --git a/Engine/Source/Developer/TraceServices/Private/Analyzers/LoadTimeTraceAnalysis.h b/Engine/Source/Developer/TraceServices/Private/Analyzers/LoadTimeTraceAnalysis.h
--- a/Engine/Source/Developer/TraceServices/Private/Analyzers/LoadTimeTraceAnalysis.h
+++ b/Engine/Source/Developer/TraceServices/Private/Analyzers/LoadTimeTraceAnalysis.h
@@ -90,7 +90,7 @@ private:
 		
 		TSharedPtr<Trace::FLoadTimeProfilerProvider::CpuTimelineInternal> CpuTimeline;
 
-		void EnterPackageScope(double Time, const Trace::FPackageInfo* PackageInfo, ELoadTimeProfilePackageEventType EventType)
+		void EnterPackageScope(double Time, const Trace::FPackageInfo* PackageInfo, ELoadTimeProfilerPackageEventType EventType)
 		{
 			FScopeStackEntry& StackEntry = CpuScopeStack[CpuScopeStackDepth++];
 			StackEntry.Event.Export = CurrentEvent.Export;
diff --git a/Engine/Source/Developer/TraceServices/Private/Common/PagedArray.h b/Engine/Source/Developer/TraceServices/Private/Common/PagedArray.h
--- a/Engine/Source/Developer/TraceServices/Private/Common/PagedArray.h
+++ b/Engine/Source/Developer/TraceServices/Private/Common/PagedArray.h
@@ -59,7 +59,7 @@ public:
 
 		const PageType* GetCurrentPage()
 		{
-			return CurrentPage;
+			return Outer.FirstPage + CurrentPageIndex;
 		}
 
 		const ItemType* GetCurrentItem()
@@ -74,10 +74,8 @@ public:
 				return nullptr;
 			}
 			--CurrentPageIndex;
-			CurrentPage = Outer.FirstPage + CurrentPageIndex;
-			CurrentPageBegin = CurrentPage->Items - 1;
-			CurrentPageEnd = CurrentPage->Items + CurrentPage->Count;
-			CurrentItem = CurrentPageEnd - 1;
+			OnCurrentPageChanged();
+			CurrentItem = CurrentPageLastItem;
 			return GetCurrentPage();
 		}
 
@@ -88,30 +86,58 @@ public:
 				return nullptr;
 			}
 			++CurrentPageIndex;
-			CurrentPage = Outer.FirstPage + CurrentPageIndex;
-			CurrentPageBegin = CurrentPage->Items - 1;
-			CurrentPageEnd = CurrentPage->Items + CurrentPage->Count;
-			CurrentItem = CurrentPage->Items;
+			OnCurrentPageChanged();
+			CurrentItem = CurrentPageFirstItem;
 			return GetCurrentPage();
 		}
 
 		const ItemType* NextItem()
 		{
-			++CurrentItem;
-			if (CurrentItem >= CurrentPageEnd && !NextPage())
+			if (CurrentItem == CurrentPageLastItem)
 			{
-				return nullptr;
+				if (!NextPage())
+				{
+					return nullptr;
+				}
+				else
+				{
+					return CurrentItem;
+				}
 			}
+			++CurrentItem;
 			return CurrentItem;
 		}
 
 		const ItemType* PrevItem()
 		{
+			if (CurrentItem == CurrentPageFirstItem)
+			{
+				if (!PrevPage())
+				{
+					return nullptr;
+				}
+				else
+				{
+					return CurrentItem;
+				}
+			}
 			--CurrentItem;
-			if (CurrentItem <= CurrentPageBegin && !PrevPage())
+			return CurrentItem;
+		}
+
+		const ItemType* SetPosition(uint64 Index)
+		{
+			uint64 PageIndex = Index / Outer.PageSize;
+			uint64 ItemIndexInPage = Index % Outer.PageSize;
+			if (PageIndex != CurrentPageIndex)
 			{
-				return nullptr;
+				check(PageIndex < Outer.PagesArray.Num());
+				CurrentPageIndex = PageIndex;
+				OnCurrentPageChanged();
 			}
+			PageType* CurrentPage = Outer.FirstPage + CurrentPageIndex;
+			check(ItemIndexInPage < CurrentPage->Count);
+			CurrentItem = CurrentPage->Items + ItemIndexInPage;
 			return CurrentItem;
 		}
 
@@ -125,20 +151,32 @@ public:
 			{
 				check(InitialPageIndex < Outer.PagesArray.Num());
 				CurrentPageIndex = InitialPageIndex;
-				CurrentPage = Outer.FirstPage + InitialPageIndex;
-				CurrentPageBegin = CurrentPage->Items - 1;
-				CurrentPageEnd = CurrentPage->Items + CurrentPage->Count;
+				OnCurrentPageChanged();
+				PageType* CurrentPage = Outer.FirstPage + CurrentPageIndex;
 				check(InitialItemIndex < CurrentPage->Count);
 				CurrentItem = CurrentPage->Items + InitialItemIndex;
 			}
 		}
 
+		void OnCurrentPageChanged()
+		{
+			PageType* CurrentPage = Outer.FirstPage + CurrentPageIndex;
+			CurrentPageFirstItem = CurrentPage->Items;
+			if (CurrentPage->Items)
+			{
+				CurrentPageLastItem = CurrentPage->Items + CurrentPage->Count - 1;
+			}
+			else
+			{
+				CurrentPageLastItem = nullptr;
+			}
+		}
+
 		const TPagedArray& Outer;
-		const PageType* CurrentPage = nullptr;
 		uint64 CurrentPageIndex = 0;
 		const ItemType* CurrentItem = nullptr;
-		const ItemType* CurrentPageBegin = nullptr;
-		const ItemType* CurrentPageEnd = nullptr;
+		const ItemType* CurrentPageFirstItem = nullptr;
+		const ItemType* CurrentPageLastItem = nullptr;
 	};
 
 	TPagedArray(FSlabAllocator& InAllocator, uint64 InPageSize)
diff --git a/Engine/Source/Developer/TraceServices/Private/Common/TimelineStatistics.h b/Engine/Source/Developer/TraceServices/Private/Common/TimelineStatistics.h
new file mode 100644
--- /dev/null
+++ b/Engine/Source/Developer/TraceServices/Private/Common/TimelineStatistics.h
@@ -0,0 +1,267 @@
+// Copyright 1998-2019 Epic Games, Inc. All Rights Reserved.
+
+namespace Trace
+{
+
+class FTimelineStatistics
+{
+public:
+	template<typename TimelineType, typename BucketMappingFunc, typename BucketKeyType>
+	static void CreateAggregation(const TArray<const TimelineType*>& Timelines, BucketMappingFunc BucketMapper, double IntervalStart, double IntervalEnd, TMap<BucketKeyType, FAggregatedTimingStats>& Result)
+	{
+		TMap<BucketKeyType, FInternalAggregationEntry> InternalResult;
+		// Compute instance count and total/min/max inclusive/exclusive times for each timer.
+		for (const TimelineType* Timeline : Timelines)
+		{
+			ProcessTimeline(Timeline, BucketMapper, UpdateTotalMinMaxTimerStats, IntervalStart, IntervalEnd, InternalResult);
+		}
+		// Now, as we know min/max inclusive/exclusive times for each timer, we can compute histogram and median values.
+		const bool ComputeMedian = true;
+		if (ComputeMedian)
+		{
+			// Update bucket size (DT) for computing histogram.
+			for (auto& KV : InternalResult)
+			{
+				PreComputeHistogram(KV.Value);
+			}
+
+			// Compute histogram.
+			for (const TimelineType* Timeline : Timelines)
+			{
+				ProcessTimeline(Timeline, BucketMapper, UpdateHistogramForTimerStats, IntervalStart, IntervalEnd, InternalResult);
+			}
+		}
+
+		// Compute average and median inclusive/exclusive times.
+		for (auto& KV : InternalResult)
+		{
+			PostProcessTimerStats(KV.Value, ComputeMedian);
+			Result.Add(KV.Key, KV.Value.Stats);
+		}
+	}
+
+private:
+	enum
+	{
+		HistogramLen = 100
+	};
+
+	struct FInternalAggregationEntry
+	{
+		FInternalAggregationEntry()
+		{
+			FMemory::Memzero(InclHistogram, sizeof(int32) * HistogramLen);
+			InclDT = 1.0;
+
+			FMemory::Memzero(ExclHistogram, sizeof(int32) * HistogramLen);
+			ExclDT = 1.0;
+		}
+
+		// Histogram for computing median inclusive time.
+		int32 InclHistogram[HistogramLen];
+		double InclDT; // bucket size
+
+		// Histogram for computing median exclusive time.
+		int32 ExclHistogram[HistogramLen];
+		double ExclDT; // bucket size
+
+		FAggregatedTimingStats Stats;
+	};
+
+	template<typename TimelineType, typename BucketMappingFunc, typename BucketKeyType, typename CallbackType>
+	static void ProcessTimeline(const TimelineType* Timeline, BucketMappingFunc BucketMapper, CallbackType Callback, double IntervalStart, double IntervalEnd, TMap<BucketKeyType, FInternalAggregationEntry>& InternalResult)
+	{
+		struct FStackEntry
+		{
+			double StartTime;
+			double ExclusiveTime;
+			BucketKeyType BucketKey;
+		};
+
+		TArray<FStackEntry> Stack;
+		Stack.Reserve(1024);
+		double LastTime = 0.0;
+		Timeline->EnumerateEvents(IntervalStart, IntervalEnd, [BucketMapper, Callback, IntervalStart, IntervalEnd, &Stack, &LastTime, &InternalResult](bool IsEnter, double Time, const typename TimelineType::EventType& Event)
+		{
+			Time = FMath::Clamp(Time, IntervalStart, IntervalEnd);
+			BucketKeyType BucketKey = BucketMapper(Event);
+			if (Stack.Num())
+			{
+				FStackEntry& StackEntry = Stack.Top();
+				StackEntry.ExclusiveTime += Time - LastTime;
+			}
+			LastTime = Time;
+			if (IsEnter)
+			{
+				FStackEntry& StackEntry = Stack.AddDefaulted_GetRef();
+				StackEntry.StartTime = Time;
+				StackEntry.ExclusiveTime = 0.0;
+				StackEntry.BucketKey = BucketKey;
+				if (!InternalResult.Contains(BucketKey))
+				{
+					InternalResult.Add(BucketKey);
+				}
+			}
+			else
+			{
+				FStackEntry& StackEntry = Stack.Last();
+				double EventInclusiveTime = Time - StackEntry.StartTime;
+				check(EventInclusiveTime >= 0.0);
+				double EventExclusiveTime = StackEntry.ExclusiveTime;
+				check(EventExclusiveTime >= 0.0 && EventExclusiveTime <= EventInclusiveTime);
+				Stack.Pop(false);
+				double EventNonRecursiveInclusiveTime = EventInclusiveTime;
+				for (const FStackEntry& AncestorStackEntry : Stack)
+				{
+					if (AncestorStackEntry.BucketKey == BucketKey)
+					{
+						EventNonRecursiveInclusiveTime = 0.0;
+					}
+				}
+				Callback(InternalResult[BucketKey], EventNonRecursiveInclusiveTime, EventExclusiveTime);
+			}
+		});
+	}
+
+	static void UpdateTotalMinMaxTimerStats(FInternalAggregationEntry& AggregationEntry, double InclTime, double ExclTime)
+	{
+		FAggregatedTimingStats& Stats = AggregationEntry.Stats;
+		Stats.TotalInclusiveTime += InclTime;
+		if (InclTime < Stats.MinInclusiveTime)
+		{
+			Stats.MinInclusiveTime = InclTime;
+		}
+		if (InclTime > Stats.MaxInclusiveTime)
+		{
+			Stats.MaxInclusiveTime = InclTime;
+		}
+
+		Stats.TotalExclusiveTime += ExclTime;
+		if (ExclTime < Stats.MinExclusiveTime)
+		{
+			Stats.MinExclusiveTime = ExclTime;
+		}
+		if (ExclTime > Stats.MaxExclusiveTime)
+		{
+			Stats.MaxExclusiveTime = ExclTime;
+		}
+
+		Stats.InstanceCount++;
+	}
+
+	static void PreComputeHistogram(FInternalAggregationEntry& AggregationEntry)
+	{
+		const FAggregatedTimingStats& Stats = AggregationEntry.Stats;
+
+		// Each bucket (Histogram[i]) will be centered on a value.
+		// I.e. First bucket (bucket 0) is centered on Min value: [Min-DT/2, Min+DT/2)
+		// and last bucket (bucket N-1) is centered on Max value: [Max-DT/2, Max+DT/2).
+
+		constexpr double InvHistogramLen = 1.0 / static_cast<double>(HistogramLen - 1);
+
+		if (Stats.MaxInclusiveTime == Stats.MinInclusiveTime)
+		{
+			AggregationEntry.InclDT = 1.0; // single large bucket
+		}
+		else
+		{
+			AggregationEntry.InclDT = (Stats.MaxInclusiveTime - Stats.MinInclusiveTime) * InvHistogramLen;
+		}
+
+		if (Stats.MaxExclusiveTime == Stats.MinExclusiveTime)
+		{
+			AggregationEntry.ExclDT = 1.0; // single large bucket
+		}
+		else
+		{
+			AggregationEntry.ExclDT = (Stats.MaxExclusiveTime - Stats.MinExclusiveTime) * InvHistogramLen;
+		}
+	}
+
+	static void UpdateHistogramForTimerStats(FInternalAggregationEntry& AggregationEntry, double InclTime, double ExclTime)
+	{
+		const FAggregatedTimingStats& Stats = AggregationEntry.Stats;
+
+		int32 InclIndex = static_cast<int32>((InclTime - Stats.MinInclusiveTime + AggregationEntry.InclDT / 2) / AggregationEntry.InclDT);
+		ensure(InclIndex >= 0);
+		if (InclIndex < 0)
+		{
+			InclIndex = 0;
+		}
+		ensure(InclIndex < HistogramLen);
+		if (InclIndex >= HistogramLen)
+		{
+			InclIndex = HistogramLen - 1;
+		}
+		AggregationEntry.InclHistogram[InclIndex]++;
+
+		int32 ExclIndex = static_cast<int32>((ExclTime - Stats.MinExclusiveTime + AggregationEntry.ExclDT / 2) / AggregationEntry.ExclDT);
+		ensure(ExclIndex >= 0);
+		if (ExclIndex < 0)
+		{
+			ExclIndex = 0;
+		}
+		ensure(ExclIndex < HistogramLen);
+		if (ExclIndex >= HistogramLen)
+		{
+			ExclIndex = HistogramLen - 1;
+		}
+		AggregationEntry.ExclHistogram[ExclIndex]++;
+	}
+
+	static void PostProcessTimerStats(FInternalAggregationEntry& AggregationEntry, bool ComputeMedian)
+	{
+		FAggregatedTimingStats& Stats = AggregationEntry.Stats;
+
+		// Compute average inclusive/exclusive times.
+		ensure(Stats.InstanceCount > 0);
+		double InvCount = 1.0f / static_cast<double>(Stats.InstanceCount);
+		Stats.AverageInclusiveTime = Stats.TotalInclusiveTime * InvCount;
+		Stats.AverageExclusiveTime = Stats.TotalExclusiveTime * InvCount;
+
+		if (ComputeMedian)
+		{
+			const int32 HalfCount = Stats.InstanceCount / 2;
+
+			// Compute median inclusive time.
+			int32 InclCount = 0;
+			for (int32 HistogramIndex = 0; HistogramIndex < HistogramLen; HistogramIndex++)
+			{
+				InclCount += AggregationEntry.InclHistogram[HistogramIndex];
+				if (InclCount > HalfCount)
+				{
+					Stats.MedianInclusiveTime = Stats.MinInclusiveTime + HistogramIndex * AggregationEntry.InclDT;
+					if (HistogramIndex > 0 &&
+						Stats.InstanceCount % 2 == 0 &&
+						InclCount - AggregationEntry.InclHistogram[HistogramIndex] == HalfCount)
+					{
+						const double PrevMedian = Stats.MinInclusiveTime + (HistogramIndex - 1) * AggregationEntry.InclDT;
+						Stats.MedianInclusiveTime = (Stats.MedianInclusiveTime + PrevMedian) / 2;
+					}
+					break;
+				}
+			}
+
+			// Compute median exclusive time.
+			int32 ExclCount = 0;
+			for (int32 HistogramIndex = 0; HistogramIndex < HistogramLen; HistogramIndex++)
+			{
+				ExclCount += AggregationEntry.InclHistogram[HistogramIndex];
+				if (ExclCount > HalfCount)
+				{
+					Stats.MedianExclusiveTime = Stats.MinExclusiveTime + HistogramIndex * AggregationEntry.ExclDT;
+					if (HistogramIndex > 0 &&
+						Stats.InstanceCount % 2 == 0 &&
+						ExclCount - AggregationEntry.ExclHistogram[HistogramIndex] == HalfCount)
+					{
+						const double PrevMedian = Stats.MinExclusiveTime + (HistogramIndex - 1) * AggregationEntry.ExclDT;
+						Stats.MedianExclusiveTime = (Stats.MedianExclusiveTime + PrevMedian) / 2;
+					}
+					break;
+				}
+			}
+		}
+	}
+};
+
+}
\ No newline at end of file
diff --git a/Engine/Source/Developer/TraceServices/Private/Model/LoadTimeProfiler.cpp b/Engine/Source/Developer/TraceServices/Private/Model/LoadTimeProfiler.cpp
--- a/Engine/Source/Developer/TraceServices/Private/Model/LoadTimeProfiler.cpp
+++ b/Engine/Source/Developer/TraceServices/Private/Model/LoadTimeProfiler.cpp
@@ -2,13 +2,16 @@
 
 #include "Model/LoadTimeProfiler.h"
 #include "AnalysisServicePrivate.h"
+#include "Common/TimelineStatistics.h"
 
 namespace Trace
 {
 
-FLoadTimeProfilerProvider::FLoadTimeProfilerProvider(FSlabAllocator& InAllocator, FAnalysisSessionLock& InSessionLock)
+FLoadTimeProfilerProvider::FLoadTimeProfilerProvider(FSlabAllocator& InAllocator, FAnalysisSessionLock& InSessionLock, FStringStore& InStringStore)
 	: Allocator(InAllocator)
 	, SessionLock(InSessionLock)
+	, StringStore(InStringStore)
+	, ClassInfos(Allocator, 4096)
 	, Packages(Allocator, 4096)
 	, Exports(Allocator, 4090)
 	, MainThreadCpuTimeline(MakeShared<CpuTimelineInternal>(Allocator))
@@ -42,6 +45,110 @@ void FLoadTimeProfilerProvider::ReadAsyncLoadingThreadCpuTimeline(TFunctionRef<v
 	Callback(*AsyncLoadingThreadCpuTimeline);
 }
 
+ITable<FLoadTimeProfilerAggregatedStats>* FLoadTimeProfilerProvider::CreateEventAggregation(double IntervalStart, double IntervalEnd) const
+{
+	TArray<const CpuTimelineInternal*> Timelines;
+	Timelines.Add(&MainThreadCpuTimeline.Get());
+	Timelines.Add(&AsyncLoadingThreadCpuTimeline.Get());
+
+	auto BucketMapper = [](const FLoadTimeProfilerCpuEvent& Event)
+	{
+		return Event.PackageEventType;
+	};
+	TMap<ELoadTimeProfilerPackageEventType, FAggregatedTimingStats> Aggregation;
+	FTimelineStatistics::CreateAggregation(Timelines, BucketMapper, IntervalStart, IntervalEnd, Aggregation);
+	TTable<FAggregatedStatsTableLayout>* Table = new TTable<FAggregatedStatsTableLayout>();
+	for (const auto& KV : Aggregation)
+	{
+		FLoadTimeProfilerAggregatedStats& Row = Table->AddRow();
+		switch (KV.Key)
+		{
+		case LoadTimeProfilerPackageEventType_CreateLinker:
+			Row.Name = TEXT("CreateLinker");
+			break;
+		case LoadTimeProfilerPackageEventType_FinishLinker:
+			Row.Name = TEXT("FinishLinker");
+			break;
+		case LoadTimeProfilerPackageEventType_StartImportPackages:
+			Row.Name = TEXT("StartImportPackages");
+			break;
+		case LoadTimeProfilerPackageEventType_SetupImports:
+			Row.Name = TEXT("SetupImports");
+			break;
+		case LoadTimeProfilerPackageEventType_SetupExports:
+			Row.Name = TEXT("SetupExports");
+			break;
+		case LoadTimeProfilerPackageEventType_ProcessImportsAndExports:
+			Row.Name = TEXT("ProcessImportsAndExports");
+			break;
+		case LoadTimeProfilerPackageEventType_ExportsDone:
+			Row.Name = TEXT("ExportsDone");
+			break;
+		case LoadTimeProfilerPackageEventType_PostLoadWait:
+			Row.Name = TEXT("PostLoadWait");
+			break;
+		case LoadTimeProfilerPackageEventType_StartPostLoad:
+			Row.Name = TEXT("StartPostLoad");
+			break;
+		case LoadTimeProfilerPackageEventType_Tick:
+			Row.Name = TEXT("Tick");
+			break;
+		case LoadTimeProfilerPackageEventType_Finish:
+			Row.Name = TEXT("Finish");
+			break;
+		}
+		const FAggregatedTimingStats& Stats = KV.Value;
+		Row.Count = Stats.InstanceCount;
+		Row.Total = Stats.TotalExclusiveTime;
+		Row.Min = Stats.MinExclusiveTime;
+		Row.Max = Stats.MaxExclusiveTime;
+		Row.Average = Stats.AverageExclusiveTime;
+		Row.Median = Stats.MedianExclusiveTime;
+	}
+	return Table;
+}
+
+ITable<FLoadTimeProfilerAggregatedStats>* FLoadTimeProfilerProvider::CreateObjectTypeAggregation(double IntervalStart, double IntervalEnd) const
+{
+	TArray<const CpuTimelineInternal*> Timelines;
+	Timelines.Add(&MainThreadCpuTimeline.Get());
+	Timelines.Add(&AsyncLoadingThreadCpuTimeline.Get());
+
+	auto BucketMapper = [](const FLoadTimeProfilerCpuEvent& Event) -> const FClassInfo*
+	{
+		return Event.Export ? Event.Export->Class : nullptr;
+	};
+	TMap<const Trace::FClassInfo*, FAggregatedTimingStats> Aggregation;
+	FTimelineStatistics::CreateAggregation(Timelines, BucketMapper, IntervalStart, IntervalEnd, Aggregation);
+	TTable<FAggregatedStatsTableLayout>* Table = new TTable<FAggregatedStatsTableLayout>();
+	for (const auto& KV : Aggregation)
+	{
+		const FClassInfo* ClassInfo = KV.Key;
+		if (ClassInfo)
+		{
+			FLoadTimeProfilerAggregatedStats& Row = Table->AddRow();
+			const FAggregatedTimingStats& Stats = KV.Value;
+			Row.Name = ClassInfo->Name;
+			Row.Count = Stats.InstanceCount;
+			Row.Total = Stats.TotalExclusiveTime;
+			Row.Min = Stats.MinExclusiveTime;
+			Row.Max = Stats.MaxExclusiveTime;
+			Row.Average = Stats.AverageExclusiveTime;
+			Row.Median = Stats.MedianExclusiveTime;
+		}
+	}
+	return Table;
+}
+
+const Trace::FClassInfo& FLoadTimeProfilerProvider::AddClassInfo(const TCHAR* ClassName)
+{
+	SessionLock.WriteAccessCheck();
+
+	FClassInfo& ClassInfo = ClassInfos.PushBack();
+	ClassInfo.Name = StringStore.Store(ClassName);
+	return ClassInfo;
+}
+
 Trace::FPackageInfo& FLoadTimeProfilerProvider::CreatePackage(const TCHAR* PackageName)
 {
 	SessionLock.WriteAccessCheck();
@@ -49,7 +156,7 @@ Trace::FPackageInfo& FLoadTimeProfilerProvider::CreatePackage(const TCHAR* Packa
 	uint32 PackageId = Packages.Num();
 	FPackageInfo& Package = Packages.PushBack();
 	Package.Id = PackageId;
-	Package.Name = PackageName;
+	Package.Name = StringStore.Store(PackageName);
 	return Package;
 }
 
diff --git a/Engine/Source/Developer/TraceServices/Private/Model/LoadTimeProfiler.h b/Engine/Source/Developer/TraceServices/Private/Model/LoadTimeProfiler.h
--- a/Engine/Source/Developer/TraceServices/Private/Model/LoadTimeProfiler.h
+++ b/Engine/Source/Developer/TraceServices/Private/Model/LoadTimeProfiler.h
@@ -5,7 +5,9 @@
 #include "TraceServices/AnalysisService.h"
 #include "Common/SlabAllocator.h"
 #include "Common/PagedArray.h"
+#include "Common/StringStore.h"
 #include "Model/MonotonicTimeline.h"
+#include "Model/Tables.h"
 
 namespace Trace
 {
@@ -18,11 +20,14 @@ class FLoadTimeProfilerProvider
 public:
 	typedef TMonotonicTimeline<FLoadTimeProfilerCpuEvent> CpuTimelineInternal;
 
-	FLoadTimeProfilerProvider(FSlabAllocator& Allocator, FAnalysisSessionLock& InSessionLock);
+	FLoadTimeProfilerProvider(FSlabAllocator& Allocator, FAnalysisSessionLock& InSessionLock, FStringStore& StringStore);
 	virtual uint64 GetPackageCount() const override { return Packages.Num(); }
 	virtual void EnumeratePackages(TFunctionRef<void(const FPackageInfo&)> Callback) const override;
 	virtual void ReadMainThreadCpuTimeline(TFunctionRef<void(const CpuTimeline&)> Callback) const override;
 	virtual void ReadAsyncLoadingThreadCpuTimeline(TFunctionRef<void(const CpuTimeline&)> Callback) const override;
+	virtual ITable<FLoadTimeProfilerAggregatedStats>* CreateEventAggregation(double IntervalStart, double IntervalEnd) const override;
+	virtual ITable<FLoadTimeProfilerAggregatedStats>* CreateObjectTypeAggregation(double IntervalStart, double IntervalEnd) const override;
+	const FClassInfo& AddClassInfo(const TCHAR* ClassName);
 	FPackageInfo& CreatePackage(const TCHAR* PackageName);
 	FPackageExportInfo& CreateExport();
 	TSharedRef<CpuTimelineInternal> EditMainThreadCpuTimeline() { return MainThreadCpuTimeline; }
@@ -31,10 +36,22 @@ public:
 private:
 	FSlabAllocator& Allocator;
 	FAnalysisSessionLock& SessionLock;
+	FStringStore& StringStore;
+	TPagedArray<FClassInfo> ClassInfos;
 	TPagedArray<FPackageInfo> Packages;
 	TPagedArray<FPackageExportInfo> Exports;
 	TSharedRef<CpuTimelineInternal> MainThreadCpuTimeline;
 	TSharedRef<CpuTimelineInternal> AsyncLoadingThreadCpuTimeline;
+
+	UE_TRACE_TABLE_LAYOUT_BEGIN(FAggregatedStatsTableLayout, FLoadTimeProfilerAggregatedStats)
+		UE_TRACE_TABLE_COLUMN(Name, TEXT("Name"))
+		UE_TRACE_TABLE_COLUMN(Count, TEXT("Count"))
+		UE_TRACE_TABLE_COLUMN(Total, TEXT("Total"))
+		UE_TRACE_TABLE_COLUMN(Min, TEXT("Min"))
+		UE_TRACE_TABLE_COLUMN(Max, TEXT("Max"))
+		UE_TRACE_TABLE_COLUMN(Average, TEXT("Avg"))
+		UE_TRACE_TABLE_COLUMN(Median, TEXT("Med"))
+	UE_TRACE_TABLE_LAYOUT_END()
 };
 
 }
diff --git a/Engine/Source/Developer/TraceServices/Private/Model/Log.cpp b/Engine/Source/Developer/TraceServices/Private/Model/Log.cpp
--- a/Engine/Source/Developer/TraceServices/Private/Model/Log.cpp
+++ b/Engine/Source/Developer/TraceServices/Private/Model/Log.cpp
@@ -13,6 +13,8 @@ FLogProvider::FLogProvider(FSlabAllocator& InAllocator, FAnalysisSessionLock& In
 	, StringStore(InStringStore)
 	, Categories(Allocator, 128)
 	, MessageSpecs(Allocator, 1024)
+	, Messages(Allocator, 1024)
+	, MessagesTable(Messages)
 {
 
 }
@@ -51,7 +53,7 @@ void FLogProvider::AppendMessage(uint64 LogPoint, double Time, const uint8* Form
 {
 	SessionLock.WriteAccessCheck();
 	check(SpecMap.Contains(LogPoint));
-	FLogMessageInternal& InternalMessage = Messages.AddDefaulted_GetRef();
+	FLogMessageInternal& InternalMessage = Messages.PushBack();
 	InternalMessage.Time = Time;
 	InternalMessage.Spec = SpecMap[LogPoint];
 	FFormatArgsHelper::Format(FormatBuffer, FormatBufferSize - 1, InternalMessage.Spec->FormatString, FormatArgs);
diff --git a/Engine/Source/Developer/TraceServices/Private/Model/Log.h b/Engine/Source/Developer/TraceServices/Private/Model/Log.h
--- a/Engine/Source/Developer/TraceServices/Private/Model/Log.h
+++ b/Engine/Source/Developer/TraceServices/Private/Model/Log.h
@@ -5,6 +5,8 @@
 #include "TraceServices/AnalysisService.h"
 #include "Templates/SharedPointer.h"
 #include "Common/PagedArray.h"
+#include "Model/Tables.h"
+#include "Misc/OutputDeviceHelper.h"
 
 namespace Trace
 {
@@ -49,10 +51,20 @@ public:
 	virtual void EnumerateMessagesByIndex(uint64 Start, uint64 End, TFunctionRef<void(const FLogMessage&)> Callback) const override;
 	virtual uint64 GetCategoryCount() const override { return Categories.Num(); }
 	virtual void EnumerateCategories(TFunctionRef<void(const FLogCategory&)> Callback) const override;
+	virtual const IUntypedTable& GetMessagesTable() const override { return MessagesTable; }
 
 private:
 	void ConstructMessage(uint64 Id, TFunctionRef<void(const FLogMessage &)> Callback) const;
-	
+
+	UE_TRACE_TABLE_LAYOUT_BEGIN(FMessagesTableLayout, FLogMessageInternal)
+		UE_TRACE_TABLE_COLUMN(Time, TEXT("Time"))
+		UE_TRACE_TABLE_PROJECTED_COLUMN(TableColumnType_CString, TEXT("Verbosity"), [](const FLogMessageInternal& Message) { return FOutputDeviceHelper::VerbosityToString(Message.Spec->Verbosity); })
+		UE_TRACE_TABLE_PROJECTED_COLUMN(TableColumnType_CString, TEXT("Category"), [](const FLogMessageInternal& Message) { return Message.Spec->Category->Name; })
+		UE_TRACE_TABLE_PROJECTED_COLUMN(TableColumnType_CString, TEXT("File"), [](const FLogMessageInternal& Message) { return Message.Spec->File; })
+		UE_TRACE_TABLE_PROJECTED_COLUMN(TableColumnType_Int, TEXT("Line"), [](const FLogMessageInternal& Message) { return Message.Spec->Line; })
+		UE_TRACE_TABLE_COLUMN(Message, TEXT("Message"))
+	UE_TRACE_TABLE_LAYOUT_END()
+
 	enum
 	{
 		FormatBufferSize = 65536
@@ -65,8 +77,9 @@ private:
 	TMap<uint64, FLogMessageSpec*> SpecMap;
 	TPagedArray<FLogCategory> Categories;
 	TPagedArray<FLogMessageSpec> MessageSpecs;
-	TArray<FLogMessageInternal> Messages;
+	TPagedArray<FLogMessageInternal> Messages;
 	TCHAR FormatBuffer[FormatBufferSize];
+	TTableView<FMessagesTableLayout> MessagesTable;
 };
 
 }
\ No newline at end of file
diff --git a/Engine/Source/Developer/TraceServices/Private/Model/MonotonicTimeline.h b/Engine/Source/Developer/TraceServices/Private/Model/MonotonicTimeline.h
--- a/Engine/Source/Developer/TraceServices/Private/Model/MonotonicTimeline.h
+++ b/Engine/Source/Developer/TraceServices/Private/Model/MonotonicTimeline.h
@@ -25,11 +25,13 @@ struct FMonotonicTimelineDefaultSettings
 	}
 };
 
-template<typename EventType, typename SettingsType = FMonotonicTimelineDefaultSettings>
+template<typename InEventType, typename SettingsType = FMonotonicTimelineDefaultSettings>
 class TMonotonicTimeline
-	: public ITimeline<EventType>
+	: public ITimeline<InEventType>
 {
 public:
+	using EventType = InEventType;
+
 	TMonotonicTimeline(FSlabAllocator& InAllocator)
 		: Allocator(InAllocator)
 	{
diff --git a/Engine/Source/Developer/TraceServices/Private/Model/Tables.h b/Engine/Source/Developer/TraceServices/Private/Model/Tables.h
new file mode 100644
--- /dev/null
+++ b/Engine/Source/Developer/TraceServices/Private/Model/Tables.h
@@ -0,0 +1,537 @@
+// Copyright 1998-2019 Epic Games, Inc. All Rights Reserved.
+
+#pragma once
+
+#include "CoreTypes.h"
+#include "TraceServices/AnalysisService.h"
+#include "Templates/Function.h"
+#include "Common/PagedArray.h"
+
+#define UE_TRACE_TABLE_LAYOUT_BEGIN(Name, InRowType) \
+	class Name \
+		: public ::Trace::TTableLayoutBase<Name, InRowType> \
+	{ \
+	public: \
+		typedef InRowType RowType; \
+		virtual uint8 GetColumnCount() const override { return decltype(LastColumn)::Index; } \
+		template<uint8 IndexValue> \
+		constexpr ::Trace::ETableColumnType GetColumnTypeInternal() const { return TableColumnType_Invalid; } \
+		template<uint8 IndexValue> \
+		constexpr const TCHAR* GetColumnNameInternal() const { return nullptr; } \
+		template<uint8 IndexValue> \
+		::Trace::FColumnValueContainer GetColumnValueInternal(const RowType& Row) const { return 0; } \
+		::Trace::TColumnDeclaration<0, RowType
+		
+#define UE_TRACE_TABLE_COLUMN(Name, DisplayName) \
+		> PREPROCESSOR_JOIN(__Column__, __LINE__); \
+		template<> \
+		constexpr ::Trace::ETableColumnType GetColumnTypeInternal<decltype(PREPROCESSOR_JOIN(__Column__, __LINE__))::Index>() const { return ::Trace::GetColumnTypeFromNativeType<decltype(GetMemberTypeHelper(&RowType::Name))>(); } \
+		template<> \
+		constexpr const TCHAR* GetColumnNameInternal<decltype(PREPROCESSOR_JOIN(__Column__, __LINE__))::Index>() const { return DisplayName; } \
+		template<> \
+		::Trace::FColumnValueContainer GetColumnValueInternal<decltype(PREPROCESSOR_JOIN(__Column__, __LINE__))::Index>(const RowType& Row) const { return Row.Name; } \
+		::Trace::TColumnDeclaration<decltype(PREPROCESSOR_JOIN(__Column__, __LINE__))::Index + 1, RowType
+		
+#define UE_TRACE_TABLE_PROJECTED_COLUMN(ColumnType, DisplayName, ProjectionFunc) \
+		> PREPROCESSOR_JOIN(__Column__, __LINE__); \
+		TFunction<FNativeTypeFromColumnType<::Trace::ColumnType>::Type(const RowType&)> PREPROCESSOR_JOIN(__ColumnProjector__, __LINE__) = ProjectionFunc; \
+		template<> \
+		constexpr ::Trace::ETableColumnType GetColumnTypeInternal<decltype(PREPROCESSOR_JOIN(__Column__, __LINE__))::Index>() const { return ::Trace::ColumnType; } \
+		template<> \
+		constexpr const TCHAR* GetColumnNameInternal<decltype(PREPROCESSOR_JOIN(__Column__, __LINE__))::Index>() const { return DisplayName; } \
+		template<> \
+		::Trace::FColumnValueContainer GetColumnValueInternal<decltype(PREPROCESSOR_JOIN(__Column__, __LINE__))::Index>(const RowType& Row) const { return PREPROCESSOR_JOIN(__ColumnProjector__, __LINE__)(Row); } \
+		::Trace::TColumnDeclaration<decltype(PREPROCESSOR_JOIN(__Column__, __LINE__))::Index + 1, RowType
+
+#define UE_TRACE_TABLE_LAYOUT_END() \
+		> LastColumn;\
+	};
+
+namespace Trace
+{
+
+template<typename Type>
+static constexpr ETableColumnType GetColumnTypeFromNativeType();
+
+template<>
+constexpr ETableColumnType GetColumnTypeFromNativeType<bool>()
+{
+	return TableColumnType_Bool;
+}
+
+template<>
+constexpr ETableColumnType GetColumnTypeFromNativeType<int8>()
+{
+	return TableColumnType_Int;
+};
+
+template<>
+constexpr ETableColumnType GetColumnTypeFromNativeType<uint8>()
+{
+	return TableColumnType_Int;
+};
+
+template<>
+constexpr ETableColumnType GetColumnTypeFromNativeType<int16>()
+{
+	return TableColumnType_Int;
+};
+
+template<>
+constexpr ETableColumnType GetColumnTypeFromNativeType<uint16>()
+{
+	return TableColumnType_Int;
+};
+
+template<>
+constexpr ETableColumnType GetColumnTypeFromNativeType<int32>()
+{
+	return TableColumnType_Int;
+};
+
+template<>
+constexpr ETableColumnType GetColumnTypeFromNativeType<uint32>()
+{
+	return TableColumnType_Int;
+};
+
+template<>
+constexpr ETableColumnType GetColumnTypeFromNativeType<int64>()
+{
+	return TableColumnType_Int;
+};
+
+template<>
+constexpr ETableColumnType GetColumnTypeFromNativeType<uint64>()
+{
+	return TableColumnType_Int;
+};
+
+template<>
+constexpr ETableColumnType GetColumnTypeFromNativeType<float>()
+{
+	return TableColumnType_Float;
+};
+
+template<>
+constexpr ETableColumnType GetColumnTypeFromNativeType<double>()
+{
+	return TableColumnType_Double;
+};
+
+template<>
+constexpr ETableColumnType GetColumnTypeFromNativeType<const TCHAR*>()
+{
+	return TableColumnType_CString;
+}
+
+template<ETableColumnType ColumnType>
+struct FNativeTypeFromColumnType
+{
+	
+};
+
+template<>
+struct FNativeTypeFromColumnType<TableColumnType_Bool>
+{
+	using Type = bool;
+};
+
+template<>
+struct FNativeTypeFromColumnType<TableColumnType_Int>
+{
+	using Type = int64;
+};
+
+template<>
+struct FNativeTypeFromColumnType<TableColumnType_Float>
+{
+	using Type = float;
+};
+
+template<>
+struct FNativeTypeFromColumnType<TableColumnType_Double>
+{
+	using Type = double;
+};
+
+template<>
+struct FNativeTypeFromColumnType<TableColumnType_CString>
+{
+	using Type = const TCHAR*;
+};
+
+template<typename MemberType, typename ObjectType>
+MemberType GetMemberTypeHelper(MemberType ObjectType::*);
+
+struct FColumnValueContainer
+{
+	FColumnValueContainer(bool Value)
+	{
+		BoolValue = Value;
+	}
+
+	FColumnValueContainer(int8 Value)
+	{
+		IntValue = Value;
+	}
+
+	FColumnValueContainer(uint8 Value)
+	{
+		IntValue = Value;
+	}
+
+	FColumnValueContainer(int32 Value)
+	{
+		IntValue = Value;
+	}
+
+	FColumnValueContainer(uint32 Value)
+	{
+		IntValue = Value;
+	}
+
+	FColumnValueContainer(int64 Value)
+	{
+		IntValue = Value;
+	}
+
+	FColumnValueContainer(uint64 Value)
+	{
+		IntValue = int64(Value);
+	}
+
+	FColumnValueContainer(float Value)
+	{
+		FloatValue = Value;
+	}
+
+	FColumnValueContainer(double Value)
+	{
+		DoubleValue = Value;
+	}
+
+	FColumnValueContainer(const TCHAR* Value)
+	{
+		CStringValue = Value;
+	}
+
+	union
+	{
+		bool BoolValue;
+		int64 IntValue;
+		float FloatValue;
+		double DoubleValue;
+		const TCHAR* CStringValue;
+	};
+};
+
+template<uint8 IndexValue, typename RowType>
+struct TColumnDeclaration
+{
+	enum
+	{
+		Index = IndexValue
+	};
+
+	static constexpr ETableColumnType GetType() { return ::Trace::TableColumnType_Invalid; }
+	static constexpr const TCHAR* GetName() { return nullptr; }
+	static const void* GetValue(const RowType& Row) { return nullptr; }
+};
+
+template<typename LayoutType>
+class TTableReader
+	: public ITableReader<typename LayoutType::RowType>
+{
+public:
+	TTableReader(const LayoutType& InLayout, const TPagedArray<typename LayoutType::RowType>& InRows)
+		: Layout(InLayout)
+		, Iterator(InRows.GetIteratorFromItem(0))
+	{
+		CurrentRow = Iterator.GetCurrentItem();
+	}
+
+	virtual bool IsValid() const override
+	{
+		return CurrentRow != nullptr;
+	}
+
+	virtual void NextRow() override
+	{
+		CurrentRow = Iterator.NextItem();
+	}
+
+	virtual void SetRowIndex(uint64 RowIndex) override
+	{
+		CurrentRow = Iterator.SetPosition(RowIndex);
+	}
+
+	virtual const typename LayoutType::RowType* GetCurrentRow() const override
+	{
+		return CurrentRow;
+	}
+
+	virtual bool GetValueBool(uint8 ColumnIndex) const override
+	{
+		if (!CurrentRow)
+		{
+			return false;
+		}
+		ETableColumnType ColumnType = Layout.GetColumnTypeConstExpr(ColumnIndex);
+		switch (ColumnType)
+		{
+		case TableColumnType_Bool:
+			return Layout.GetColumnValue(*CurrentRow, ColumnIndex).BoolValue;
+		case TableColumnType_Int:
+			return static_cast<bool>(Layout.GetColumnValue(*CurrentRow, ColumnIndex).IntValue);
+		case TableColumnType_Float:
+			return static_cast<bool>(Layout.GetColumnValue(*CurrentRow, ColumnIndex).FloatValue);
+		case TableColumnType_Double:
+			return static_cast<bool>(Layout.GetColumnValue(*CurrentRow, ColumnIndex).DoubleValue);
+		}
+		return false;
+	}
+
+	virtual int64 GetValueInt(uint8 ColumnIndex) const override
+	{
+		if (!CurrentRow)
+		{
+			return 0;
+		}
+		ETableColumnType ColumnType = Layout.GetColumnTypeConstExpr(ColumnIndex);
+		switch (ColumnType)
+		{
+		case TableColumnType_Bool:
+			return static_cast<int64>(Layout.GetColumnValue(*CurrentRow, ColumnIndex).BoolValue);
+		case TableColumnType_Int:
+			return Layout.GetColumnValue(*CurrentRow, ColumnIndex).IntValue;
+		case TableColumnType_Float:
+			return static_cast<int64>(Layout.GetColumnValue(*CurrentRow, ColumnIndex).FloatValue);
+		case TableColumnType_Double:
+			return static_cast<int64>(Layout.GetColumnValue(*CurrentRow, ColumnIndex).DoubleValue);
+		}
+		return 0;
+	}
+
+	virtual float GetValueFloat(uint8 ColumnIndex) const override
+	{
+		if (!CurrentRow)
+		{
+			return 0.0;
+		}
+		ETableColumnType ColumnType = Layout.GetColumnTypeConstExpr(ColumnIndex);
+		switch (ColumnType)
+		{
+		case TableColumnType_Bool:
+			return static_cast<float>(Layout.GetColumnValue(*CurrentRow, ColumnIndex).BoolValue);
+		case TableColumnType_Int:
+			return static_cast<float>(Layout.GetColumnValue(*CurrentRow, ColumnIndex).IntValue);
+		case TableColumnType_Float:
+			return Layout.GetColumnValue(*CurrentRow, ColumnIndex).FloatValue;
+		case TableColumnType_Double:
+			return static_cast<double>(Layout.GetColumnValue(*CurrentRow, ColumnIndex).DoubleValue);
+		}
+		return 0.0;
+	}
+
+	virtual double GetValueDouble(uint8 ColumnIndex) const override
+	{
+		if (!CurrentRow)
+		{
+			return 0.0;
+		}
+		ETableColumnType ColumnType = Layout.GetColumnTypeConstExpr(ColumnIndex);
+		switch (ColumnType)
+		{
+		case TableColumnType_Bool:
+			return static_cast<double>(Layout.GetColumnValue(*CurrentRow, ColumnIndex).BoolValue);
+		case TableColumnType_Int:
+			return static_cast<double>(Layout.GetColumnValue(*CurrentRow, ColumnIndex).IntValue);
+		case TableColumnType_Float:
+			return static_cast<double>(Layout.GetColumnValue(*CurrentRow, ColumnIndex).FloatValue);
+		case TableColumnType_Double:
+			return Layout.GetColumnValue(*CurrentRow, ColumnIndex).DoubleValue;
+		}
+		return 0.0;
+	}
+
+	virtual const TCHAR* GetValueCString(uint8 ColumnIndex) const override
+	{
+		if (!CurrentRow)
+		{
+			return nullptr;
+		}
+		ETableColumnType ColumnType = Layout.GetColumnTypeConstExpr(ColumnIndex);
+		if (ColumnType == TableColumnType_CString)
+		{
+			return Layout.GetColumnValue(*CurrentRow, ColumnIndex).CStringValue;
+		}
+		return nullptr;
+	}
+
+private:
+	const LayoutType& Layout;
+	typename TPagedArray<typename LayoutType::RowType>::FIterator Iterator;
+	const typename LayoutType::RowType* CurrentRow;
+};
+
+template<typename LayoutType, typename RowType>
+class TTableLayoutBase
+	: public ITableLayout
+{
+public:
+	constexpr ETableColumnType GetColumnTypeConstExpr(uint8 ColumnIndex) const
+	{
+		const LayoutType* This = static_cast<const LayoutType*>(this);
+		switch (ColumnIndex)
+		{
+		case 0: return This->template GetColumnTypeInternal<0>();
+		case 1: return This->template GetColumnTypeInternal<1>();
+		case 2: return This->template GetColumnTypeInternal<2>();
+		case 3: return This->template GetColumnTypeInternal<3>();
+		case 4: return This->template GetColumnTypeInternal<4>();
+		case 5: return This->template GetColumnTypeInternal<5>();
+		case 6: return This->template GetColumnTypeInternal<6>();
+		case 7: return This->template GetColumnTypeInternal<7>();
+		case 8: return This->template GetColumnTypeInternal<8>();
+		case 9: return This->template GetColumnTypeInternal<9>();
+		case 10: return This->template GetColumnTypeInternal<10>();
+		case 11: return This->template GetColumnTypeInternal<11>();
+		case 12: return This->template GetColumnTypeInternal<12>();
+		case 13: return This->template GetColumnTypeInternal<13>();
+		case 14: return This->template GetColumnTypeInternal<14>();
+		case 15: return This->template GetColumnTypeInternal<15>();
+		}
+		return ::Trace::TableColumnType_Invalid;
+	}
+
+	virtual ETableColumnType GetColumnType(uint8 ColumnIndex) const override
+	{
+		return GetColumnTypeConstExpr(ColumnIndex);
+	}
+
+	constexpr const TCHAR* GetColumnNameConstExpr(uint8 ColumnIndex) const
+	{
+		const LayoutType* This = static_cast<const LayoutType*>(this);
+		switch (ColumnIndex)
+		{
+		case 0: return This->template GetColumnNameInternal<0>();
+		case 1: return This->template GetColumnNameInternal<1>();
+		case 2: return This->template GetColumnNameInternal<2>();
+		case 3: return This->template GetColumnNameInternal<3>();
+		case 4: return This->template GetColumnNameInternal<4>();
+		case 5: return This->template GetColumnNameInternal<5>();
+		case 6: return This->template GetColumnNameInternal<6>();
+		case 7: return This->template GetColumnNameInternal<7>();
+		case 8: return This->template GetColumnNameInternal<8>();
+		case 9: return This->template GetColumnNameInternal<9>();
+		case 10: return This->template GetColumnNameInternal<10>();
+		case 11: return This->template GetColumnNameInternal<11>();
+		case 12: return This->template GetColumnNameInternal<12>();
+		case 13: return This->template GetColumnNameInternal<13>();
+		case 14: return This->template GetColumnNameInternal<14>();
+		case 15: return This->template GetColumnNameInternal<15>();
+		}
+		return nullptr;
+	}
+
+	virtual const TCHAR* GetColumnName(uint8 ColumnIndex) const override
+	{
+		return GetColumnNameConstExpr(ColumnIndex);
+	}
+
+	FColumnValueContainer GetColumnValue(const RowType& Row, uint8 ColumnIndex) const
+	{
+		const LayoutType* This = static_cast<const LayoutType*>(this);
+		switch (ColumnIndex)
+		{
+		case 0: return This->template GetColumnValueInternal<0>(Row);
+		case 1: return This->template GetColumnValueInternal<1>(Row);
+		case 2: return This->template GetColumnValueInternal<2>(Row);
+		case 3: return This->template GetColumnValueInternal<3>(Row);
+		case 4: return This->template GetColumnValueInternal<4>(Row);
+		case 5: return This->template GetColumnValueInternal<5>(Row);
+		case 6: return This->template GetColumnValueInternal<6>(Row);
+		case 7: return This->template GetColumnValueInternal<7>(Row);
+		case 8: return This->template GetColumnValueInternal<8>(Row);
+		case 9: return This->template GetColumnValueInternal<9>(Row);
+		case 10: return This->template GetColumnValueInternal<10>(Row);
+		case 11: return This->template GetColumnValueInternal<11>(Row);
+		case 12: return This->template GetColumnValueInternal<12>(Row);
+		case 13: return This->template GetColumnValueInternal<13>(Row);
+		case 14: return This->template GetColumnValueInternal<14>(Row);
+		case 15: return This->template GetColumnValueInternal<15>(Row);
+		}
+		return nullptr;
+	}
+};
+
+template<typename LayoutType>
+class TTableView
+	: public ITable<typename LayoutType::RowType>
+{
+public:
+	TTableView(const TPagedArray<typename LayoutType::RowType>& InRows)
+		: Rows(InRows)
+	{
+	}
+
+	virtual const ITableLayout& GetLayout() const override
+	{
+		return Layout;
+	}
+
+	virtual uint64 GetRowCount() const override
+	{
+		return Rows.Num();
+	}
+
+	virtual ITableReader<typename LayoutType::RowType>* CreateReader() const override
+	{
+		return new TTableReader<LayoutType>(Layout, Rows);
+	}
+
+private:
+	LayoutType Layout;
+	const TPagedArray<typename LayoutType::RowType>& Rows;
+};
+
+template<typename LayoutType>
+class TTable
+	: public ITable<typename LayoutType::RowType>
+{
+public:
+	TTable()
+		: Allocator(2 << 20)
+		, Rows(Allocator, 1024)
+	{
+
+	}
+
+	virtual const ITableLayout& GetLayout() const override
+	{
+		return Layout;
+	}
+
+	virtual uint64 GetRowCount() const override
+	{
+		return Rows.Num();
+	}
+
+	virtual ITableReader<typename LayoutType::RowType>* CreateReader() const override
+	{
+		return new TTableReader<LayoutType>(Layout, Rows);
+	}
+
+	typename LayoutType::RowType& AddRow()
+	{
+		return Rows.PushBack();
+	}
+
+private:
+	LayoutType Layout;
+	FSlabAllocator Allocator;
+	TPagedArray<typename LayoutType::RowType> Rows;
+};
+
+}
\ No newline at end of file
diff --git a/Engine/Source/Developer/TraceServices/Private/Model/TimingProfiler.cpp b/Engine/Source/Developer/TraceServices/Private/Model/TimingProfiler.cpp
--- a/Engine/Source/Developer/TraceServices/Private/Model/TimingProfiler.cpp
+++ b/Engine/Source/Developer/TraceServices/Private/Model/TimingProfiler.cpp
@@ -3,6 +3,7 @@
 #include "Model/TimingProfiler.h"
 #include "AnalysisServicePrivate.h"
 #include "Common/StringStore.h"
+#include "Common/TimelineStatistics.h"
 
 namespace Trace
 {
@@ -128,4 +129,38 @@ void FTimingProfilerProvider::ReadTimers(TFunctionRef<void(const FTimingProfiler
 	Callback(Timers.GetData(), Timers.Num());
 }
 
-}
\ No newline at end of file
+ITable<FTimingProfilerAggregatedStats>* FTimingProfilerProvider::CreateAggregation(double IntervalStart, double IntervalEnd, TFunctionRef<bool(uint32)> CpuThreadFilter, bool IncludeGpu) const
+{
+	SessionLock.ReadAccessCheck();
+
+	TArray<const TimelineInternal*> IncludedTimelines;
+	if (IncludeGpu)
+	{
+		IncludedTimelines.Add(&Timelines[GpuTimelineIndex].Get());
+	}
+	for (const auto& KV : CpuThreadTimelineIndexMap)
+	{
+		if (CpuThreadFilter(KV.Key))
+		{
+			IncludedTimelines.Add(&Timelines[KV.Value].Get());
+		}
+	}
+
+	auto BucketMappingFunc = [this](const TimelineInternal::EventType& Event) -> const FTimingProfilerTimer*
+	{
+		return &Timers[Event.TimerIndex];
+	};
+
+	TMap<const FTimingProfilerTimer*, FAggregatedTimingStats> Aggregation;
+	FTimelineStatistics::CreateAggregation(IncludedTimelines, BucketMappingFunc, IntervalStart, IntervalEnd, Aggregation);
+	TTable<FAggregatedStatsTableLayout>* Table = new TTable<FAggregatedStatsTableLayout>();
+	for (const auto& KV : Aggregation)
+	{
+		FTimingProfilerAggregatedStats& Row = Table->AddRow();
+		*static_cast<FAggregatedTimingStats*>(&Row) = KV.Value;
+		Row.Timer = KV.Key;
+	}
+	return Table;
+}
+
+}
diff --git a/Engine/Source/Developer/TraceServices/Private/Model/TimingProfiler.h b/Engine/Source/Developer/TraceServices/Private/Model/TimingProfiler.h
--- a/Engine/Source/Developer/TraceServices/Private/Model/TimingProfiler.h
+++ b/Engine/Source/Developer/TraceServices/Private/Model/TimingProfiler.h
@@ -5,6 +5,7 @@
 #include "TraceServices/AnalysisService.h"
 #include "Common/SlabAllocator.h"
 #include "Model/MonotonicTimeline.h"
+#include "Model/Tables.h"
 
 namespace Trace
 {
@@ -30,6 +31,7 @@ public:
 	virtual uint64 GetTimelineCount() const override { return Timelines.Num(); }
 	virtual void EnumerateTimelines(TFunctionRef<void(const Timeline&)> Callback) const override;
 	virtual void ReadTimers(TFunctionRef<void(const FTimingProfilerTimer*, uint64)> Callback) const override;
+	virtual ITable<FTimingProfilerAggregatedStats>* CreateAggregation(double IntervalStart, double IntervalEnd, TFunctionRef<bool(uint32)> CpuThreadFilter, bool IncludeGpu) const override;
 
 private:
 	FTimingProfilerTimer& AddTimerInternal(const TCHAR* Name, bool IsGpuEvent);
@@ -41,6 +43,21 @@ private:
 	TArray<TSharedRef<TimelineInternal>> Timelines;
 	TMap<uint32, uint32> CpuThreadTimelineIndexMap;
 	uint32 GpuTimelineIndex = 0;
+
+	UE_TRACE_TABLE_LAYOUT_BEGIN(FAggregatedStatsTableLayout, FTimingProfilerAggregatedStats)
+		UE_TRACE_TABLE_PROJECTED_COLUMN(TableColumnType_CString, TEXT("Name"), [](const FTimingProfilerAggregatedStats& Row) { return Row.Timer->Name; })
+		UE_TRACE_TABLE_COLUMN(InstanceCount, TEXT("Count"))
+		UE_TRACE_TABLE_COLUMN(TotalInclusiveTime, TEXT("Incl"))
+		UE_TRACE_TABLE_COLUMN(MinInclusiveTime, TEXT("I.Min"))
+		UE_TRACE_TABLE_COLUMN(MaxInclusiveTime, TEXT("I.Max"))
+		UE_TRACE_TABLE_COLUMN(AverageInclusiveTime, TEXT("I.Avg"))
+		UE_TRACE_TABLE_COLUMN(MedianInclusiveTime, TEXT("I.Med"))
+		UE_TRACE_TABLE_COLUMN(TotalExclusiveTime, TEXT("Excl"))
+		UE_TRACE_TABLE_COLUMN(MinExclusiveTime, TEXT("E.Min"))
+		UE_TRACE_TABLE_COLUMN(MaxExclusiveTime, TEXT("E.Max"))
+		UE_TRACE_TABLE_COLUMN(AverageExclusiveTime, TEXT("E.Avg"))
+		UE_TRACE_TABLE_COLUMN(MedianExclusiveTime, TEXT("E.Med"))
+	UE_TRACE_TABLE_LAYOUT_END()
 };
 
 }
\ No newline at end of file
diff --git a/Engine/Source/Developer/TraceServices/Public/TraceServices/AnalysisService.h b/Engine/Source/Developer/TraceServices/Public/TraceServices/AnalysisService.h
--- a/Engine/Source/Developer/TraceServices/Public/TraceServices/AnalysisService.h
+++ b/Engine/Source/Developer/TraceServices/Public/TraceServices/AnalysisService.h
@@ -15,6 +15,65 @@
 namespace Trace
 {
 
+enum ETableColumnType
+{
+	TableColumnType_Invalid,
+	TableColumnType_Bool,
+	TableColumnType_Int,
+	TableColumnType_Float,
+	TableColumnType_Double,
+	TableColumnType_CString,
+};
+
+class ITableLayout
+{
+public:
+	virtual ~ITableLayout() = default;
+	virtual uint8 GetColumnCount() const = 0;
+	virtual const TCHAR* GetColumnName(uint8 ColumnIndex) const = 0;
+	virtual ETableColumnType GetColumnType(uint8 ColumnIndex) const = 0;
+};
+
+class IUntypedTableReader
+{
+public:
+	virtual ~IUntypedTableReader() = default;
+	virtual bool IsValid() const = 0;
+	virtual void NextRow() = 0;
+	virtual void SetRowIndex(uint64 RowIndex) = 0;
+	virtual bool GetValueBool(uint8 ColumnIndex) const = 0;
+	virtual int64 GetValueInt(uint8 ColumnIndex) const = 0;
+	virtual float GetValueFloat(uint8 ColumnIndex) const = 0;
+	virtual double GetValueDouble(uint8 ColumnIndex) const = 0;
+	virtual const TCHAR* GetValueCString(uint8 ColumnIndex) const = 0;
+};
+
+template<typename RowType>
+class ITableReader
+	: public IUntypedTableReader
+{
+public:
+	virtual const RowType* GetCurrentRow() const = 0;
+};
+
+class IUntypedTable
+{
+public:
+	virtual ~IUntypedTable() = default;
+	virtual const ITableLayout& GetLayout() const = 0;
+	virtual uint64 GetRowCount() const = 0;
+	virtual IUntypedTableReader* CreateReader() const = 0;
+};
+
+template<typename RowType>
+class ITable
+	: public IUntypedTable
+{
+public:
+	virtual ~ITable() = default;
+	virtual ITableReader<RowType>* CreateReader() const = 0;
+};
+
 struct FBookmark
 {
 	double Time;
@@ -56,6 +115,7 @@ public:
 	virtual bool ReadMessage(uint64 Index, TFunctionRef<void(const FLogMessage&)> Callback) const = 0;
 	virtual uint64 GetCategoryCount() const = 0;
 	virtual void EnumerateCategories(TFunctionRef<void(const FLogCategory&)> Callback) const = 0;
+	virtual const IUntypedTable& GetMessagesTable() const = 0;
 };
 
 struct FThreadInfo
@@ -73,10 +133,12 @@ public:
 	virtual void EnumerateThreads(TFunctionRef<void(const FThreadInfo&)> Callback) const = 0;
 };
 
-template<typename EventType>
+template<typename InEventType>
 class ITimeline
 {
 public:
+	typedef InEventType EventType;
+
 	virtual ~ITimeline() = default;
 	virtual uint64 GetModCount() const = 0;
 	virtual uint64 GetEventCount() const = 0;
@@ -99,6 +161,27 @@ struct FTimingProfilerEvent
 	uint32 TimerIndex;
 };
 
+struct FAggregatedTimingStats
+{
+	uint64 InstanceCount = 0;
+	double TotalInclusiveTime = 0.0;
+	double MinInclusiveTime = DBL_MAX;
+	double MaxInclusiveTime = -DBL_MAX;
+	double AverageInclusiveTime = 0.0;
+	double MedianInclusiveTime = 0.0;
+	double TotalExclusiveTime = 0.0;
+	double MinExclusiveTime = DBL_MAX;
+	double MaxExclusiveTime = -DBL_MAX;
+	double AverageExclusiveTime = 0.0;
+	double MedianExclusiveTime = 0.0;
+};
+
+struct FTimingProfilerAggregatedStats
+	: public FAggregatedTimingStats
+{
+	const FTimingProfilerTimer* Timer = nullptr;
+};
+
 class ITimingProfilerProvider
 {
 public:
@@ -111,6 +194,7 @@ public:
 	virtual uint64 GetTimelineCount() const = 0;
 	virtual void EnumerateTimelines(TFunctionRef<void(const Timeline&)> Callback) const = 0;
 	virtual void ReadTimers(TFunctionRef<void(const FTimingProfilerTimer*, uint64)> Callback) const = 0;
+	virtual ITable<FTimingProfilerAggregatedStats>* CreateAggregation(double IntervalStart, double IntervalEnd, TFunctionRef<bool(uint32)> CpuThreadFilter, bool IncludeGpu) const = 0;
 };
 
 struct FFileInfo
@@ -171,7 +255,7 @@ struct FPackageSummaryInfo
 
 struct FClassInfo
 {
-	FString Name;
+	const TCHAR* Name;
 };
 
 struct FPackageExportInfo
@@ -186,7 +270,7 @@ struct FPackageExportInfo
 struct FPackageInfo
 {
 	uint32 Id;
-	FString Name;
+	const TCHAR* Name = nullptr;
 	FPackageSummaryInfo Summary;
 	TArray<FPackageExportInfo*> Exports;
 };
@@ -195,10 +279,21 @@ struct FLoadTimeProfilerCpuEvent
 {
 	const FPackageInfo* Package = nullptr;
 	const FPackageExportInfo* Export = nullptr;
-	ELoadTimeProfilePackageEventType PackageEventType = LoadTimeProfilerPackageEventType_None;
+	ELoadTimeProfilerPackageEventType PackageEventType = LoadTimeProfilerPackageEventType_None;
 	ELoadTimeProfilerObjectEventType ExportEventType = LoadTimeProfilerObjectEventType_None;
 };
 
+struct FLoadTimeProfilerAggregatedStats
+{
+	const TCHAR* Name;
+	uint64 Count;
+	double Total;
+	double Min;
+	double Max;
+	double Average;
+	double Median;
+};
+
 class ILoadTimeProfilerProvider
 {
 public:
@@ -209,6 +304,8 @@ public:
 	virtual void EnumeratePackages(TFunctionRef<void(const FPackageInfo&)> Callback) const = 0; 
 	virtual void ReadMainThreadCpuTimeline(TFunctionRef<void(const CpuTimeline&)> Callback) const = 0;
 	virtual void ReadAsyncLoadingThreadCpuTimeline(TFunctionRef<void(const CpuTimeline&)> Callback) const = 0;
+	virtual ITable<FLoadTimeProfilerAggregatedStats>* CreateEventAggregation(double IntervalStart, double IntervalEnd) const = 0;
+	virtual ITable<FLoadTimeProfilerAggregatedStats>* CreateObjectTypeAggregation(double IntervalStart, double IntervalEnd) const = 0;
 };
 
 enum ECounterDisplayHint
diff --git a/Engine/Source/Runtime/CoreUObject/CoreUObject.Build.cs b/Engine/Source/Runtime/CoreUObject/CoreUObject.Build.cs
--- a/Engine/Source/Runtime/CoreUObject/CoreUObject.Build.cs
+++ b/Engine/Source/Runtime/CoreUObject/CoreUObject.Build.cs
@@ -20,6 +20,7 @@ public class CoreUObject : ModuleRules
             );
 
 		PublicDependencyModuleNames.Add("Core");
+        PublicDependencyModuleNames.Add("TraceLog");
 
 		PrivateDependencyModuleNames.Add("Projects");
         PrivateDependencyModuleNames.Add("Json");
diff --git a/Engine/Source/Runtime/CoreUObject/Private/Serialization/LoadTimeTrace.cpp b/Engine/Source/Runtime/CoreUObject/Private/Serialization/LoadTimeTrace.cpp
--- a/Engine/Source/Runtime/CoreUObject/Private/Serialization/LoadTimeTrace.cpp
+++ b/Engine/Source/Runtime/CoreUObject/Private/Serialization/LoadTimeTrace.cpp
@@ -339,7 +339,7 @@ void FLoadTimeProfilerTracePrivate::OutputLinkerArchiveAssociation(const FLinker
 		<< LinkerArchiveAssociation.Archive(Archive);
 }
 
-void FLoadTimeProfilerTracePrivate::OutputQueueEvent(const FAsyncPackage* AsyncPackage, ELoadTimeProfilePackageEventType EventType, TAsyncLoadPriority UserPriority, int32 PackageSerialNumber, int32 EventSystemPriority)
+void FLoadTimeProfilerTracePrivate::OutputQueueEvent(const FAsyncPackage* AsyncPackage, ELoadTimeProfilerPackageEventType EventType, TAsyncLoadPriority UserPriority, int32 PackageSerialNumber, int32 EventSystemPriority)
 {
 	UE_TRACE_LOG(LoadTime, QueueEvent)
 		<< QueueEvent.Cycle(FPlatformTime::Cycles64())
@@ -359,7 +359,7 @@ void FLoadTimeProfilerTracePrivate::OutputClassInfo(const UClass* Class, const T
 		<< ClassInfo.Attachment(Name, NameSize);
 }
 
-FLoadTimeProfilerTracePrivate::FAsyncPackageScope::FAsyncPackageScope(const FAsyncPackage* AsyncPackage, ELoadTimeProfilePackageEventType EventType)
+FLoadTimeProfilerTracePrivate::FAsyncPackageScope::FAsyncPackageScope(const FAsyncPackage* AsyncPackage, ELoadTimeProfilerPackageEventType EventType)
 {
 	UE_TRACE_LOG(LoadTime, BeginAsyncPackageScope)
 		<< BeginAsyncPackageScope.AsyncPackage(AsyncPackage)
diff --git a/Engine/Source/Runtime/CoreUObject/Private/Serialization/LoadTimeTracePrivate.h b/Engine/Source/Runtime/CoreUObject/Private/Serialization/LoadTimeTracePrivate.h
--- a/Engine/Source/Runtime/CoreUObject/Private/Serialization/LoadTimeTracePrivate.h
+++ b/Engine/Source/Runtime/CoreUObject/Private/Serialization/LoadTimeTracePrivate.h
@@ -28,12 +28,12 @@ struct FLoadTimeProfilerTracePrivate
 	static void OutputAsyncPackageRequestAssociation(const FAsyncPackage* AsyncPackage, uint64 RequestId);
 	static void OutputAsyncPackageLinkerAssociation(const FAsyncPackage* AsyncPackage, const FLinkerLoad* Linker);
 	static void OutputLinkerArchiveAssociation(const FLinkerLoad* Linker, const FArchive* Archive);
-	static void OutputQueueEvent(const FAsyncPackage* AsyncPackage, ELoadTimeProfilePackageEventType EventType, TAsyncLoadPriority UserPriority, int32 PackageSerialNumber, int32 EventSystemPriority);
+	static void OutputQueueEvent(const FAsyncPackage* AsyncPackage, ELoadTimeProfilerPackageEventType EventType, TAsyncLoadPriority UserPriority, int32 PackageSerialNumber, int32 EventSystemPriority);
 	static void OutputClassInfo(const UClass* Class, const TCHAR* Name);
 
 	struct FAsyncPackageScope
 	{
-		FAsyncPackageScope(const FAsyncPackage* AsyncPackage, ELoadTimeProfilePackageEventType EventType);
+		FAsyncPackageScope(const FAsyncPackage* AsyncPackage, ELoadTimeProfilerPackageEventType EventType);
 		~FAsyncPackageScope();
 	};
 
diff --git a/Engine/Source/Runtime/CoreUObject/Public/Serialization/LoadTimeTrace.h b/Engine/Source/Runtime/CoreUObject/Public/Serialization/LoadTimeTrace.h
--- a/Engine/Source/Runtime/CoreUObject/Public/Serialization/LoadTimeTrace.h
+++ b/Engine/Source/Runtime/CoreUObject/Public/Serialization/LoadTimeTrace.h
@@ -10,7 +10,7 @@
 #define LOADTIMEPROFILERTRACE_ENABLED 0
 #endif
 
-enum ELoadTimeProfilePackageEventType
+enum ELoadTimeProfilerPackageEventType
 {
 	LoadTimeProfilerPackageEventType_CreateLinker,
 	LoadTimeProfilerPackageEventType_FinishLinker,
