Fixed EC non-unity compile errors.

diff --git a/Engine/Source/Developer/TraceServices/Private/Analyzers/PlatformFileTraceAnalysis.h b/Engine/Source/Developer/TraceServices/Private/Analyzers/PlatformFileTraceAnalysis.h
--- a/Engine/Source/Developer/TraceServices/Private/Analyzers/PlatformFileTraceAnalysis.h
+++ b/Engine/Source/Developer/TraceServices/Private/Analyzers/PlatformFileTraceAnalysis.h
@@ -4,6 +4,7 @@
 
 #include "Trace/Trace.h"
 #include "Trace/Analyzer.h"
+#include "Containers/Map.h"
 
 namespace Trace
 {
diff --git a/Engine/Source/Developer/TraceServices/Private/Analyzers/StatsTraceAnalysis.h b/Engine/Source/Developer/TraceServices/Private/Analyzers/StatsTraceAnalysis.h
--- a/Engine/Source/Developer/TraceServices/Private/Analyzers/StatsTraceAnalysis.h
+++ b/Engine/Source/Developer/TraceServices/Private/Analyzers/StatsTraceAnalysis.h
@@ -4,6 +4,7 @@
 
 #include "Trace/Trace.h"
 #include "Trace/Analyzer.h"
+#include "Templates/SharedPointer.h"
 
 namespace Trace
 {
