UnrealInsights:
TraceServices modular extension support

diff --git a/Engine/Source/Developer/TraceInsights/Private/Insights/ViewModels/GraphTrack.cpp b/Engine/Source/Developer/TraceInsights/Private/Insights/ViewModels/GraphTrack.cpp
--- a/Engine/Source/Developer/TraceInsights/Private/Insights/ViewModels/GraphTrack.cpp
+++ b/Engine/Source/Developer/TraceInsights/Private/Insights/ViewModels/GraphTrack.cpp
@@ -373,7 +373,7 @@ void FFramesGraphTrack::UpdateSeries(FGraphTrackSeries& Series, const FTimingTra
 	if (Session.IsValid())
 	{
 		Trace::FAnalysisSessionReadScope SessionReadScope(*Session.Get());
-		const Trace::IFrameProvider& FramesProvider = Session->ReadFrameProvider();
+		const Trace::IFrameProvider& FramesProvider = ReadFrameProvider(*Session.Get());
 		uint64 FrameCount = FramesProvider.GetFrameCount(FrameType);
 		FramesProvider.EnumerateFrames(FrameType, 0, FrameCount - 1, [&Builder](const Trace::FFrame& Frame)
 		{
diff --git a/Engine/Source/Developer/TraceInsights/Private/Insights/ViewModels/LogMessage.cpp b/Engine/Source/Developer/TraceInsights/Private/Insights/ViewModels/LogMessage.cpp
--- a/Engine/Source/Developer/TraceInsights/Private/Insights/ViewModels/LogMessage.cpp
+++ b/Engine/Source/Developer/TraceInsights/Private/Insights/ViewModels/LogMessage.cpp
@@ -145,7 +145,7 @@ FLogMessageRecord& FLogMessageCache::Get(uint64 Index)
 	if (Session.IsValid())
 	{
 		Trace::FAnalysisSessionReadScope SessionReadScope(*Session.Get());
-		const Trace::ILogProvider& LogProvider = Session->ReadLogProvider();
+		const Trace::ILogProvider& LogProvider = Trace::ReadLogProvider(*Session.Get());
 		LogProvider.ReadMessage(Index, [this, Index](const Trace::FLogMessage& Message)
 		{
 			FScopeLock Lock(&CriticalSection);
@@ -174,7 +174,7 @@ TSharedPtr<FLogMessageRecord> FLogMessageCache::GetUncached(uint64 Index) const
 	if (Session.IsValid())
 	{
 		Trace::FAnalysisSessionReadScope SessionReadScope(*Session.Get());
-		const Trace::ILogProvider& LogProvider = Session->ReadLogProvider();
+		const Trace::ILogProvider& LogProvider = Trace::ReadLogProvider(*Session.Get());
 		LogProvider.ReadMessage(Index, [&EntryPtr](const Trace::FLogMessage& Message)
 		{
 			EntryPtr = MakeShareable(new FLogMessageRecord(Message));
diff --git a/Engine/Source/Developer/TraceInsights/Private/Insights/ViewModels/MarkersTimingTrack.cpp b/Engine/Source/Developer/TraceInsights/Private/Insights/ViewModels/MarkersTimingTrack.cpp
--- a/Engine/Source/Developer/TraceInsights/Private/Insights/ViewModels/MarkersTimingTrack.cpp
+++ b/Engine/Source/Developer/TraceInsights/Private/Insights/ViewModels/MarkersTimingTrack.cpp
@@ -143,7 +143,7 @@ void FMarkersTimingTrack::Update(const FTimingTrackViewport& InViewport)
 	{
 		Trace::FAnalysisSessionReadScope SessionReadScope(*Session.Get());
 
-		const Trace::ILogProvider& LogProvider = Session->ReadLogProvider();
+		const Trace::ILogProvider& LogProvider = Trace::ReadLogProvider(*Session.Get());
 		Builder.BeginLog(LogProvider);
 
 		LogProvider.EnumerateMessages(
diff --git a/Engine/Source/Developer/TraceInsights/Private/Insights/Widgets/SFrameTrack.cpp b/Engine/Source/Developer/TraceInsights/Private/Insights/Widgets/SFrameTrack.cpp
--- a/Engine/Source/Developer/TraceInsights/Private/Insights/Widgets/SFrameTrack.cpp
+++ b/Engine/Source/Developer/TraceInsights/Private/Insights/Widgets/SFrameTrack.cpp
@@ -186,7 +186,7 @@ void SFrameTrack::Tick(const FGeometry& AllottedGeometry, const double InCurrent
 		{
 			Trace::FAnalysisSessionReadScope SessionReadScope(*Session.Get());
 
-			const Trace::IFrameProvider& FramesProvider = Session->ReadFrameProvider();
+			const Trace::IFrameProvider& FramesProvider = Trace::ReadFrameProvider(*Session.Get());
 
 			for (int32 FrameType = 0; FrameType < TraceFrameType_Count; ++FrameType)
 			{
@@ -240,7 +240,7 @@ void SFrameTrack::UpdateState()
 	{
 		Trace::FAnalysisSessionReadScope SessionReadScope(*Session.Get());
 
-		const Trace::IFrameProvider& FramesProvider = Session->ReadFrameProvider();
+		const Trace::IFrameProvider& FramesProvider = Trace::ReadFrameProvider(*Session.Get());
 
 		for (int32 FrameType = 0; FrameType < TraceFrameType_Count; ++FrameType)
 		{
diff --git a/Engine/Source/Developer/TraceInsights/Private/Insights/Widgets/SLogView.cpp b/Engine/Source/Developer/TraceInsights/Private/Insights/Widgets/SLogView.cpp
--- a/Engine/Source/Developer/TraceInsights/Private/Insights/Widgets/SLogView.cpp
+++ b/Engine/Source/Developer/TraceInsights/Private/Insights/Widgets/SLogView.cpp
@@ -705,7 +705,7 @@ void SLogView::Tick(const FGeometry& AllottedGeometry, const double InCurrentTim
 	if (Session.IsValid())
 	{
 		Trace::FAnalysisSessionReadScope SessionReadScope(*Session.Get());
-		const Trace::ILogProvider& LogProvider = Session->ReadLogProvider();
+		const Trace::ILogProvider& LogProvider = Trace::ReadLogProvider(*Session.Get());
 
 		NewMessageCount = static_cast<int32>(LogProvider.GetMessageCount());
 
diff --git a/Engine/Source/Developer/TraceInsights/Private/Insights/Widgets/SStatsView.cpp b/Engine/Source/Developer/TraceInsights/Private/Insights/Widgets/SStatsView.cpp
--- a/Engine/Source/Developer/TraceInsights/Private/Insights/Widgets/SStatsView.cpp
+++ b/Engine/Source/Developer/TraceInsights/Private/Insights/Widgets/SStatsView.cpp
@@ -1401,11 +1401,11 @@ void SStatsView::RebuildTree(bool bResync)
 		bListHasChanged = true;
 	}
 
-	if (Session.IsValid())
+	if (Session.IsValid() && Trace::ReadCounterProvider(*Session.Get()))
 	{
 		Trace::FAnalysisSessionReadScope SessionReadScope(*Session.Get());
 
-		const Trace::ICounterProvider& CountersProvider = Session->ReadCounterProvider();
+		const Trace::ICounterProvider& CountersProvider = *Trace::ReadCounterProvider(*Session.Get());
 
 		if (!bResync)
 		{
@@ -1715,7 +1715,7 @@ void SStatsView::UpdateStats(double StartTime, double EndTime)
 		StatsNodePtr->ResetAggregatedStats();
 	}
 
-	if (Session.IsValid() && StartTime < EndTime)
+	if (Session.IsValid() && StartTime < EndTime && Trace::ReadCounterProvider(*Session.Get()))
 	{
 		const bool bComputeMedian = true;
 
@@ -1725,7 +1725,7 @@ void SStatsView::UpdateStats(double StartTime, double EndTime)
 		{
 			Trace::FAnalysisSessionReadScope SessionReadScope(*Session.Get());
 
-			const Trace::ICounterProvider& CountersProvider = Session->ReadCounterProvider();
+			const Trace::ICounterProvider& CountersProvider = *Trace::ReadCounterProvider(*Session.Get());
 
 			// Compute instance count and total/min/max inclusive/exclusive times for each counter.
 			// Iterate through all counters.
diff --git a/Engine/Source/Developer/TraceInsights/Private/Insights/Widgets/STimersView.cpp b/Engine/Source/Developer/TraceInsights/Private/Insights/Widgets/STimersView.cpp
--- a/Engine/Source/Developer/TraceInsights/Private/Insights/Widgets/STimersView.cpp
+++ b/Engine/Source/Developer/TraceInsights/Private/Insights/Widgets/STimersView.cpp
@@ -1431,11 +1431,11 @@ void STimersView::RebuildTree(bool bResync)
 		bListHasChanged = true;
 	}
 
-	if (Session.IsValid())
+	if (Session.IsValid() && Trace::ReadTimingProfilerProvider(*Session.Get()))
 	{
 		Trace::FAnalysisSessionReadScope SessionReadScope(*Session.Get());
 
-		const Trace::ITimingProfilerProvider& TimingProfilerProvider = Session->ReadTimingProfilerProvider();
+		const Trace::ITimingProfilerProvider& TimingProfilerProvider = *Trace::ReadTimingProfilerProvider(*Session.Get());
 
 		TimingProfilerProvider.ReadTimers([this, &bResync, &bListHasChanged](const Trace::FTimingProfilerTimer* Timers, uint64 TimersCount)
 		{
@@ -1483,8 +1483,7 @@ void STimersView::UpdateStats(double StartTime, double EndTime)
 		TimerNodePtr->ResetAggregatedStats();
 	}
 
-	if (Session.IsValid() &&
-		StartTime < EndTime)
+	if (Session.IsValid() && StartTime < EndTime && Trace::ReadTimingProfilerProvider(*Session.Get()))
 	{
 		auto ThreadFilter = [](uint32 ThreadId)
 		{
@@ -1494,7 +1493,7 @@ void STimersView::UpdateStats(double StartTime, double EndTime)
 		TUniquePtr<Trace::ITable<Trace::FTimingProfilerAggregatedStats>> AggregationResultTable;
 		{
 			Trace::FAnalysisSessionReadScope SessionReadScope(*Session.Get());
-			const Trace::ITimingProfilerProvider& TimingProfilerProvider = Session->ReadTimingProfilerProvider();
+			const Trace::ITimingProfilerProvider& TimingProfilerProvider = *Trace::ReadTimingProfilerProvider(*Session.Get());
 			AggregationResultTable.Reset(TimingProfilerProvider.CreateAggregation(StartTime, EndTime, ThreadFilter, true));
 		}
 
diff --git a/Engine/Source/Developer/TraceInsights/Private/Insights/Widgets/STimingView.cpp b/Engine/Source/Developer/TraceInsights/Private/Insights/Widgets/STimingView.cpp
--- a/Engine/Source/Developer/TraceInsights/Private/Insights/Widgets/STimingView.cpp
+++ b/Engine/Source/Developer/TraceInsights/Private/Insights/Widgets/STimingView.cpp
@@ -370,7 +370,7 @@ void STimingView::Tick(const FGeometry& AllottedGeometry, const double InCurrent
 	UpdateIo();
 
 	TSharedPtr<const Trace::IAnalysisSession> Session = FInsightsManager::Get()->GetSession();
-	if (Session)
+	if (Session && Trace::ReadTimingProfilerProvider(*Session.Get()))
 	{
 		Trace::FAnalysisSessionReadScope SessionReadScope(*Session.Get());
 
@@ -406,7 +406,7 @@ void STimingView::Tick(const FGeometry& AllottedGeometry, const double InCurrent
 		// Sync CachedTimelines and TimingEventsTracks with available timelines on analysis side.
 		// TODO: can we make this more efficient?
 
-		const Trace::ITimingProfilerProvider& TimingProfilerProvider = Session->ReadTimingProfilerProvider();
+		const Trace::ITimingProfilerProvider& TimingProfilerProvider = *Trace::ReadTimingProfilerProvider(*Session.Get());
 
 		// Check if we have a GPU track.
 		uint32 GpuTimelineIndex;
@@ -421,7 +421,7 @@ void STimingView::Tick(const FGeometry& AllottedGeometry, const double InCurrent
 
 		// Check available CPU tracks.
 		int32 Order = 1;
-		const Trace::IThreadProvider& ThreadProvider = Session->ReadThreadProvider();
+		const Trace::IThreadProvider& ThreadProvider = Trace::ReadThreadProvider(*Session.Get());
 		ThreadProvider.EnumerateThreads([this, &Order, &bIsTimingEventsTrackDirty, &TimingProfilerProvider](const Trace::FThreadInfo& ThreadInfo)
 		{
 			uint32 CpuTimelineIndex;
@@ -1020,11 +1020,11 @@ void STimingView::DrawCpuGpuTimelineTrack(FTimingViewDrawHelper& Helper, FTiming
 	if (Helper.BeginTimeline(Track))
 	{
 		TSharedPtr<const Trace::IAnalysisSession> Session = FInsightsManager::Get()->GetSession();
-		if (Session.IsValid())
+		if (Session.IsValid() && Trace::ReadTimingProfilerProvider(*Session.Get()))
 		{
 			Trace::FAnalysisSessionReadScope SessionReadScope(*Session.Get());
 
-			const Trace::ITimingProfilerProvider& TimingProfilerProvider = Session->ReadTimingProfilerProvider();
+			const Trace::ITimingProfilerProvider& TimingProfilerProvider = *Trace::ReadTimingProfilerProvider(*Session.Get());
 
 			TimingProfilerProvider.ReadTimers([this, &Helper, &Track, &TimingProfilerProvider](const Trace::FTimingProfilerTimer* Timers, uint64 TimersCount)
 			{
@@ -1166,11 +1166,11 @@ void STimingView::UpdateIo()
 		AllIoEvents.Reset();
 
 		TSharedPtr<const Trace::IAnalysisSession> Session = FInsightsManager::Get()->GetSession();
-		if (Session.IsValid())
+		if (Session.IsValid() && Trace::ReadFileActivityProvider(*Session.Get()))
 		{
 			Trace::FAnalysisSessionReadScope SessionReadScope(*Session.Get());
 
-			const Trace::IFileActivityProvider& FileActivityProvider = Session->ReadFileActivityProvider();
+			const Trace::IFileActivityProvider& FileActivityProvider = *Trace::ReadFileActivityProvider(*Session.Get());
 
 			// Enumerate all IO events and cache them.
 			FileActivityProvider.EnumerateFileActivity([this](const Trace::FFileInfo& FileInfo, const Trace::IFileActivityProvider::Timeline& Timeline)
@@ -2511,11 +2511,11 @@ bool STimingView::SearchTimingEvent(const double InStartTime,
 	FSearchTimingEventContext Ctx(InStartTime, InEndTime, InPredicate, InOutTimingEvent, bInStopAtFirstMatch, bInSearchForLargestEvent);
 
 	TSharedPtr<const Trace::IAnalysisSession> Session = FInsightsManager::Get()->GetSession();
-	if (Session.IsValid())
+	if (Session.IsValid() && Trace::ReadTimingProfilerProvider(*Session.Get()))
 	{
 		Trace::FAnalysisSessionReadScope SessionReadScope(*Session.Get());
 
-		const Trace::ITimingProfilerProvider& TimingProfilerProvider = Session->ReadTimingProfilerProvider();
+		const Trace::ITimingProfilerProvider& TimingProfilerProvider = *Trace::ReadTimingProfilerProvider(*Session.Get());
 
 		TimingProfilerProvider.ReadTimeline(Ctx.TimingEvent.Track->GetId(), [&Ctx](const Trace::ITimingProfilerProvider::Timeline& Timeline)
 		{
diff --git a/Engine/Source/Developer/TraceServices/Private/AnalysisService.cpp b/Engine/Source/Developer/TraceServices/Private/AnalysisService.cpp
--- a/Engine/Source/Developer/TraceServices/Private/AnalysisService.cpp
+++ b/Engine/Source/Developer/TraceServices/Private/AnalysisService.cpp
@@ -6,14 +6,14 @@
 #include "Trace/Analysis.h"
 #include "Trace/DataStream.h"
 #include "HAL/PlatformFile.h"
-#include "Analyzers/LoadTimeTraceAnalysis.h"
-#include "Analyzers/CpuProfilerTraceAnalysis.h"
 #include "Analyzers/MiscTraceAnalysis.h"
 #include "Analyzers/LogTraceAnalysis.h"
-#include "Analyzers/GpuProfilerTraceAnalysis.h"
-#include "Analyzers/PlatformFileTraceAnalysis.h"
-#include "Analyzers/StatsTraceAnalysis.h"
 #include "Math/RandomStream.h"
+#include "ModuleServicePrivate.h"
+#include "Model/LogPrivate.h"
+#include "Model/FramesPrivate.h"
+#include "Model/BookmarksPrivate.h"
+#include "Model/ThreadsPrivate.h"
 
 namespace Trace
 {
@@ -65,19 +65,39 @@ FAnalysisSession::FAnalysisSession(const TCHAR* SessionName)
 	, DurationSeconds(0.0)
 	, Allocator(32 << 20)
 	, StringStore(Allocator)
-	, BookmarkProvider(Lock, StringStore)
-	, LogProvider(Allocator, Lock, StringStore)
-	, ThreadProvider(Lock, StringStore)
-	, FrameProvider(Allocator, Lock)
-	, TimingProfilerProvider(Allocator, Lock, StringStore)
-	, FileActivityProvider(Allocator, Lock)
-	, LoadTimeProfilerProvider(Allocator, Lock, StringStore)
-	, CounterProvider(Allocator, Lock)
 {
 
 }
 
-FAnalysisService::FAnalysisService()
+FAnalysisSession::~FAnalysisSession()
+{
+	for (auto& KV : Providers)
+	{
+		IProvider* Provider = KV.Value;
+		delete Provider;
+	}
+}
+
+void FAnalysisSession::AddProvider(const FName& InName, IProvider* Provider)
+{
+	Providers.Add(InName, Provider);
+}
+
+const IProvider* FAnalysisSession::ReadProviderPrivate(const FName& InName) const
+{
+	IProvider* const* FindIt = Providers.Find(InName);
+	if (FindIt)
+	{
+		return *FindIt;
+	}
+	else
+	{
+		return nullptr;
+	}
+}
+
+FAnalysisService::FAnalysisService(FModuleService& InModuleService)
+	: ModuleService(InModuleService)
 {
 
 }
@@ -96,16 +116,21 @@ void FAnalysisService::AnalyzeInternal(TSharedRef<FAnalysisSession> AnalysisSess
 
 	TArray<IAnalyzer*> Analyzers;
 	{
-		Trace::FAnalysisSessionEditScope _(AnalysisSession.Get());
-		Analyzers.Add(new FMiscTraceAnalyzer(AnalysisSession.Get()));
-		Analyzers.Add(new FAsyncLoadingTraceAnalyzer(AnalysisSession.Get()));
-		Analyzers.Add(new FCpuProfilerAnalyzer(AnalysisSession.Get()));
-		Analyzers.Add(new FLogTraceAnalyzer(AnalysisSession.Get()));
-		Analyzers.Add(new FGpuProfilerAnalyzer(AnalysisSession.Get()));
-		Analyzers.Add(new FPlatformFileTraceAnalyzer(AnalysisSession.Get()));
-		Analyzers.Add(new FStatsAnalyzer(AnalysisSession.Get()));
+		IAnalysisSession& Session = AnalysisSession.Get();
+		Trace::FAnalysisSessionEditScope _(Session);
+		FBookmarkProvider* BookmarkProvider = new FBookmarkProvider(Session);
+		Session.AddProvider(FBookmarkProvider::ProviderName, BookmarkProvider);
+		FLogProvider* LogProvider = new FLogProvider(Session);
+		Session.AddProvider(FLogProvider::ProviderName, LogProvider);
+		FThreadProvider* ThreadProvider = new FThreadProvider(Session);
+		Session.AddProvider(FThreadProvider::ProviderName, ThreadProvider);
+		FFrameProvider* FrameProvider = new FFrameProvider(Session);
+		Session.AddProvider(FFrameProvider::ProviderName, FrameProvider);
+		Analyzers.Add(new FMiscTraceAnalyzer(Session, *ThreadProvider, *BookmarkProvider, *LogProvider, *FrameProvider));
+		Analyzers.Add(new FLogTraceAnalyzer(Session, *LogProvider));
+		ModuleService.OnAnalysisBegin(Session, Analyzers);
 	}
-
+	
 	FAnalysisContext Context;
 	for (Trace::IAnalyzer* Analyzer : Analyzers)
 	{
diff --git a/Engine/Source/Developer/TraceServices/Private/AnalysisServicePrivate.h b/Engine/Source/Developer/TraceServices/Private/AnalysisServicePrivate.h
--- a/Engine/Source/Developer/TraceServices/Private/AnalysisServicePrivate.h
+++ b/Engine/Source/Developer/TraceServices/Private/AnalysisServicePrivate.h
@@ -6,19 +6,13 @@
 #include "Async/AsyncWork.h"
 #include "Containers/HashTable.h"
 #include "Misc/ScopeLock.h"
-#include "Model/TimingProfiler.h"
 #include "Common/StringStore.h"
-#include "Model/Bookmarks.h"
-#include "Model/Log.h"
-#include "Model/Threads.h"
-#include "Model/FileActivity.h"
-#include "Model/Frames.h"
-#include "Model/LoadTimeProfiler.h"
-#include "Model/Counters.h"
 
 namespace Trace
 {
 
+class FModuleService;
+
 class FAnalysisSessionLock
 {
 public:
@@ -42,47 +36,31 @@ class FAnalysisSession
 {
 public:
 	FAnalysisSession(const TCHAR* SessionName);
+	virtual ~FAnalysisSession();
 
 	virtual const TCHAR* GetName() const { return *Name; }
 	virtual bool IsAnalysisComplete() const override { return IsComplete; }
-	virtual double GetDurationSeconds() const { return DurationSeconds; }
-
-	void UpdateDuration(double Seconds) { Lock.WriteAccessCheck(); DurationSeconds = Seconds > DurationSeconds ? Seconds : DurationSeconds; }
+	virtual double GetDurationSeconds() const { Lock.ReadAccessCheck(); return DurationSeconds; }
+	virtual void UpdateDurationSeconds(double Duration) override { Lock.WriteAccessCheck(); DurationSeconds = FMath::Max(Duration, DurationSeconds); }
 	void SetComplete() { IsComplete = true; }
 
-	FBookmarkProvider& EditBookmarkProvider() { return BookmarkProvider; }
-	virtual const IBookmarkProvider& ReadBookmarkProvider() const override { return BookmarkProvider; }
-
-	FLogProvider& EditLogProvider() { return LogProvider; }
-	virtual const ILogProvider& ReadLogProvider() const override { return LogProvider; }
-
-	FThreadProvider& EditThreadProvider() { return ThreadProvider; }
-	virtual const IThreadProvider& ReadThreadProvider() const override { return ThreadProvider; }
-
-	FFrameProvider& EditFrameProvider() { return FrameProvider; }
-	virtual const IFrameProvider& ReadFrameProvider() const override { return FrameProvider; }
-
-	FTimingProfilerProvider& EditTimingProfilerProvider() { return TimingProfilerProvider; }
-	virtual const FTimingProfilerProvider& ReadTimingProfilerProvider() const override { return TimingProfilerProvider; }
-
-	FFileActivityProvider& EditFileActivityProvider() { return FileActivityProvider; }
-	virtual const FFileActivityProvider& ReadFileActivityProvider() const override { return FileActivityProvider; }
-
-	FLoadTimeProfilerProvider& EditLoadTimeProfilerProvider() { return LoadTimeProfilerProvider; }
-	virtual const ILoadTimeProfilerProvider& ReadLoadTimeProfilerProvider() const override { return LoadTimeProfilerProvider; }
-
-	FCounterProvider& EditCounterProvider() { return CounterProvider; }
-	virtual const ICounterProvider& ReadCounterProvider() const override { return CounterProvider; }
-
-	const TCHAR* StoreString(const TCHAR* String) { return StringStore.Store(String); }
+	virtual ILinearAllocator& GetLinearAllocator() override { return Allocator; }
+	virtual const TCHAR* StoreString(const TCHAR* String) override { return StringStore.Store(String); }
 
 	virtual void BeginRead() const override { Lock.BeginRead(); }
 	virtual void EndRead() const override { Lock.EndRead(); }
 
-	void BeginEdit() { Lock.BeginEdit(); }
-	void EndEdit() { Lock.EndEdit(); }
+	virtual void BeginEdit() override { Lock.BeginEdit(); }
+	virtual void EndEdit() override { Lock.EndEdit(); }
+
+	virtual void ReadAccessCheck() const override { return Lock.ReadAccessCheck(); }
+	virtual void WriteAccessCheck() override { return Lock.WriteAccessCheck(); }
+
+	virtual void AddProvider(const FName& Name, IProvider* Provider) override;
 
 private:
+	virtual const IProvider* ReadProviderPrivate(const FName& Name) const override;
+
 	mutable FAnalysisSessionLock Lock;
 
 	FString Name;
@@ -90,37 +68,14 @@ private:
 	double DurationSeconds = 0.0;
 	FSlabAllocator Allocator;
 	FStringStore StringStore;
-	FBookmarkProvider BookmarkProvider;
-	FLogProvider LogProvider;
-	FThreadProvider ThreadProvider;
-	FFrameProvider FrameProvider;
-	FTimingProfilerProvider TimingProfilerProvider;
-	FFileActivityProvider FileActivityProvider;
-	FLoadTimeProfilerProvider LoadTimeProfilerProvider;
-	FCounterProvider CounterProvider;
-};
-
-struct FAnalysisSessionEditScope
-{
-	FAnalysisSessionEditScope(FAnalysisSession& InAnalysisSession)
-		: AnalysisSession(InAnalysisSession)
-	{
-		AnalysisSession.BeginEdit();
-	}
-
-	~FAnalysisSessionEditScope()
-	{
-		AnalysisSession.EndEdit();
-	}
-
-	FAnalysisSession& AnalysisSession;
+	TMap<FName, IProvider*> Providers;
 };
 
 class FAnalysisService
 	: public IAnalysisService
 {
 public:
-	FAnalysisService();
+	FAnalysisService(FModuleService& ModuleService);
 	virtual ~FAnalysisService();
 	virtual TSharedPtr<const IAnalysisSession> Analyze(const TCHAR* SessionName, TUniquePtr<Trace::IInDataStream>&& DataStream) override;
 	virtual TSharedPtr<const IAnalysisSession> StartAnalysis(const TCHAR* SessionName, TUniquePtr<Trace::IInDataStream>&& DataStream) override;
@@ -146,6 +101,7 @@ private:
 
 	void AnalyzeInternal(TSharedRef<FAnalysisSession> AnalysisSession, Trace::IInDataStream* DataStream);
 
+	FModuleService& ModuleService;
 	FAnalysisStartedEvent AnalysisStartedEvent;
 	FAnalysisFinishedEvent AnalysisFinishedEvent;
 	TArray<TSharedPtr<FAsyncTask<FAnalysisWorker>>> Tasks;
diff --git a/Engine/Source/Developer/TraceServices/Private/Analyzers/CpuProfilerTraceAnalysis.cpp b/Engine/Source/Developer/TraceServices/Private/Analyzers/CpuProfilerTraceAnalysis.cpp
--- a/Engine/Source/Developer/TraceServices/Private/Analyzers/CpuProfilerTraceAnalysis.cpp
+++ b/Engine/Source/Developer/TraceServices/Private/Analyzers/CpuProfilerTraceAnalysis.cpp
@@ -2,12 +2,11 @@
 #include "CpuProfilerTraceAnalysis.h"
 #include "AnalysisServicePrivate.h"
 #include "Common/Utils.h"
-#include "Model/Threads.h"
+#include "TraceServices/Model/Threads.h"
 
-FCpuProfilerAnalyzer::FCpuProfilerAnalyzer(Trace::FAnalysisSession& InSession)
+FCpuProfilerAnalyzer::FCpuProfilerAnalyzer(Trace::IAnalysisSession& InSession, Trace::FTimingProfilerProvider& InTimingProfilerProvider)
 	: Session(InSession)
-	, ThreadProvider(InSession.EditThreadProvider())
-	, TimingProfilerProvider(InSession.EditTimingProfilerProvider())
+	, TimingProfilerProvider(InTimingProfilerProvider)
 {
 
 }
@@ -66,7 +65,7 @@ void FCpuProfilerAnalyzer::OnEvent(uint16 RouteId, const FOnEventContext& Contex
 		check(BufferPtr == BufferEnd);
 		if (LastCycle)
 		{
-			Session.UpdateDuration(Context.SessionContext.TimestampFromCycle(LastCycle));
+			Session.UpdateDurationSeconds(Context.SessionContext.TimestampFromCycle(LastCycle));
 		}
 		ThreadState->LastCycle = LastCycle;
 		break;
@@ -81,7 +80,6 @@ TSharedRef<FCpuProfilerAnalyzer::FThreadState> FCpuProfilerAnalyzer::GetThreadSt
 		TSharedRef<FThreadState> ThreadState = MakeShared<FThreadState>();
 		ThreadState->Timeline = &TimingProfilerProvider.EditCpuThreadTimeline(ThreadId);
 		ThreadStatesMap.Add(ThreadId, ThreadState);
-		ThreadProvider.EnsureThreadExists(ThreadId);
 		return ThreadState;
 	}
 	else
diff --git a/Engine/Source/Developer/TraceServices/Private/Analyzers/CpuProfilerTraceAnalysis.h b/Engine/Source/Developer/TraceServices/Private/Analyzers/CpuProfilerTraceAnalysis.h
--- a/Engine/Source/Developer/TraceServices/Private/Analyzers/CpuProfilerTraceAnalysis.h
+++ b/Engine/Source/Developer/TraceServices/Private/Analyzers/CpuProfilerTraceAnalysis.h
@@ -5,19 +5,18 @@
 #include "Trace/Trace.h"
 #include "Trace/Analyzer.h"
 #include "Containers/UnrealString.h"
-#include "Model/TimingProfiler.h"
+#include "Model/TimingProfilerPrivate.h"
 
 namespace Trace
 {
-	class FThreadProvider;
-	class FAnalysisSession;
+	class IAnalysisSession;
 }
 
 class FCpuProfilerAnalyzer
 	: public Trace::IAnalyzer
 {
 public:
-	FCpuProfilerAnalyzer(Trace::FAnalysisSession& Session);
+	FCpuProfilerAnalyzer(Trace::IAnalysisSession& Session, Trace::FTimingProfilerProvider& TimingProfilerProvider);
 	virtual void OnAnalysisBegin(const FOnAnalysisContext& Context) override;
 	virtual void OnEvent(uint16 RouteId, const FOnEventContext& Context) override;
 	virtual void OnAnalysisEnd() override {};
@@ -44,8 +43,7 @@ private:
 		RouteId_EventBatch,
 	};
 
-	Trace::FAnalysisSession& Session;
-	Trace::FThreadProvider& ThreadProvider;
+	Trace::IAnalysisSession& Session;
 	Trace::FTimingProfilerProvider& TimingProfilerProvider;
 	TMap<uint32, TSharedRef<FThreadState>> ThreadStatesMap;
 	TMap<uint16, uint32> ScopeIdToEventIdMap;
diff --git a/Engine/Source/Developer/TraceServices/Private/Analyzers/GpuProfilerTraceAnalysis.cpp b/Engine/Source/Developer/TraceServices/Private/Analyzers/GpuProfilerTraceAnalysis.cpp
--- a/Engine/Source/Developer/TraceServices/Private/Analyzers/GpuProfilerTraceAnalysis.cpp
+++ b/Engine/Source/Developer/TraceServices/Private/Analyzers/GpuProfilerTraceAnalysis.cpp
@@ -3,9 +3,9 @@
 #include "AnalysisServicePrivate.h"
 #include "Common/Utils.h"
 
-FGpuProfilerAnalyzer::FGpuProfilerAnalyzer(Trace::FAnalysisSession& InSession)
+FGpuProfilerAnalyzer::FGpuProfilerAnalyzer(Trace::FAnalysisSession& InSession, Trace::FTimingProfilerProvider& InTimingProfilerProvider)
 	: Session(InSession)
-	, TimingProfilerProvider(InSession.EditTimingProfilerProvider())
+	, TimingProfilerProvider(InTimingProfilerProvider)
 	, Timeline(TimingProfilerProvider.EditGpuTimeline())
 	, Calibrated(false)
 {
diff --git a/Engine/Source/Developer/TraceServices/Private/Analyzers/GpuProfilerTraceAnalysis.h b/Engine/Source/Developer/TraceServices/Private/Analyzers/GpuProfilerTraceAnalysis.h
--- a/Engine/Source/Developer/TraceServices/Private/Analyzers/GpuProfilerTraceAnalysis.h
+++ b/Engine/Source/Developer/TraceServices/Private/Analyzers/GpuProfilerTraceAnalysis.h
@@ -4,7 +4,7 @@
 
 #include "Trace/Trace.h"
 #include "Trace/Analyzer.h"
-#include "Model/TimingProfiler.h"
+#include "Model/TimingProfilerPrivate.h"
 
 namespace Trace
 {
@@ -15,7 +15,7 @@ class FGpuProfilerAnalyzer
 	: public Trace::IAnalyzer
 {
 public:
-	FGpuProfilerAnalyzer(Trace::FAnalysisSession& Session);
+	FGpuProfilerAnalyzer(Trace::FAnalysisSession& Session, Trace::FTimingProfilerProvider& TimingProfilerProvider);
 	virtual void OnAnalysisBegin(const FOnAnalysisContext& Context) override;
 	virtual void OnEvent(uint16 RouteId, const FOnEventContext& Context) override;
 	virtual void OnAnalysisEnd() override {};
diff --git a/Engine/Source/Developer/TraceServices/Private/Analyzers/LoadTimeTraceAnalysis.cpp b/Engine/Source/Developer/TraceServices/Private/Analyzers/LoadTimeTraceAnalysis.cpp
--- a/Engine/Source/Developer/TraceServices/Private/Analyzers/LoadTimeTraceAnalysis.cpp
+++ b/Engine/Source/Developer/TraceServices/Private/Analyzers/LoadTimeTraceAnalysis.cpp
@@ -4,12 +4,12 @@
 #include "HAL/FileManager.h"
 #include "Serialization/LoadTimeTrace.h"
 #include "Trace/Trace.h"
-#include "Model/LoadTimeProfiler.h"
+#include "Model/LoadTimeProfilerPrivate.h"
 #include "Analyzers/MiscTraceAnalysis.h"
 
-FAsyncLoadingTraceAnalyzer::FAsyncLoadingTraceAnalyzer(Trace::FAnalysisSession& InSession)
+FAsyncLoadingTraceAnalyzer::FAsyncLoadingTraceAnalyzer(Trace::IAnalysisSession& InSession, Trace::FLoadTimeProfilerProvider& InLoadTimeProfilerProvider)
 	: Session(InSession)
-	, LoadTimeProfilerProvider(Session.EditLoadTimeProfilerProvider())
+	, LoadTimeProfilerProvider(InLoadTimeProfilerProvider)
 {
 }
 
diff --git a/Engine/Source/Developer/TraceServices/Private/Analyzers/LoadTimeTraceAnalysis.h b/Engine/Source/Developer/TraceServices/Private/Analyzers/LoadTimeTraceAnalysis.h
--- a/Engine/Source/Developer/TraceServices/Private/Analyzers/LoadTimeTraceAnalysis.h
+++ b/Engine/Source/Developer/TraceServices/Private/Analyzers/LoadTimeTraceAnalysis.h
@@ -5,7 +5,7 @@
 #include "Templates/SharedPointer.h"
 #include "Trace/Analyzer.h"
 #include "AnalysisServicePrivate.h"
-#include "Model/LoadTimeProfiler.h"
+#include "Model/LoadTimeProfilerPrivate.h"
 
 namespace Trace
 {
@@ -23,7 +23,7 @@ inline bool operator!=(const Trace::FLoadTimeProfilerCpuEvent& Lhs, const Trace:
 class FAsyncLoadingTraceAnalyzer : public Trace::IAnalyzer
 {
 public:
-	FAsyncLoadingTraceAnalyzer(Trace::FAnalysisSession& Session);
+	FAsyncLoadingTraceAnalyzer(Trace::IAnalysisSession& Session, Trace::FLoadTimeProfilerProvider& LoadTimeProfilerProvider);
 	virtual ~FAsyncLoadingTraceAnalyzer();
 
 	virtual void OnAnalysisBegin(const FOnAnalysisContext& Context) override;
@@ -198,7 +198,7 @@ private:
 		RouteId_ClassInfo,
 	};
 
-	Trace::FAnalysisSession& Session;
+	Trace::IAnalysisSession& Session;
 	Trace::FLoadTimeProfilerProvider& LoadTimeProfilerProvider;
 	//TSharedPtr<Trace::FMonotonicTimeline<uint64>> BlockingRequestsTimeline;
 
diff --git a/Engine/Source/Developer/TraceServices/Private/Analyzers/LogTraceAnalysis.cpp b/Engine/Source/Developer/TraceServices/Private/Analyzers/LogTraceAnalysis.cpp
--- a/Engine/Source/Developer/TraceServices/Private/Analyzers/LogTraceAnalysis.cpp
+++ b/Engine/Source/Developer/TraceServices/Private/Analyzers/LogTraceAnalysis.cpp
@@ -2,11 +2,11 @@
 #include "LogTraceAnalysis.h"
 #include "AnalysisServicePrivate.h"
 #include "Logging/LogTrace.h"
-#include "Model/Log.h"
+#include "Model/LogPrivate.h"
 
-FLogTraceAnalyzer::FLogTraceAnalyzer(Trace::FAnalysisSession& InSession)
+FLogTraceAnalyzer::FLogTraceAnalyzer(Trace::IAnalysisSession& InSession, Trace::FLogProvider& InLogProvider)
 	: Session(InSession)
-	, LogProvider(InSession.EditLogProvider())
+	, LogProvider(InLogProvider)
 {
 
 }
diff --git a/Engine/Source/Developer/TraceServices/Private/Analyzers/LogTraceAnalysis.h b/Engine/Source/Developer/TraceServices/Private/Analyzers/LogTraceAnalysis.h
--- a/Engine/Source/Developer/TraceServices/Private/Analyzers/LogTraceAnalysis.h
+++ b/Engine/Source/Developer/TraceServices/Private/Analyzers/LogTraceAnalysis.h
@@ -7,7 +7,7 @@
 
 namespace Trace
 {
-	class FAnalysisSession;
+	class IAnalysisSession;
 	class FLogProvider;
 }
 
@@ -15,7 +15,7 @@ class FLogTraceAnalyzer
 	: public Trace::IAnalyzer
 {
 public:
-	FLogTraceAnalyzer(Trace::FAnalysisSession& Session);
+	FLogTraceAnalyzer(Trace::IAnalysisSession& Session, Trace::FLogProvider& LogProvider);
 	virtual void OnAnalysisBegin(const FOnAnalysisContext& Context) override;
 	virtual void OnEvent(uint16 RouteId, const FOnEventContext& Context) override;
 	virtual void OnAnalysisEnd() override {};
@@ -28,6 +28,6 @@ private:
 		RouteId_LogMessage,
 	};
 
-	Trace::FAnalysisSession& Session;
+	Trace::IAnalysisSession& Session;
 	Trace::FLogProvider& LogProvider;
 };
diff --git a/Engine/Source/Developer/TraceServices/Private/Analyzers/MiscTraceAnalysis.cpp b/Engine/Source/Developer/TraceServices/Private/Analyzers/MiscTraceAnalysis.cpp
--- a/Engine/Source/Developer/TraceServices/Private/Analyzers/MiscTraceAnalysis.cpp
+++ b/Engine/Source/Developer/TraceServices/Private/Analyzers/MiscTraceAnalysis.cpp
@@ -1,17 +1,21 @@
 // Copyright 1998-2019 Epic Games, Inc. All Rights Reserved.
 #include "MiscTraceAnalysis.h"
-#include "AnalysisServicePrivate.h"
-#include "Model/Log.h"
-#include "Model/Threads.h"
-#include "Model/Bookmarks.h"
-#include "Model/Frames.h"
+#include "TraceServices/Model/AnalysisSession.h"
+#include "Model/LogPrivate.h"
+#include "Model/ThreadsPrivate.h"
+#include "Model/BookmarksPrivate.h"
+#include "Model/FramesPrivate.h"
 
-FMiscTraceAnalyzer::FMiscTraceAnalyzer(Trace::FAnalysisSession& InSession)
+FMiscTraceAnalyzer::FMiscTraceAnalyzer(Trace::IAnalysisSession& InSession,
+									   Trace::FThreadProvider& InThreadProvider,
+									   Trace::FBookmarkProvider& InBookmarkProvider,
+									   Trace::FLogProvider& InLogProvider,
+									   Trace::FFrameProvider& InFrameProvider)
 	: Session(InSession)
-	, ThreadProvider(InSession.EditThreadProvider())
-	, BookmarkProvider(InSession.EditBookmarkProvider())
-	, LogProvider(InSession.EditLogProvider())
-	, FrameProvider(InSession.EditFrameProvider())
+	, ThreadProvider(InThreadProvider)
+	, BookmarkProvider(InBookmarkProvider)
+	, LogProvider(InLogProvider)
+	, FrameProvider(InFrameProvider)
 {
 	Trace::FLogCategory& BookmarkLogCategory = LogProvider.GetCategory(Trace::FLogProvider::ReservedLogCategory_Bookmark);
 	BookmarkLogCategory.Name = TEXT("LogBookmark");
diff --git a/Engine/Source/Developer/TraceServices/Private/Analyzers/MiscTraceAnalysis.h b/Engine/Source/Developer/TraceServices/Private/Analyzers/MiscTraceAnalysis.h
--- a/Engine/Source/Developer/TraceServices/Private/Analyzers/MiscTraceAnalysis.h
+++ b/Engine/Source/Developer/TraceServices/Private/Analyzers/MiscTraceAnalysis.h
@@ -10,7 +10,7 @@
 
 namespace Trace
 {
-	class FAnalysisSession;
+	class IAnalysisSession;
 	class FThreadProvider;
 	class FBookmarkProvider;
 	class FLogProvider;
@@ -21,7 +21,11 @@ class FMiscTraceAnalyzer
 	: public Trace::IAnalyzer
 {
 public:
-	FMiscTraceAnalyzer(Trace::FAnalysisSession& Session);
+	FMiscTraceAnalyzer(Trace::IAnalysisSession& Session,
+					   Trace::FThreadProvider& ThreadProvider,
+					   Trace::FBookmarkProvider& BookmarkProvider,
+					   Trace::FLogProvider& LogProvider,
+					   Trace::FFrameProvider& FrameProvider);
 	virtual void OnAnalysisBegin(const FOnAnalysisContext& Context) override;
 	virtual void OnEvent(uint16 RouteId, const FOnEventContext& Context) override;
 	virtual void OnAnalysisEnd() override {};
@@ -47,7 +51,7 @@ private:
 
 	FThreadState* GetThreadState(uint32 ThreadId);
 
-	Trace::FAnalysisSession& Session;
+	Trace::IAnalysisSession& Session;
 	Trace::FThreadProvider& ThreadProvider;
 	Trace::FBookmarkProvider& BookmarkProvider;
 	Trace::FLogProvider& LogProvider;
diff --git a/Engine/Source/Developer/TraceServices/Private/Analyzers/PlatformFileTraceAnalysis.cpp b/Engine/Source/Developer/TraceServices/Private/Analyzers/PlatformFileTraceAnalysis.cpp
--- a/Engine/Source/Developer/TraceServices/Private/Analyzers/PlatformFileTraceAnalysis.cpp
+++ b/Engine/Source/Developer/TraceServices/Private/Analyzers/PlatformFileTraceAnalysis.cpp
@@ -3,9 +3,9 @@
 #include "AnalysisServicePrivate.h"
 #include "Model/FileActivity.h"
 
-FPlatformFileTraceAnalyzer::FPlatformFileTraceAnalyzer(Trace::FAnalysisSession& InSession)
+FPlatformFileTraceAnalyzer::FPlatformFileTraceAnalyzer(Trace::IAnalysisSession& InSession, Trace::FFileActivityProvider& InFileActivityProvider)
 	: Session(InSession)
-	, FileActivityProvider(Session.EditFileActivityProvider())
+	, FileActivityProvider(InFileActivityProvider)
 {
 
 }
diff --git a/Engine/Source/Developer/TraceServices/Private/Analyzers/PlatformFileTraceAnalysis.h b/Engine/Source/Developer/TraceServices/Private/Analyzers/PlatformFileTraceAnalysis.h
--- a/Engine/Source/Developer/TraceServices/Private/Analyzers/PlatformFileTraceAnalysis.h
+++ b/Engine/Source/Developer/TraceServices/Private/Analyzers/PlatformFileTraceAnalysis.h
@@ -5,10 +5,10 @@
 #include "Trace/Trace.h"
 #include "Trace/Analyzer.h"
 #include "Containers/Map.h"
+#include "TraceServices/Model/AnalysisSession.h"
 
 namespace Trace
 {
-	class FAnalysisSession;
 	class FTimeline;
 	class FFileActivityProvider;
 }
@@ -17,7 +17,7 @@ class FPlatformFileTraceAnalyzer
 	: public Trace::IAnalyzer
 {
 public:
-	FPlatformFileTraceAnalyzer(Trace::FAnalysisSession& Session);
+	FPlatformFileTraceAnalyzer(Trace::IAnalysisSession& Session, Trace::FFileActivityProvider& FileActivityProvider);
 	virtual void OnAnalysisBegin(const FOnAnalysisContext& Context) override;
 	virtual void OnEvent(uint16 RouteId, const FOnEventContext& Context) override;
 	virtual void OnAnalysisEnd() override;
@@ -43,7 +43,7 @@ private:
 		uint32 FileIndex;
 	};
 
-	Trace::FAnalysisSession& Session;
+	Trace::IAnalysisSession& Session;
 	Trace::FFileActivityProvider& FileActivityProvider;
 	TMap<uint64, uint32> OpenFilesMap;
 	TMap<uint64, FPendingActivity> PendingOpenMap;
diff --git a/Engine/Source/Developer/TraceServices/Private/Analyzers/StatsTraceAnalysis.cpp b/Engine/Source/Developer/TraceServices/Private/Analyzers/StatsTraceAnalysis.cpp
--- a/Engine/Source/Developer/TraceServices/Private/Analyzers/StatsTraceAnalysis.cpp
+++ b/Engine/Source/Developer/TraceServices/Private/Analyzers/StatsTraceAnalysis.cpp
@@ -4,9 +4,9 @@
 #include "Common/Utils.h"
 #include "Model/Counters.h"
 
-FStatsAnalyzer::FStatsAnalyzer(Trace::FAnalysisSession& InSession)
+FStatsAnalyzer::FStatsAnalyzer(Trace::FAnalysisSession& InSession, Trace::FCounterProvider& InCounterProvider)
 	: Session(InSession)
-	, CounterProvider(Session.EditCounterProvider())
+	, CounterProvider(InCounterProvider)
 {
 
 }
diff --git a/Engine/Source/Developer/TraceServices/Private/Analyzers/StatsTraceAnalysis.h b/Engine/Source/Developer/TraceServices/Private/Analyzers/StatsTraceAnalysis.h
--- a/Engine/Source/Developer/TraceServices/Private/Analyzers/StatsTraceAnalysis.h
+++ b/Engine/Source/Developer/TraceServices/Private/Analyzers/StatsTraceAnalysis.h
@@ -17,7 +17,7 @@ class FStatsAnalyzer
 	: public Trace::IAnalyzer
 {
 public:
-	FStatsAnalyzer(Trace::FAnalysisSession& Session);
+	FStatsAnalyzer(Trace::FAnalysisSession& Session, Trace::FCounterProvider& CounterProvider);
 	virtual void OnAnalysisBegin(const FOnAnalysisContext& Context) override;
 	virtual void OnEvent(uint16 RouteId, const FOnEventContext& Context) override;
 	virtual void OnAnalysisEnd() override;
diff --git a/Engine/Source/Developer/TraceServices/Private/Common/PagedArray.h b/Engine/Source/Developer/TraceServices/Private/Common/PagedArray.h
--- a/Engine/Source/Developer/TraceServices/Private/Common/PagedArray.h
+++ b/Engine/Source/Developer/TraceServices/Private/Common/PagedArray.h
@@ -2,7 +2,8 @@
 
 #pragma once
 
-#include "Common/SlabAllocator.h"
+#include "TraceServices/Containers/Allocators.h"
+#include "Containers/Array.h"
 
 template<typename ItemType>
 struct TPagedArrayPage
@@ -179,7 +180,7 @@ public:
 		const ItemType* CurrentPageLastItem = nullptr;
 	};
 
-	TPagedArray(FSlabAllocator& InAllocator, uint64 InPageSize)
+	TPagedArray(Trace::ILinearAllocator& InAllocator, uint64 InPageSize)
 		: Allocator(InAllocator)
 		, PageSize(InPageSize)
 	{
@@ -219,7 +220,7 @@ public:
 		{
 			LastPage = &PagesArray.AddDefaulted_GetRef();
 			FirstPage = PagesArray.GetData();
-			LastPage->Items = Allocator.Allocate<ItemType>(PageSize);
+			LastPage->Items = reinterpret_cast<ItemType*>(Allocator.Allocate(PageSize * sizeof(ItemType)));
 		}
 		++TotalItemCount;
 		ItemType* ItemPtr = LastPage->Items + LastPage->Count;
@@ -297,7 +298,7 @@ public:
 	}
 
 private:
-	FSlabAllocator& Allocator;
+	Trace::ILinearAllocator& Allocator;
 	TArray<PageType> PagesArray;
 	PageType* FirstPage = nullptr;
 	PageType* LastPage = nullptr;
diff --git a/Engine/Source/Developer/TraceServices/Private/Common/SlabAllocator.h b/Engine/Source/Developer/TraceServices/Private/Common/SlabAllocator.h
--- a/Engine/Source/Developer/TraceServices/Private/Common/SlabAllocator.h
+++ b/Engine/Source/Developer/TraceServices/Private/Common/SlabAllocator.h
@@ -4,8 +4,10 @@
 
 #include "HAL/UnrealMemory.h"
 #include "Containers/Array.h"
+#include "TraceServices/Containers/Allocators.h"
 
 class FSlabAllocator
+	: public Trace::ILinearAllocator
 {
 public:
 	FSlabAllocator(uint64 InSlabSize)
@@ -22,10 +24,9 @@ public:
 		}
 	}
 
-	template<typename T>
-	T* Allocate(uint64 Count)
+	virtual void* Allocate(uint64 Size) override
 	{
-		uint64 AllocationSize = Count * sizeof(T) + (16 - 1) / 16;
+		uint64 AllocationSize = Size + (16 - 1) / 16;
 		if (!CurrentSlab || CurrentSlabAllocatedSize + AllocationSize > SlabSize)
 		{
 			TotalAllocatedSize += SlabSize;
@@ -34,9 +35,8 @@ public:
 			CurrentSlabAllocatedSize = 0;
 			Slabs.Add(CurrentSlab);
 		}
-		T* Allocation = reinterpret_cast<T*>(CurrentSlab + CurrentSlabAllocatedSize);
+		void* Allocation = CurrentSlab + CurrentSlabAllocatedSize;
 		CurrentSlabAllocatedSize += AllocationSize;
-		TotalItemCount += Count;
 		return Allocation;
 	}
 
@@ -45,6 +45,5 @@ private:
 	uint8* CurrentSlab = nullptr;
 	const uint64 SlabSize;
 	uint64 CurrentSlabAllocatedSize = 0;
-	uint64 TotalItemCount = 0;
 	uint64 TotalAllocatedSize = 0;
 };
diff --git a/Engine/Source/Developer/TraceServices/Private/Common/StringStore.cpp b/Engine/Source/Developer/TraceServices/Private/Common/StringStore.cpp
--- a/Engine/Source/Developer/TraceServices/Private/Common/StringStore.cpp
+++ b/Engine/Source/Developer/TraceServices/Private/Common/StringStore.cpp
@@ -23,7 +23,7 @@ const TCHAR* FStringStore::Store(const TCHAR* String)
 	int32 StringLength = FCString::Strlen(String) + 1;
 	if (BufferLeft < StringLength)
 	{
-		BufferPtr = Allocator.Allocate<TCHAR>(BlockSize);
+		BufferPtr = reinterpret_cast<TCHAR*>(Allocator.Allocate(BlockSize * sizeof(TCHAR)));
 		++BlockCount;
 		BufferLeft = BlockSize;
 	}
diff --git a/Engine/Source/Developer/TraceServices/Private/Model/Bookmarks.cpp b/Engine/Source/Developer/TraceServices/Private/Model/Bookmarks.cpp
--- a/Engine/Source/Developer/TraceServices/Private/Model/Bookmarks.cpp
+++ b/Engine/Source/Developer/TraceServices/Private/Model/Bookmarks.cpp
@@ -1,22 +1,24 @@
 // Copyright 1998-2019 Epic Games, Inc. All Rights Reserved.
 
-#include "Model/Bookmarks.h"
+#include "TraceServices/Model/Bookmarks.h"
+#include "Model/BookmarksPrivate.h"
 #include "AnalysisServicePrivate.h"
 #include "Common/FormatArgs.h"
 
 namespace Trace
 {
 
-FBookmarkProvider::FBookmarkProvider(const FAnalysisSessionLock& InSessionLock, FStringStore& InStringStore)
-	: SessionLock(InSessionLock)
-	, StringStore(InStringStore)
+const FName FBookmarkProvider::ProviderName("BookmarkProvider");
+
+FBookmarkProvider::FBookmarkProvider(IAnalysisSession& InSession)
+	: Session(InSession)
 {
 
 }
 
 FBookmarkSpec& FBookmarkProvider::GetSpec(uint64 BookmarkPoint)
 {
-	SessionLock.WriteAccessCheck();
+	Session.WriteAccessCheck();
 	if (SpecMap.Contains(BookmarkPoint))
 	{
 		return *SpecMap[BookmarkPoint].Get();
@@ -31,18 +33,18 @@ FBookmarkSpec& FBookmarkProvider::GetSpec(uint64 BookmarkPoint)
 
 void FBookmarkProvider::AppendBookmark(double Time, uint64 BookmarkPoint, const uint8* FormatArgs)
 {
-	SessionLock.WriteAccessCheck();
+	Session.WriteAccessCheck();
 	FBookmarkSpec& Spec = GetSpec(BookmarkPoint);
 	TSharedRef<FBookmarkInternal> Bookmark = MakeShared<FBookmarkInternal>();
 	Bookmark->Time = Time;
 	FFormatArgsHelper::Format(FormatBuffer, FormatBufferSize - 1, Spec.FormatString, FormatArgs);
-	Bookmark->Text = StringStore.Store(FormatBuffer);
+	Bookmark->Text = Session.StoreString(FormatBuffer);
 	Bookmarks.Add(Bookmark);
 }
 
 void FBookmarkProvider::EnumerateBookmarks(double IntervalStart, double IntervalEnd, TFunctionRef<void(const FBookmark &)> Callback) const
 {
-	SessionLock.ReadAccessCheck();
+	Session.ReadAccessCheck();
 	if (IntervalStart > IntervalEnd)
 	{
 		return;
diff --git a/Engine/Source/Developer/TraceServices/Private/Model/Bookmarks.h b/Engine/Source/Developer/TraceServices/Private/Model/Bookmarks.h
deleted file mode 100644
diff --git a/Engine/Source/Developer/TraceServices/Private/Model/BookmarksPrivate.h b/Engine/Source/Developer/TraceServices/Private/Model/BookmarksPrivate.h
new file mode 100644
--- /dev/null
+++ b/Engine/Source/Developer/TraceServices/Private/Model/BookmarksPrivate.h
@@ -0,0 +1,52 @@
+// Copyright 1998-2019 Epic Games, Inc. All Rights Reserved.
+
+#pragma once
+
+#include "TraceServices/AnalysisService.h"
+#include "Templates/SharedPointer.h"
+
+namespace Trace
+{
+
+class FAnalysisSessionLock;
+class FStringStore;
+
+struct FBookmarkSpec
+{
+	const TCHAR* File;
+	const TCHAR* FormatString;
+	int32 Line;
+};
+
+struct FBookmarkInternal
+{
+	double Time;
+	const TCHAR* Text;
+};
+
+class FBookmarkProvider
+	: public IBookmarkProvider
+{
+public:
+	static const FName ProviderName;
+
+	FBookmarkProvider(IAnalysisSession& Session);
+
+	FBookmarkSpec& GetSpec(uint64 BookmarkPoint);
+	virtual uint64 GetBookmarkCount() const override { return Bookmarks.Num(); }
+	void AppendBookmark(double Time, uint64 BookmarkPoint, const uint8* FormatArgs);
+	virtual void EnumerateBookmarks(double IntervalStart, double IntervalEnd, TFunctionRef<void(const FBookmark&)> Callback) const override;
+
+private:
+	enum
+	{
+		FormatBufferSize = 65536
+	};
+
+	IAnalysisSession& Session;
+	TMap<uint64, TSharedPtr<FBookmarkSpec>> SpecMap;
+	TArray<TSharedRef<FBookmarkInternal>> Bookmarks;
+	TCHAR FormatBuffer[FormatBufferSize];
+};
+
+}
\ No newline at end of file
diff --git a/Engine/Source/Developer/TraceServices/Private/Model/Counters.cpp b/Engine/Source/Developer/TraceServices/Private/Model/Counters.cpp
--- a/Engine/Source/Developer/TraceServices/Private/Model/Counters.cpp
+++ b/Engine/Source/Developer/TraceServices/Private/Model/Counters.cpp
@@ -6,15 +6,15 @@
 namespace Trace
 {
 
-FCounterInternal::FCounterInternal(FSlabAllocator& InAllocator, uint32 InId, const TCHAR* InName, const TCHAR* InDescription, ECounterDisplayHint InDisplayHint)
+FCounterInternal::FCounterInternal(ILinearAllocator& Allocator, uint32 InId, const TCHAR* InName, const TCHAR* InDescription, ECounterDisplayHint InDisplayHint)
 	: Name(InName)
 	, Description(InDescription)
 	, Id(InId)
 	, DisplayHint(InDisplayHint)
-	, Timestamps(InAllocator, 1024)
-	, Ops(InAllocator, 1024)
-	, IntOpArguments(InAllocator, 1024)
-	, DoubleOpArguments(InAllocator, 1024)
+	, Timestamps(Allocator, 1024)
+	, Ops(Allocator, 1024)
+	, IntOpArguments(Allocator, 1024)
+	, DoubleOpArguments(Allocator, 1024)
 {
 
 }
@@ -94,9 +94,8 @@ void FCounterInternal::EnumerateFloatValues(double IntervalStart, double Interva
 	}
 }
 
-FCounterProvider::FCounterProvider(FSlabAllocator& InAllocator, FAnalysisSessionLock& InSessionLock)
-	: Allocator(InAllocator)
-	, SessionLock(InSessionLock)
+FCounterProvider::FCounterProvider(IAnalysisSession& InSession)
+	: Session(InSession)
 {
 
 }
@@ -120,7 +119,7 @@ void FCounterProvider::EnumerateCounters(TFunctionRef<void(const ICounter &)> Ca
 
 FCounterInternal* FCounterProvider::CreateCounter(const TCHAR* Name, const TCHAR* Description, ECounterDisplayHint DisplayHint)
 {
-	FCounterInternal* Counter = new FCounterInternal(Allocator, Counters.Num(), Name, Description, DisplayHint);
+	FCounterInternal* Counter = new FCounterInternal(Session.GetLinearAllocator(), Counters.Num(), Name, Description, DisplayHint);
 	Counters.Add(Counter);
 	return Counter;
 }
diff --git a/Engine/Source/Developer/TraceServices/Private/Model/Counters.h b/Engine/Source/Developer/TraceServices/Private/Model/Counters.h
--- a/Engine/Source/Developer/TraceServices/Private/Model/Counters.h
+++ b/Engine/Source/Developer/TraceServices/Private/Model/Counters.h
@@ -2,9 +2,10 @@
 
 #pragma once
 
-#include "TraceServices/AnalysisService.h"
+#include "TraceServices/Model/AnalysisSession.h"
 #include "Templates/SharedPointer.h"
 #include "Common/PagedArray.h"
+#include "TraceServices/Model/Stats.h"
 
 namespace Trace
 {
@@ -21,7 +22,7 @@ class FCounterInternal
 	: public ICounter
 {
 public:
-	FCounterInternal(FSlabAllocator& InAllocator, uint32 Id, const TCHAR* Name, const TCHAR* Description, ECounterDisplayHint DisplayHint);
+	FCounterInternal(ILinearAllocator& Allocator, uint32 Id, const TCHAR* Name, const TCHAR* Description, ECounterDisplayHint DisplayHint);
 	virtual const TCHAR* GetName() const override { return *Name; }
 	virtual uint32 GetId() const override { return Id; }
 	virtual ECounterDisplayHint GetDisplayHint() const { return DisplayHint; }
@@ -52,7 +53,7 @@ class FCounterProvider
 	: public ICounterProvider
 {
 public:
-	FCounterProvider(FSlabAllocator& Allocator, FAnalysisSessionLock& SessionLock);
+	FCounterProvider(IAnalysisSession& Session);
 	virtual ~FCounterProvider();
 	virtual uint64 GetCounterCount() const override { return Counters.Num(); }
 	virtual void EnumerateCounters(TFunctionRef<void(const ICounter&)> Callback) const override;
@@ -64,8 +65,7 @@ public:
 	void Set(FCounterInternal& Counter, double Time, double Value);
 
 private:
-	FSlabAllocator& Allocator;
-	FAnalysisSessionLock& SessionLock;
+	IAnalysisSession& Session;
 	TArray<FCounterInternal*> Counters;
 };
 
diff --git a/Engine/Source/Developer/TraceServices/Private/Model/FileActivity.cpp b/Engine/Source/Developer/TraceServices/Private/Model/FileActivity.cpp
--- a/Engine/Source/Developer/TraceServices/Private/Model/FileActivity.cpp
+++ b/Engine/Source/Developer/TraceServices/Private/Model/FileActivity.cpp
@@ -7,10 +7,9 @@
 namespace Trace
 {
 
-FFileActivityProvider::FFileActivityProvider(FSlabAllocator& InAllocator, FAnalysisSessionLock& InSessionLock)
-	: Allocator(InAllocator)
-	, SessionLock(InSessionLock)
-	, Files(Allocator, 1024)
+FFileActivityProvider::FFileActivityProvider(IAnalysisSession& InSession)
+	: Session(InSession)
+	, Files(InSession.GetLinearAllocator(), 1024)
 {
 }
 
@@ -34,7 +33,7 @@ uint32 FFileActivityProvider::GetFileIndex(const TCHAR* Path)
 	uint32 FileIndex = Files.Num();
 	FFileInfoInternal& FileInfo = Files.PushBack();
 	FileInfo.Path = Path;
-	FileInfo.ActivityTimeline = MakeShared<TimelineInternal>(Allocator);
+	FileInfo.ActivityTimeline = MakeShared<TimelineInternal>(Session.GetLinearAllocator());
 	return FileIndex;
 }
 
diff --git a/Engine/Source/Developer/TraceServices/Private/Model/FileActivity.h b/Engine/Source/Developer/TraceServices/Private/Model/FileActivity.h
--- a/Engine/Source/Developer/TraceServices/Private/Model/FileActivity.h
+++ b/Engine/Source/Developer/TraceServices/Private/Model/FileActivity.h
@@ -2,9 +2,12 @@
 
 #pragma once
 
-#include "TraceServices/AnalysisService.h"
+#include "TraceServices/Model/AnalysisSession.h"
 #include "Common/PagedArray.h"
 #include "Model/IntervalTimeline.h"
+#include "TraceServices/Model/LoadTimeProfiler.h"
+#include "Containers/UnrealString.h"
+#include "Templates/SharedPointer.h"
 
 namespace Trace
 {
@@ -24,7 +27,7 @@ public:
 	};
 	typedef TIntervalTimeline<FFileActivity, FTimelineSettings> TimelineInternal;
 
-	FFileActivityProvider(FSlabAllocator& Allocator, FAnalysisSessionLock& SessionLock);
+	FFileActivityProvider(IAnalysisSession& Session);
 	virtual void EnumerateFileActivity(TFunctionRef<bool(const FFileInfo&, const Timeline&)> Callback) const override;
 	uint32 GetFileIndex(const TCHAR* Path);
 	uint64 BeginActivity(uint32 FileIndex, EFileActivityType Type, uint64 Offset, uint64 Size, double Time);
@@ -38,8 +41,7 @@ private:
 		TSharedPtr<TimelineInternal> ActivityTimeline;
 	};
 
-	FSlabAllocator& Allocator;
-	FAnalysisSessionLock& SessionLock;
+	IAnalysisSession& Session;
 	TPagedArray<FFileInfoInternal> Files;
 };
 
diff --git a/Engine/Source/Developer/TraceServices/Private/Model/Frames.cpp b/Engine/Source/Developer/TraceServices/Private/Model/Frames.cpp
--- a/Engine/Source/Developer/TraceServices/Private/Model/Frames.cpp
+++ b/Engine/Source/Developer/TraceServices/Private/Model/Frames.cpp
@@ -1,31 +1,33 @@
 // Copyright 1998-2019 Epic Games, Inc. All Rights Reserved.
 
-#include "Model/Frames.h"
+#include "TraceServices/Model/Frames.h"
+#include "Model/FramesPrivate.h"
 #include "AnalysisServicePrivate.h"
 #include <limits>
 
 namespace Trace
 {
 
-FFrameProvider::FFrameProvider(FSlabAllocator& InAllocator, FAnalysisSessionLock& InSessionLock)
-	: Allocator(InAllocator)
-	, SessionLock(InSessionLock)
+const FName FFrameProvider::ProviderName("FrameProvider");
+
+FFrameProvider::FFrameProvider(IAnalysisSession& InSession)
+	: Session(InSession)
 {
 	for (int32 FrameType = 0; FrameType < TraceFrameType_Count; ++FrameType)
 	{
-		Frames.Emplace(Allocator, 65536);
+		Frames.Emplace(Session.GetLinearAllocator(), 65536);
 	}
 }
 
 uint64 FFrameProvider::GetFrameCount(ETraceFrameType FrameType) const
 {
-	SessionLock.ReadAccessCheck();
+	Session.ReadAccessCheck();
 	return Frames[FrameType].Num();
 }
 
 void FFrameProvider::EnumerateFrames(ETraceFrameType FrameType, uint64 Start, uint64 End, TFunctionRef<void(const FFrame&)> Callback) const
 {
-	SessionLock.ReadAccessCheck();
+	Session.ReadAccessCheck();
 
 	End = FMath::Min(End, Frames[FrameType].Num());
 	if (Start >= End)
@@ -43,7 +45,7 @@ void FFrameProvider::EnumerateFrames(ETraceFrameType FrameType, uint64 Start, ui
 
 void FFrameProvider::BeginFrame(ETraceFrameType FrameType, double Time)
 {
-	SessionLock.WriteAccessCheck();
+	Session.WriteAccessCheck();
 	
 	uint64 Index = Frames[FrameType].Num();
 	FFrame& Frame = Frames[FrameType].PushBack();
@@ -54,9 +56,14 @@ void FFrameProvider::BeginFrame(ETraceFrameType FrameType, double Time)
 
 void FFrameProvider::EndFrame(ETraceFrameType FrameType, double Time)
 {
-	SessionLock.WriteAccessCheck();
+	Session.WriteAccessCheck();
 	FFrame& Frame = Frames[FrameType][Frames[FrameType].Num() - 1];
 	Frame.EndTime = Time;
 }
 
+const IFrameProvider& ReadFrameProvider(const IAnalysisSession& Session)
+{
+	return *Session.ReadProvider<IFrameProvider>(FFrameProvider::ProviderName);
+}
+
 }
\ No newline at end of file
diff --git a/Engine/Source/Developer/TraceServices/Private/Model/Frames.h b/Engine/Source/Developer/TraceServices/Private/Model/Frames.h
deleted file mode 100644
diff --git a/Engine/Source/Developer/TraceServices/Private/Model/FramesPrivate.h b/Engine/Source/Developer/TraceServices/Private/Model/FramesPrivate.h
new file mode 100644
--- /dev/null
+++ b/Engine/Source/Developer/TraceServices/Private/Model/FramesPrivate.h
@@ -0,0 +1,30 @@
+// Copyright 1998-2019 Epic Games, Inc. All Rights Reserved.
+
+#pragma once
+
+#include "TraceServices/AnalysisService.h"
+#include "ProfilingDebugging/MiscTrace.h"
+#include "Common/PagedArray.h"
+
+namespace Trace
+{
+
+class FFrameProvider
+	: public IFrameProvider
+{
+public:
+	static const FName ProviderName;
+
+	FFrameProvider(IAnalysisSession& Session);
+
+	virtual uint64 GetFrameCount(ETraceFrameType FrameType) const override;
+	virtual void EnumerateFrames(ETraceFrameType FrameType, uint64 Start, uint64 End, TFunctionRef<void(const FFrame&)> Callback) const override;
+	void BeginFrame(ETraceFrameType FrameType, double Time);
+	void EndFrame(ETraceFrameType FrameType, double Time);
+
+private:
+	IAnalysisSession& Session;
+	TArray<TPagedArray<FFrame>> Frames;
+};
+
+}
\ No newline at end of file
diff --git a/Engine/Source/Developer/TraceServices/Private/Model/IntervalTimeline.h b/Engine/Source/Developer/TraceServices/Private/Model/IntervalTimeline.h
--- a/Engine/Source/Developer/TraceServices/Private/Model/IntervalTimeline.h
+++ b/Engine/Source/Developer/TraceServices/Private/Model/IntervalTimeline.h
@@ -2,8 +2,9 @@
 
 #pragma once
 
-#include "TraceServices/AnalysisService.h"
+#include "TraceServices/Containers/Timelines.h"
 #include "Common/PagedArray.h"
+#include "TraceServices/Containers/Allocators.h"
 #include <limits>
 
 namespace Trace
@@ -22,7 +23,7 @@ class TIntervalTimeline
 	: public ITimeline<EventType>
 {
 public:
-	TIntervalTimeline(FSlabAllocator& Allocator)
+	TIntervalTimeline(ILinearAllocator& Allocator)
 		: Events(Allocator, SettingsType::EventsPerPage)
 	{
 	}
diff --git a/Engine/Source/Developer/TraceServices/Private/Model/LoadTimeProfiler.cpp b/Engine/Source/Developer/TraceServices/Private/Model/LoadTimeProfiler.cpp
--- a/Engine/Source/Developer/TraceServices/Private/Model/LoadTimeProfiler.cpp
+++ b/Engine/Source/Developer/TraceServices/Private/Model/LoadTimeProfiler.cpp
@@ -1,28 +1,27 @@
 // Copyright 1998-2019 Epic Games, Inc. All Rights Reserved.
 
-#include "Model/LoadTimeProfiler.h"
+#include "TraceServices/Model/LoadTimeProfiler.h"
+#include "Model/LoadTimeProfilerPrivate.h"
 #include "AnalysisServicePrivate.h"
 #include "Common/TimelineStatistics.h"
 
 namespace Trace
 {
 
-FLoadTimeProfilerProvider::FLoadTimeProfilerProvider(FSlabAllocator& InAllocator, FAnalysisSessionLock& InSessionLock, FStringStore& InStringStore)
-	: Allocator(InAllocator)
-	, SessionLock(InSessionLock)
-	, StringStore(InStringStore)
-	, ClassInfos(Allocator, 4096)
-	, Packages(Allocator, 4096)
-	, Exports(Allocator, 4090)
-	, MainThreadCpuTimeline(MakeShared<CpuTimelineInternal>(Allocator))
-	, AsyncLoadingThreadCpuTimeline(MakeShared<CpuTimelineInternal>(Allocator))
+FLoadTimeProfilerProvider::FLoadTimeProfilerProvider(IAnalysisSession& InSession)
+	: Session(InSession)
+	, ClassInfos(Session.GetLinearAllocator(), 4096)
+	, Packages(Session.GetLinearAllocator(), 4096)
+	, Exports(Session.GetLinearAllocator(), 4090)
+	, MainThreadCpuTimeline(MakeShared<CpuTimelineInternal>(Session.GetLinearAllocator()))
+	, AsyncLoadingThreadCpuTimeline(MakeShared<CpuTimelineInternal>(Session.GetLinearAllocator()))
 {
 	
 }
 
 void FLoadTimeProfilerProvider::EnumeratePackages(TFunctionRef<void(const FPackageInfo&)> Callback) const
 {
-	SessionLock.ReadAccessCheck();
+	Session.ReadAccessCheck();
 
 	auto Iterator = Packages.GetIteratorFromItem(0);
 	const FPackageInfo* Package = Iterator.GetCurrentItem();
@@ -35,13 +34,13 @@ void FLoadTimeProfilerProvider::EnumeratePackages(TFunctionRef<void(const FPacka
 
 void FLoadTimeProfilerProvider::ReadMainThreadCpuTimeline(TFunctionRef<void(const CpuTimeline &)> Callback) const
 {
-	SessionLock.ReadAccessCheck();
+	Session.ReadAccessCheck();
 	Callback(*MainThreadCpuTimeline);
 }
 
 void FLoadTimeProfilerProvider::ReadAsyncLoadingThreadCpuTimeline(TFunctionRef<void(const CpuTimeline &)> Callback) const
 {
-	SessionLock.ReadAccessCheck();
+	Session.ReadAccessCheck();
 	Callback(*AsyncLoadingThreadCpuTimeline);
 }
 
@@ -142,27 +141,27 @@ ITable<FLoadTimeProfilerAggregatedStats>* FLoadTimeProfilerProvider::CreateObjec
 
 const Trace::FClassInfo& FLoadTimeProfilerProvider::AddClassInfo(const TCHAR* ClassName)
 {
-	SessionLock.WriteAccessCheck();
+	Session.WriteAccessCheck();
 
 	FClassInfo& ClassInfo = ClassInfos.PushBack();
-	ClassInfo.Name = StringStore.Store(ClassName);
+	ClassInfo.Name = Session.StoreString(ClassName);
 	return ClassInfo;
 }
 
 Trace::FPackageInfo& FLoadTimeProfilerProvider::CreatePackage(const TCHAR* PackageName)
 {
-	SessionLock.WriteAccessCheck();
+	Session.WriteAccessCheck();
 
 	uint32 PackageId = Packages.Num();
 	FPackageInfo& Package = Packages.PushBack();
 	Package.Id = PackageId;
-	Package.Name = StringStore.Store(PackageName);
+	Package.Name = Session.StoreString(PackageName);
 	return Package;
 }
 
 Trace::FPackageExportInfo& FLoadTimeProfilerProvider::CreateExport()
 {
-	SessionLock.WriteAccessCheck();
+	Session.WriteAccessCheck();
 
 	uint32 ExportId = Exports.Num();
 	FPackageExportInfo& Export = Exports.PushBack();
diff --git a/Engine/Source/Developer/TraceServices/Private/Model/LoadTimeProfiler.h b/Engine/Source/Developer/TraceServices/Private/Model/LoadTimeProfiler.h
deleted file mode 100644
diff --git a/Engine/Source/Developer/TraceServices/Private/Model/LoadTimeProfilerPrivate.h b/Engine/Source/Developer/TraceServices/Private/Model/LoadTimeProfilerPrivate.h
new file mode 100644
--- /dev/null
+++ b/Engine/Source/Developer/TraceServices/Private/Model/LoadTimeProfilerPrivate.h
@@ -0,0 +1,56 @@
+// Copyright 1998-2019 Epic Games, Inc. All Rights Reserved.
+
+#pragma once
+
+#include "TraceServices/Model/LoadTimeProfiler.h"
+#include "TraceServices/AnalysisService.h"
+#include "Common/SlabAllocator.h"
+#include "Common/PagedArray.h"
+#include "Common/StringStore.h"
+#include "Model/MonotonicTimeline.h"
+#include "Model/Tables.h"
+
+namespace Trace
+{
+
+class FAnalysisSessionLock;
+
+class FLoadTimeProfilerProvider
+	: public ILoadTimeProfilerProvider
+{
+public:
+	typedef TMonotonicTimeline<FLoadTimeProfilerCpuEvent> CpuTimelineInternal;
+
+	FLoadTimeProfilerProvider(IAnalysisSession& Session);
+	virtual uint64 GetPackageCount() const override { return Packages.Num(); }
+	virtual void EnumeratePackages(TFunctionRef<void(const FPackageInfo&)> Callback) const override;
+	virtual void ReadMainThreadCpuTimeline(TFunctionRef<void(const CpuTimeline&)> Callback) const override;
+	virtual void ReadAsyncLoadingThreadCpuTimeline(TFunctionRef<void(const CpuTimeline&)> Callback) const override;
+	virtual ITable<FLoadTimeProfilerAggregatedStats>* CreateEventAggregation(double IntervalStart, double IntervalEnd) const override;
+	virtual ITable<FLoadTimeProfilerAggregatedStats>* CreateObjectTypeAggregation(double IntervalStart, double IntervalEnd) const override;
+	const FClassInfo& AddClassInfo(const TCHAR* ClassName);
+	FPackageInfo& CreatePackage(const TCHAR* PackageName);
+	FPackageExportInfo& CreateExport();
+	TSharedRef<CpuTimelineInternal> EditMainThreadCpuTimeline() { return MainThreadCpuTimeline; }
+	TSharedRef<CpuTimelineInternal> EditAsyncLoadingThreadCpuTimeline() { return AsyncLoadingThreadCpuTimeline; }
+
+private:
+	IAnalysisSession& Session;
+	TPagedArray<FClassInfo> ClassInfos;
+	TPagedArray<FPackageInfo> Packages;
+	TPagedArray<FPackageExportInfo> Exports;
+	TSharedRef<CpuTimelineInternal> MainThreadCpuTimeline;
+	TSharedRef<CpuTimelineInternal> AsyncLoadingThreadCpuTimeline;
+
+	UE_TRACE_TABLE_LAYOUT_BEGIN(FAggregatedStatsTableLayout, FLoadTimeProfilerAggregatedStats)
+		UE_TRACE_TABLE_COLUMN(Name, TEXT("Name"))
+		UE_TRACE_TABLE_COLUMN(Count, TEXT("Count"))
+		UE_TRACE_TABLE_COLUMN(Total, TEXT("Total"))
+		UE_TRACE_TABLE_COLUMN(Min, TEXT("Min"))
+		UE_TRACE_TABLE_COLUMN(Max, TEXT("Max"))
+		UE_TRACE_TABLE_COLUMN(Average, TEXT("Avg"))
+		UE_TRACE_TABLE_COLUMN(Median, TEXT("Med"))
+	UE_TRACE_TABLE_LAYOUT_END()
+};
+
+}
diff --git a/Engine/Source/Developer/TraceServices/Private/Model/Log.cpp b/Engine/Source/Developer/TraceServices/Private/Model/Log.cpp
--- a/Engine/Source/Developer/TraceServices/Private/Model/Log.cpp
+++ b/Engine/Source/Developer/TraceServices/Private/Model/Log.cpp
@@ -1,19 +1,20 @@
 // Copyright 1998-2019 Epic Games, Inc. All Rights Reserved.
 
-#include "Model/Log.h"
+#include "TraceServices/Model/Log.h"
+#include "Model/LogPrivate.h"
 #include "AnalysisServicePrivate.h"
 #include "Common/FormatArgs.h"
 
 namespace Trace
 {
 
-FLogProvider::FLogProvider(FSlabAllocator& InAllocator, FAnalysisSessionLock& InSessionLock, FStringStore& InStringStore)
-	: Allocator(InAllocator)
-	, SessionLock(InSessionLock)
-	, StringStore(InStringStore)
-	, Categories(Allocator, 128)
-	, MessageSpecs(Allocator, 1024)
-	, Messages(Allocator, 1024)
+const FName FLogProvider::ProviderName("LogProvider");
+
+FLogProvider::FLogProvider(IAnalysisSession& InSession)
+	: Session(InSession)
+	, Categories(InSession.GetLinearAllocator(), 128)
+	, MessageSpecs(InSession.GetLinearAllocator(), 1024)
+	, Messages(InSession.GetLinearAllocator(), 1024)
 	, MessagesTable(Messages)
 {
 
@@ -21,7 +22,7 @@ FLogProvider::FLogProvider(FSlabAllocator& InAllocator, FAnalysisSessionLock& In
 
 FLogCategory& FLogProvider::GetCategory(uint64 CategoryPointer)
 {
-	SessionLock.WriteAccessCheck();
+	Session.WriteAccessCheck();
 	if (CategoryMap.Contains(CategoryPointer))
 	{
 		return *CategoryMap[CategoryPointer];
@@ -36,7 +37,7 @@ FLogCategory& FLogProvider::GetCategory(uint64 CategoryPointer)
 
 FLogMessageSpec& FLogProvider::GetMessageSpec(uint64 LogPoint)
 {
-	SessionLock.WriteAccessCheck();
+	Session.WriteAccessCheck();
 	if (SpecMap.Contains(LogPoint))
 	{
 		return *SpecMap[LogPoint];
@@ -51,24 +52,24 @@ FLogMessageSpec& FLogProvider::GetMessageSpec(uint64 LogPoint)
 
 void FLogProvider::AppendMessage(uint64 LogPoint, double Time, const uint8* FormatArgs)
 {
-	SessionLock.WriteAccessCheck();
+	Session.WriteAccessCheck();
 	check(SpecMap.Contains(LogPoint));
 	FLogMessageInternal& InternalMessage = Messages.PushBack();
 	InternalMessage.Time = Time;
 	InternalMessage.Spec = SpecMap[LogPoint];
 	FFormatArgsHelper::Format(FormatBuffer, FormatBufferSize - 1, InternalMessage.Spec->FormatString, FormatArgs);
-	InternalMessage.Message = StringStore.Store(FormatBuffer);
+	InternalMessage.Message = Session.StoreString(FormatBuffer);
 }
 
 uint64 FLogProvider::GetMessageCount() const
 {
-	SessionLock.ReadAccessCheck();
+	Session.ReadAccessCheck();
 	return Messages.Num();
 }
 
 bool FLogProvider::ReadMessage(uint64 Index, TFunctionRef<void(const FLogMessage &)> Callback) const
 {
-	SessionLock.ReadAccessCheck();
+	Session.ReadAccessCheck();
 	if (Index >= Messages.Num())
 	{
 		return false;
@@ -79,7 +80,7 @@ bool FLogProvider::ReadMessage(uint64 Index, TFunctionRef<void(const FLogMessage
 
 void FLogProvider::EnumerateMessages(double IntervalStart, double IntervalEnd, TFunctionRef<void(const FLogMessage&)> Callback) const
 {
-	SessionLock.ReadAccessCheck();
+	Session.ReadAccessCheck();
 	if (IntervalStart > IntervalEnd)
 	{
 		return;
@@ -97,7 +98,7 @@ void FLogProvider::EnumerateMessages(double IntervalStart, double IntervalEnd, T
 
 void FLogProvider::EnumerateMessagesByIndex(uint64 Start, uint64 End, TFunctionRef<void(const FLogMessage &)> Callback) const
 {
-	SessionLock.ReadAccessCheck();
+	Session.ReadAccessCheck();
 	uint64 Count = Messages.Num();
 	if (Start >= Count)
 	{
@@ -133,6 +134,7 @@ void FLogProvider::ConstructMessage(uint64 Index, TFunctionRef<void(const FLogMe
 
 void FLogProvider::EnumerateCategories(TFunctionRef<void(const FLogCategory &)> Callback) const
 {
+	Session.ReadAccessCheck();
 	auto Iterator = Categories.GetIteratorFromItem(0);
 	const FLogCategory* Category = Iterator.GetCurrentItem();
 	while (Category)
@@ -142,4 +144,9 @@ void FLogProvider::EnumerateCategories(TFunctionRef<void(const FLogCategory &)>
 	}
 }
 
+const ILogProvider& ReadLogProvider(const IAnalysisSession& Session)
+{
+	return *Session.ReadProvider<ILogProvider>(FLogProvider::ProviderName);
+}
+
 }
\ No newline at end of file
diff --git a/Engine/Source/Developer/TraceServices/Private/Model/Log.h b/Engine/Source/Developer/TraceServices/Private/Model/Log.h
deleted file mode 100644
diff --git a/Engine/Source/Developer/TraceServices/Private/Model/LogPrivate.h b/Engine/Source/Developer/TraceServices/Private/Model/LogPrivate.h
new file mode 100644
--- /dev/null
+++ b/Engine/Source/Developer/TraceServices/Private/Model/LogPrivate.h
@@ -0,0 +1,85 @@
+// Copyright 1998-2019 Epic Games, Inc. All Rights Reserved.
+
+#pragma once
+
+#include "TraceServices/AnalysisService.h"
+#include "Templates/SharedPointer.h"
+#include "Common/PagedArray.h"
+#include "Model/Tables.h"
+#include "Misc/OutputDeviceHelper.h"
+
+namespace Trace
+{
+
+class FAnalysisSessionLock;
+class FStringStore;
+
+struct FLogMessageSpec
+{
+	FLogCategory* Category = nullptr;
+	const TCHAR* File = nullptr;
+	const TCHAR* FormatString = nullptr;
+	int32 Line;
+	ELogVerbosity::Type Verbosity;
+};
+
+struct FLogMessageInternal
+{
+	FLogMessageSpec* Spec = nullptr;
+	double Time;
+	const TCHAR* Message = nullptr;
+};
+
+class FLogProvider :
+	public ILogProvider
+{
+public:
+	enum
+	{
+		ReservedLogCategory_Bookmark = 0
+	};
+
+	static const FName ProviderName;
+
+	FLogProvider(IAnalysisSession& Session);
+
+	FLogCategory& GetCategory(uint64 CategoryPointer);
+	FLogMessageSpec& GetMessageSpec(uint64 LogPoint);
+	void AppendMessage(uint64 LogPoint, double Time, const uint8* FormatArgs);
+
+	virtual uint64 GetMessageCount() const override;
+	virtual bool ReadMessage(uint64 Index, TFunctionRef<void(const FLogMessage &)> Callback) const override;
+	virtual void EnumerateMessages(double IntervalStart, double IntervalEnd, TFunctionRef<void(const FLogMessage&)> Callback) const override;
+	virtual void EnumerateMessagesByIndex(uint64 Start, uint64 End, TFunctionRef<void(const FLogMessage&)> Callback) const override;
+	virtual uint64 GetCategoryCount() const override { return Categories.Num(); }
+	virtual void EnumerateCategories(TFunctionRef<void(const FLogCategory&)> Callback) const override;
+	virtual const IUntypedTable& GetMessagesTable() const override { return MessagesTable; }
+
+private:
+	void ConstructMessage(uint64 Id, TFunctionRef<void(const FLogMessage &)> Callback) const;
+
+	UE_TRACE_TABLE_LAYOUT_BEGIN(FMessagesTableLayout, FLogMessageInternal)
+		UE_TRACE_TABLE_COLUMN(Time, TEXT("Time"))
+		UE_TRACE_TABLE_PROJECTED_COLUMN(TableColumnType_CString, TEXT("Verbosity"), [](const FLogMessageInternal& Message) { return FOutputDeviceHelper::VerbosityToString(Message.Spec->Verbosity); })
+		UE_TRACE_TABLE_PROJECTED_COLUMN(TableColumnType_CString, TEXT("Category"), [](const FLogMessageInternal& Message) { return Message.Spec->Category->Name; })
+		UE_TRACE_TABLE_PROJECTED_COLUMN(TableColumnType_CString, TEXT("File"), [](const FLogMessageInternal& Message) { return Message.Spec->File; })
+		UE_TRACE_TABLE_PROJECTED_COLUMN(TableColumnType_Int, TEXT("Line"), [](const FLogMessageInternal& Message) { return Message.Spec->Line; })
+		UE_TRACE_TABLE_COLUMN(Message, TEXT("Message"))
+	UE_TRACE_TABLE_LAYOUT_END()
+
+	enum
+	{
+		FormatBufferSize = 65536
+	};
+
+	IAnalysisSession& Session;
+	TMap<uint64, FLogCategory*> CategoryMap;
+	TMap<uint64, FLogMessageSpec*> SpecMap;
+	TPagedArray<FLogCategory> Categories;
+	TPagedArray<FLogMessageSpec> MessageSpecs;
+	TPagedArray<FLogMessageInternal> Messages;
+	TCHAR FormatBuffer[FormatBufferSize];
+	TTableView<FMessagesTableLayout> MessagesTable;
+};
+
+}
\ No newline at end of file
diff --git a/Engine/Source/Developer/TraceServices/Private/Model/MonotonicTimeline.h b/Engine/Source/Developer/TraceServices/Private/Model/MonotonicTimeline.h
--- a/Engine/Source/Developer/TraceServices/Private/Model/MonotonicTimeline.h
+++ b/Engine/Source/Developer/TraceServices/Private/Model/MonotonicTimeline.h
@@ -32,7 +32,7 @@ class TMonotonicTimeline
 public:
 	using EventType = InEventType;
 
-	TMonotonicTimeline(FSlabAllocator& InAllocator)
+	TMonotonicTimeline(ILinearAllocator& InAllocator)
 		: Allocator(InAllocator)
 	{
 		
@@ -362,7 +362,7 @@ private:
 
 	struct FDetailLevel
 	{
-		FDetailLevel(FSlabAllocator& Allocator, double InResolution)
+		FDetailLevel(ILinearAllocator& Allocator, double InResolution)
 			: Resolution(InResolution)
 			, ScopeEntries(Allocator, SettingsType::ScopeEntriesPageSize)
 			, Events(Allocator, SettingsType::EventsPageSize)
@@ -418,7 +418,7 @@ private:
 			LastPage->InitialStackCount = DetailLevel.InsertionState.CurrentDepth;
 			if (LastPage->InitialStackCount)
 			{
-				LastPage->InitialStack = Allocator.Allocate<FEventStackEntry>(LastPage->InitialStackCount);
+				LastPage->InitialStack = reinterpret_cast<FEventStackEntry*>(Allocator.Allocate(LastPage->InitialStackCount * sizeof(FEventStackEntry)));
 				memcpy(LastPage->InitialStack, DetailLevel.InsertionState.EventStack, LastPage->InitialStackCount * sizeof(FEventStackEntry));
 			}
 		}
@@ -464,7 +464,7 @@ private:
 		return DetailLevel.Events[Index];
 	}
 
-	FSlabAllocator& Allocator;
+	ILinearAllocator& Allocator;
 	TArray<FDetailLevel> DetailLevels;
 	uint64 ModCount = 0;
 };
diff --git a/Engine/Source/Developer/TraceServices/Private/Model/Tables.h b/Engine/Source/Developer/TraceServices/Private/Model/Tables.h
--- a/Engine/Source/Developer/TraceServices/Private/Model/Tables.h
+++ b/Engine/Source/Developer/TraceServices/Private/Model/Tables.h
@@ -6,6 +6,7 @@
 #include "TraceServices/AnalysisService.h"
 #include "Templates/Function.h"
 #include "Common/PagedArray.h"
+#include "Common/SlabAllocator.h"
 
 #define UE_TRACE_TABLE_LAYOUT_BEGIN(Name, InRowType) \
 	class Name \
diff --git a/Engine/Source/Developer/TraceServices/Private/Model/Threads.cpp b/Engine/Source/Developer/TraceServices/Private/Model/Threads.cpp
--- a/Engine/Source/Developer/TraceServices/Private/Model/Threads.cpp
+++ b/Engine/Source/Developer/TraceServices/Private/Model/Threads.cpp
@@ -1,6 +1,7 @@
 // Copyright 1998-2019 Epic Games, Inc. All Rights Reserved.
 
-#include "Model/Threads.h"
+#include "TraceServices/Model/Threads.h"
+#include "Model/ThreadsPrivate.h"
 #include "Misc/ScopeLock.h"
 #include "AnalysisServicePrivate.h"
 #include "Common/StringStore.h"
@@ -8,9 +9,10 @@
 namespace Trace
 {
 
-FThreadProvider::FThreadProvider(const FAnalysisSessionLock& InSessionLock, FStringStore& InStringStore)
-	: SessionLock(InSessionLock)
-	, StringStore(InStringStore)
+const FName FThreadProvider::ProviderName = "ThreadProvider";
+
+FThreadProvider::FThreadProvider(IAnalysisSession& InSession)
+	: Session(InSession)
 {
 
 }
@@ -23,31 +25,15 @@ FThreadProvider::~FThreadProvider()
 	}
 }
 
-void FThreadProvider::EnsureThreadExists(uint32 Id)
-{
-	SessionLock.WriteAccessCheck();
-
-	if (!ThreadMap.Contains(Id))
-	{
-		FThreadInfoInternal* ThreadInfo = new FThreadInfoInternal();
-		ThreadInfo->Id = Id;
-		ThreadInfo->FallbackSortOrder = SortedThreads.Num();
-		SortedThreads.Add(ThreadInfo);
-		ThreadMap.Add(Id, ThreadInfo);
-		SortThreads();
-		++ModCount;
-	}
-}
-
 void FThreadProvider::AddGameThread(uint32 Id)
 {
-	SessionLock.WriteAccessCheck();
+	Session.WriteAccessCheck();
 
 	check(!ThreadMap.Contains(Id));
 	FThreadInfoInternal* ThreadInfo = new FThreadInfoInternal();
 	ThreadInfo->Id = Id;
 	ThreadInfo->PrioritySortOrder = GetPrioritySortOrder(TPri_Normal);
-	ThreadInfo->Name = StringStore.Store(*FName(NAME_GameThread).GetPlainNameString());
+	ThreadInfo->Name = Session.StoreString(*FName(NAME_GameThread).GetPlainNameString());
 	ThreadInfo->FallbackSortOrder = SortedThreads.Num();
 	ThreadInfo->IsGameThread = true;
 	SortedThreads.Add(ThreadInfo);
@@ -57,7 +43,7 @@ void FThreadProvider::AddGameThread(uint32 Id)
 
 void FThreadProvider::AddThread(uint32 Id, const TCHAR* Name, EThreadPriority Priority)
 {
-	SessionLock.WriteAccessCheck();
+	Session.WriteAccessCheck();
 
 	FThreadInfoInternal* ThreadInfo;
 	if (!ThreadMap.Contains(Id))
@@ -73,14 +59,14 @@ void FThreadProvider::AddThread(uint32 Id, const TCHAR* Name, EThreadPriority Pr
 		ThreadInfo = ThreadMap[Id];
 	}
 	ThreadInfo->PrioritySortOrder = GetPrioritySortOrder(Priority);
-	ThreadInfo->Name = StringStore.Store(Name);
+	ThreadInfo->Name = Session.StoreString(Name);
 	SortThreads();
 	++ModCount;
 }
 
 void FThreadProvider::SetThreadPriority(uint32 Id, EThreadPriority Priority)
 {
-	SessionLock.WriteAccessCheck();
+	Session.WriteAccessCheck();
 
 	check(ThreadMap.Contains(Id));
 	FThreadInfoInternal* ThreadInfo = ThreadMap[Id];
@@ -91,7 +77,7 @@ void FThreadProvider::SetThreadPriority(uint32 Id, EThreadPriority Priority)
 
 void FThreadProvider::SetThreadGroup(uint32 Id, ETraceThreadGroup Group)
 {
-	SessionLock.WriteAccessCheck();
+	Session.WriteAccessCheck();
 
 	check(ThreadMap.Contains(Id));
 	FThreadInfoInternal* ThreadInfo = ThreadMap[Id];
@@ -129,7 +115,7 @@ void FThreadProvider::SetThreadGroup(uint32 Id, ETraceThreadGroup Group)
 
 void FThreadProvider::EnumerateThreads(TFunctionRef<void(const FThreadInfo &)> Callback) const
 {
-	SessionLock.ReadAccessCheck();
+	Session.ReadAccessCheck();
 
 	for (const FThreadInfoInternal* Thread : SortedThreads)
 	{
@@ -194,4 +180,9 @@ bool FThreadProvider::FThreadInfoInternal::operator<(const FThreadInfoInternal&
 	return IsGameThread;
 }
 
+const IThreadProvider& ReadThreadProvider(const IAnalysisSession& Session)
+{
+	return *Session.ReadProvider<IThreadProvider>(FThreadProvider::ProviderName);
+}
+
 }
\ No newline at end of file
diff --git a/Engine/Source/Developer/TraceServices/Private/Model/Threads.h b/Engine/Source/Developer/TraceServices/Private/Model/Threads.h
deleted file mode 100644
diff --git a/Engine/Source/Developer/TraceServices/Private/Model/ThreadsPrivate.h b/Engine/Source/Developer/TraceServices/Private/Model/ThreadsPrivate.h
new file mode 100644
--- /dev/null
+++ b/Engine/Source/Developer/TraceServices/Private/Model/ThreadsPrivate.h
@@ -0,0 +1,51 @@
+// Copyright 1998-2019 Epic Games, Inc. All Rights Reserved.
+
+#pragma once
+
+#include "TraceServices/Model/AnalysisSession.h"
+#include "ProfilingDebugging/MiscTrace.h"
+#include "HAL/PlatformAffinity.h"
+#include "Containers/Map.h"
+
+namespace Trace
+{
+
+class FThreadProvider
+	: public IThreadProvider
+{
+public:
+	static const FName ProviderName;
+	FThreadProvider(IAnalysisSession& Session);
+	~FThreadProvider();
+	void AddGameThread(uint32 Id);
+	void AddThread(uint32 Id, const TCHAR* Name, EThreadPriority Priority);
+	void SetThreadPriority(uint32 Id, EThreadPriority Priority);
+	void SetThreadGroup(uint32 Id, ETraceThreadGroup Group);
+	virtual uint64 GetModCount() const override { return ModCount; }
+	virtual void EnumerateThreads(TFunctionRef<void(const FThreadInfo&)> Callback) const override;
+
+private:
+	struct FThreadInfoInternal
+	{
+		uint32 Id = 0;
+		const TCHAR* Name = nullptr;
+		uint32 GroupSortOrder = ~0u;
+		uint32 PrioritySortOrder = ~0u;
+		uint32 FallbackSortOrder = ~0u;
+		const TCHAR* GroupName = nullptr;
+		bool IsGameThread = false;
+
+		bool operator<(const FThreadInfoInternal& Other) const;
+	};
+
+	void SortThreads();
+	static uint32 GetPrioritySortOrder(EThreadPriority ThreadPriority);
+	static uint32 GetGroupSortOrder(ETraceThreadGroup ThreadGroup);
+
+	IAnalysisSession& Session;
+	uint64 ModCount = 0;
+	TMap<uint32, FThreadInfoInternal*> ThreadMap;
+	TArray<FThreadInfoInternal*> SortedThreads;
+};
+
+}
\ No newline at end of file
diff --git a/Engine/Source/Developer/TraceServices/Private/Model/TimingProfiler.cpp b/Engine/Source/Developer/TraceServices/Private/Model/TimingProfiler.cpp
--- a/Engine/Source/Developer/TraceServices/Private/Model/TimingProfiler.cpp
+++ b/Engine/Source/Developer/TraceServices/Private/Model/TimingProfiler.cpp
@@ -1,6 +1,7 @@
 // Copyright 1998-2019 Epic Games, Inc. All Rights Reserved.
 
-#include "Model/TimingProfiler.h"
+#include "TraceServices/Model/TimingProfiler.h"
+#include "Model/TimingProfilerPrivate.h"
 #include "AnalysisServicePrivate.h"
 #include "Common/StringStore.h"
 #include "Common/TimelineStatistics.h"
@@ -8,12 +9,10 @@
 namespace Trace
 {
 
-FTimingProfilerProvider::FTimingProfilerProvider(FSlabAllocator& InAllocator, FAnalysisSessionLock& InSessionLock, FStringStore& InStringStore)
-	: Allocator(InAllocator)
-	, SessionLock(InSessionLock)
-	, StringStore(InStringStore)
+FTimingProfilerProvider::FTimingProfilerProvider(IAnalysisSession& InSession)
+	: Session(InSession)
 {
-	Timelines.Add(MakeShared<TimelineInternal>(Allocator));
+	Timelines.Add(MakeShared<TimelineInternal>(Session.GetLinearAllocator()));
 }
 
 FTimingProfilerProvider::~FTimingProfilerProvider()
@@ -22,7 +21,7 @@ FTimingProfilerProvider::~FTimingProfilerProvider()
 
 uint32 FTimingProfilerProvider::AddCpuTimer(const TCHAR* Name)
 {
-	SessionLock.WriteAccessCheck();
+	Session.WriteAccessCheck();
 
 	FTimingProfilerTimer& Timer = AddTimerInternal(Name, false);
 	return Timer.Id;
@@ -30,7 +29,7 @@ uint32 FTimingProfilerProvider::AddCpuTimer(const TCHAR* Name)
 
 uint32 FTimingProfilerProvider::AddGpuTimer(const TCHAR* Name)
 {
-	SessionLock.WriteAccessCheck();
+	Session.WriteAccessCheck();
 
 	FTimingProfilerTimer& Timer = AddTimerInternal(Name, true);
 	return Timer.Id;
@@ -40,7 +39,7 @@ FTimingProfilerTimer& FTimingProfilerProvider::AddTimerInternal(const TCHAR* Nam
 {
 	FTimingProfilerTimer& Timer = Timers.AddDefaulted_GetRef();
 	Timer.Id = Timers.Num() - 1;
-	Timer.Name = StringStore.Store(Name);
+	Timer.Name = Session.StoreString(Name);
 	uint32 NameHash = 0;
 	for (const TCHAR* c = Name; *c; ++c)
 	{
@@ -53,11 +52,11 @@ FTimingProfilerTimer& FTimingProfilerProvider::AddTimerInternal(const TCHAR* Nam
 
 FTimingProfilerProvider::TimelineInternal& FTimingProfilerProvider::EditCpuThreadTimeline(uint32 ThreadId)
 {
-	SessionLock.WriteAccessCheck();
+	Session.WriteAccessCheck();
 
 	if (!CpuThreadTimelineIndexMap.Contains(ThreadId))
 	{
-		TSharedRef<TimelineInternal> Timeline = MakeShared<TimelineInternal>(Allocator);
+		TSharedRef<TimelineInternal> Timeline = MakeShared<TimelineInternal>(Session.GetLinearAllocator());
 		uint32 TimelineIndex = Timelines.Num();
 		CpuThreadTimelineIndexMap.Add(ThreadId, TimelineIndex);
 		Timelines.Add(Timeline);
@@ -72,14 +71,14 @@ FTimingProfilerProvider::TimelineInternal& FTimingProfilerProvider::EditCpuThrea
 
 FTimingProfilerProvider::TimelineInternal& FTimingProfilerProvider::EditGpuTimeline()
 {
-	SessionLock.WriteAccessCheck();
+	Session.WriteAccessCheck();
 
 	return Timelines[GpuTimelineIndex].Get();
 }
 
 bool FTimingProfilerProvider::GetCpuThreadTimelineIndex(uint32 ThreadId, uint32& OutTimelineIndex) const
 {
-	SessionLock.ReadAccessCheck();
+	Session.ReadAccessCheck();
 
 	if (CpuThreadTimelineIndexMap.Contains(ThreadId))
 	{
@@ -91,7 +90,7 @@ bool FTimingProfilerProvider::GetCpuThreadTimelineIndex(uint32 ThreadId, uint32&
 
 bool FTimingProfilerProvider::GetGpuTimelineIndex(uint32& OutTimelineIndex) const
 {
-	SessionLock.ReadAccessCheck();
+	Session.ReadAccessCheck();
 
 	OutTimelineIndex = GpuTimelineIndex;
 	return true;
@@ -99,7 +98,7 @@ bool FTimingProfilerProvider::GetGpuTimelineIndex(uint32& OutTimelineIndex) cons
 
 bool FTimingProfilerProvider::ReadTimeline(uint32 Index, TFunctionRef<void(const Timeline &)> Callback) const
 {
-	SessionLock.ReadAccessCheck();
+	Session.ReadAccessCheck();
 
 	if (Index < uint32(Timelines.Num()))
 	{
@@ -114,7 +113,7 @@ bool FTimingProfilerProvider::ReadTimeline(uint32 Index, TFunctionRef<void(const
 
 void FTimingProfilerProvider::EnumerateTimelines(TFunctionRef<void(const Timeline&)> Callback) const
 {
-	SessionLock.ReadAccessCheck();
+	Session.ReadAccessCheck();
 
 	for (const auto& Timeline : Timelines)
 	{
@@ -124,14 +123,14 @@ void FTimingProfilerProvider::EnumerateTimelines(TFunctionRef<void(const Timelin
 
 void FTimingProfilerProvider::ReadTimers(TFunctionRef<void(const FTimingProfilerTimer*, uint64)> Callback) const
 {
-	SessionLock.ReadAccessCheck();
+	Session.ReadAccessCheck();
 
 	Callback(Timers.GetData(), Timers.Num());
 }
 
 ITable<FTimingProfilerAggregatedStats>* FTimingProfilerProvider::CreateAggregation(double IntervalStart, double IntervalEnd, TFunctionRef<bool(uint32)> CpuThreadFilter, bool IncludeGpu) const
 {
-	SessionLock.ReadAccessCheck();
+	Session.ReadAccessCheck();
 
 	TArray<const TimelineInternal*> IncludedTimelines;
 	if (IncludeGpu)
diff --git a/Engine/Source/Developer/TraceServices/Private/Model/TimingProfiler.h b/Engine/Source/Developer/TraceServices/Private/Model/TimingProfiler.h
deleted file mode 100644
diff --git a/Engine/Source/Developer/TraceServices/Private/Model/TimingProfilerPrivate.h b/Engine/Source/Developer/TraceServices/Private/Model/TimingProfilerPrivate.h
new file mode 100644
--- /dev/null
+++ b/Engine/Source/Developer/TraceServices/Private/Model/TimingProfilerPrivate.h
@@ -0,0 +1,61 @@
+// Copyright 1998-2019 Epic Games, Inc. All Rights Reserved.
+
+#pragma once
+
+#include "TraceServices/Model/TimingProfiler.h"
+#include "Common/SlabAllocator.h"
+#include "Model/MonotonicTimeline.h"
+#include "Model/Tables.h"
+
+namespace Trace
+{
+
+class FAnalysisSessionLock;
+class FStringStore;
+
+class FTimingProfilerProvider
+	: public ITimingProfilerProvider
+{
+public:
+	typedef TMonotonicTimeline<FTimingProfilerEvent> TimelineInternal;
+
+	FTimingProfilerProvider(IAnalysisSession& InSession);
+	virtual ~FTimingProfilerProvider();
+	uint32 AddCpuTimer(const TCHAR* Name);
+	uint32 AddGpuTimer(const TCHAR* Name);
+	TimelineInternal& EditCpuThreadTimeline(uint32 ThreadId);
+	TimelineInternal& EditGpuTimeline();
+	virtual bool GetCpuThreadTimelineIndex(uint32 ThreadId, uint32& OutTimelineIndex) const override;
+	virtual bool GetGpuTimelineIndex(uint32& OutTimelineIndex) const override;
+	virtual bool ReadTimeline(uint32 Index, TFunctionRef<void(const Timeline&)> Callback) const override;
+	virtual uint64 GetTimelineCount() const override { return Timelines.Num(); }
+	virtual void EnumerateTimelines(TFunctionRef<void(const Timeline&)> Callback) const override;
+	virtual void ReadTimers(TFunctionRef<void(const FTimingProfilerTimer*, uint64)> Callback) const override;
+	virtual ITable<FTimingProfilerAggregatedStats>* CreateAggregation(double IntervalStart, double IntervalEnd, TFunctionRef<bool(uint32)> CpuThreadFilter, bool IncludeGpu) const override;
+
+private:
+	FTimingProfilerTimer& AddTimerInternal(const TCHAR* Name, bool IsGpuEvent);
+
+	IAnalysisSession& Session;
+	TArray<FTimingProfilerTimer> Timers;
+	TArray<TSharedRef<TimelineInternal>> Timelines;
+	TMap<uint32, uint32> CpuThreadTimelineIndexMap;
+	uint32 GpuTimelineIndex = 0;
+
+	UE_TRACE_TABLE_LAYOUT_BEGIN(FAggregatedStatsTableLayout, FTimingProfilerAggregatedStats)
+		UE_TRACE_TABLE_PROJECTED_COLUMN(TableColumnType_CString, TEXT("Name"), [](const FTimingProfilerAggregatedStats& Row) { return Row.Timer->Name; })
+		UE_TRACE_TABLE_COLUMN(InstanceCount, TEXT("Count"))
+		UE_TRACE_TABLE_COLUMN(TotalInclusiveTime, TEXT("Incl"))
+		UE_TRACE_TABLE_COLUMN(MinInclusiveTime, TEXT("I.Min"))
+		UE_TRACE_TABLE_COLUMN(MaxInclusiveTime, TEXT("I.Max"))
+		UE_TRACE_TABLE_COLUMN(AverageInclusiveTime, TEXT("I.Avg"))
+		UE_TRACE_TABLE_COLUMN(MedianInclusiveTime, TEXT("I.Med"))
+		UE_TRACE_TABLE_COLUMN(TotalExclusiveTime, TEXT("Excl"))
+		UE_TRACE_TABLE_COLUMN(MinExclusiveTime, TEXT("E.Min"))
+		UE_TRACE_TABLE_COLUMN(MaxExclusiveTime, TEXT("E.Max"))
+		UE_TRACE_TABLE_COLUMN(AverageExclusiveTime, TEXT("E.Avg"))
+		UE_TRACE_TABLE_COLUMN(MedianExclusiveTime, TEXT("E.Med"))
+	UE_TRACE_TABLE_LAYOUT_END()
+};
+
+}
\ No newline at end of file
diff --git a/Engine/Source/Developer/TraceServices/Private/ModuleService.cpp b/Engine/Source/Developer/TraceServices/Private/ModuleService.cpp
new file mode 100644
--- /dev/null
+++ b/Engine/Source/Developer/TraceServices/Private/ModuleService.cpp
@@ -0,0 +1,78 @@
+// Copyright 1998-2019 Epic Games, Inc. All Rights Reserved.
+
+#include "TraceServices/ModuleService.h"
+#include "ModuleServicePrivate.h"
+#include "UObject/NameTypes.h"
+#include "Features/IModularFeatures.h"
+
+namespace Trace
+{
+
+const FName ModuleFeatureName("TraceModuleFeature");
+
+FModuleService::FModuleService()
+{
+
+}
+
+void FModuleService::Initialize()
+{
+	if (bIsInitialized)
+	{
+		return;
+	}
+	TArray<IModule*> Modules = IModularFeatures::Get().GetModularFeatureImplementations<IModule>(ModuleFeatureName);
+	for (IModule* Module : Modules)
+	{
+		FModuleInfo ModuleInfo;
+		Module->GetModuleInfo(ModuleInfo);
+		ModulesMap.Add(ModuleInfo.Name, Module);
+	}
+
+	bIsInitialized = true;
+}
+
+void FModuleService::GetAvailableModules(TArray<FModuleInfo>& OutModules)
+{
+	FScopeLock Lock(&CriticalSection);
+	Initialize();
+	OutModules.Empty(ModulesMap.Num());
+	for (const auto& KV : ModulesMap)
+	{
+		IModule* Module = KV.Value;
+		FModuleInfo& ModuleInfo = OutModules.AddDefaulted_GetRef();
+		Module->GetModuleInfo(ModuleInfo);
+	}
+}
+	
+void FModuleService::SetModuleEnabled(const FName& ModuleName, bool bEnabled)
+{
+	FScopeLock Lock(&CriticalSection);
+	Initialize();
+	IModule** FindIt = ModulesMap.Find(ModuleName);
+	if (!FindIt)
+	{
+		return;
+	}
+	if (bEnabled)
+	{
+		EnabledModules.Add(*FindIt);
+	}
+	else
+	{
+		EnabledModules.Remove(*FindIt);
+	}
+}
+
+void FModuleService::OnAnalysisBegin(IAnalysisSession& Session, TArray<IAnalyzer*>& OutAnalyzers)
+{
+	FScopeLock Lock(&CriticalSection);
+	Initialize();
+	for (const auto& KV : ModulesMap)
+	{
+		IModule* Module = KV.Value;
+		Module->OnAnalysisBegin(Session, EnabledModules.Contains(Module), OutAnalyzers);
+	}
+}
+
+}
diff --git a/Engine/Source/Developer/TraceServices/Private/ModuleServicePrivate.h b/Engine/Source/Developer/TraceServices/Private/ModuleServicePrivate.h
new file mode 100644
--- /dev/null
+++ b/Engine/Source/Developer/TraceServices/Private/ModuleServicePrivate.h
@@ -0,0 +1,33 @@
+// Copyright 1998-2019 Epic Games, Inc. All Rights Reserved.
+
+#pragma once
+
+#include "TraceServices/ModuleService.h"
+#include "Containers/Map.h"
+#include "Containers/Set.h"
+#include "Misc/ScopeLock.h"
+
+namespace Trace
+{
+
+class IAnalysisSession;
+
+class FModuleService
+	: public IModuleService
+{
+public:
+	FModuleService();
+	virtual void GetAvailableModules(TArray<FModuleInfo>& OutModules) override;
+	virtual void SetModuleEnabled(const FName& ModuleName, bool bEnabled) override;
+	void OnAnalysisBegin(IAnalysisSession& Session, TArray<IAnalyzer*>& OutAnalyzers);
+
+private:
+	void Initialize();
+
+	mutable FCriticalSection CriticalSection;
+	bool bIsInitialized = false;
+	TSet<IModule*> EnabledModules;
+	TMap<FName, IModule*> ModulesMap;
+};
+
+}
diff --git a/Engine/Source/Developer/TraceServices/Private/Modules/LoadTimeProfilerModule.cpp b/Engine/Source/Developer/TraceServices/Private/Modules/LoadTimeProfilerModule.cpp
new file mode 100644
--- /dev/null
+++ b/Engine/Source/Developer/TraceServices/Private/Modules/LoadTimeProfilerModule.cpp
@@ -0,0 +1,43 @@
+// Copyright 1998-2019 Epic Games, Inc. All Rights Reserved.
+
+#include "LoadTimeProfilerModule.h"
+#include "Analyzers/PlatformFileTraceAnalysis.h"
+#include "Analyzers/LoadTimeTraceAnalysis.h"
+#include "AnalysisServicePrivate.h"
+#include "Model/FileActivity.h"
+
+namespace Trace
+{
+
+FName FLoadTimeProfilerModule::ModuleName("TraceModule_LoadTimeProfiler");
+
+void FLoadTimeProfilerModule::GetModuleInfo(FModuleInfo& OutModuleInfo)
+{
+	OutModuleInfo.Name = ModuleName;
+	OutModuleInfo.DisplayName = TEXT("Asset Loading");
+}
+
+static const FName LoadTimeProfilerProviderName("LoadTimeProfiler");
+static const FName FileActivityProviderName("FileActivity");
+
+void FLoadTimeProfilerModule::OnAnalysisBegin(IAnalysisSession& Session, bool bIsEnabled, TArray<IAnalyzer*>& OutAnalyzers)
+{
+	FLoadTimeProfilerProvider* LoadTimeProfilerProvider = new FLoadTimeProfilerProvider(Session);
+	Session.AddProvider(LoadTimeProfilerProviderName, LoadTimeProfilerProvider);
+	OutAnalyzers.Add(new FAsyncLoadingTraceAnalyzer(Session, *LoadTimeProfilerProvider));
+	FFileActivityProvider* FileActivityProvider = new FFileActivityProvider(Session);
+	Session.AddProvider(FileActivityProviderName, FileActivityProvider);
+	OutAnalyzers.Add(new FPlatformFileTraceAnalyzer(Session, *FileActivityProvider));
+}
+
+const ILoadTimeProfilerProvider* ReadLoadTimeProfilerProvider(const IAnalysisSession& Session)
+{
+	return Session.ReadProvider<ILoadTimeProfilerProvider>(LoadTimeProfilerProviderName);
+}
+
+const IFileActivityProvider* ReadFileActivityProvider(const IAnalysisSession& Session)
+{
+	return Session.ReadProvider<IFileActivityProvider>(FileActivityProviderName);
+}
+
+}
diff --git a/Engine/Source/Developer/TraceServices/Private/Modules/LoadTimeProfilerModule.h b/Engine/Source/Developer/TraceServices/Private/Modules/LoadTimeProfilerModule.h
new file mode 100644
--- /dev/null
+++ b/Engine/Source/Developer/TraceServices/Private/Modules/LoadTimeProfilerModule.h
@@ -0,0 +1,21 @@
+// Copyright 1998-2019 Epic Games, Inc. All Rights Reserved.
+
+#pragma once
+
+#include "TraceServices/ModuleService.h"
+
+namespace Trace
+{
+
+class FLoadTimeProfilerModule
+	: public IModule
+{
+public:
+	virtual void GetModuleInfo(FModuleInfo& OutModuleInfo) override;
+	virtual void OnAnalysisBegin(IAnalysisSession& Session, bool bIsEnabled, TArray<IAnalyzer*>& OutAnalyzers) override;
+
+private:
+	static FName ModuleName;
+};
+
+}
diff --git a/Engine/Source/Developer/TraceServices/Private/Modules/StatsModule.cpp b/Engine/Source/Developer/TraceServices/Private/Modules/StatsModule.cpp
new file mode 100644
--- /dev/null
+++ b/Engine/Source/Developer/TraceServices/Private/Modules/StatsModule.cpp
@@ -0,0 +1,35 @@
+// Copyright 1998-2019 Epic Games, Inc. All Rights Reserved.
+
+#include "StatsModule.h"
+#include "Analyzers/StatsTraceAnalysis.h"
+#include "AnalysisServicePrivate.h"
+#include "TraceServices/ModuleService.h"
+#include "Model/Counters.h"
+
+namespace Trace
+{
+
+FName FStatsModule::ModuleName("TraceModule_Stats");
+
+void FStatsModule::GetModuleInfo(FModuleInfo& OutModuleInfo)
+{
+	OutModuleInfo.Name = ModuleName;
+	OutModuleInfo.DisplayName = TEXT("Stats");
+}
+	
+static const FName CounterProviderName("CounterProvider");
+
+void FStatsModule::OnAnalysisBegin(IAnalysisSession& InSession, bool bIsEnabled, TArray<IAnalyzer*>& OutAnalyzers)
+{
+	FAnalysisSession& Session = static_cast<FAnalysisSession&>(InSession);
+	FCounterProvider* CounterProvider = new FCounterProvider(InSession);
+	InSession.AddProvider(CounterProviderName, CounterProvider);
+	OutAnalyzers.Add(new FStatsAnalyzer(Session, *CounterProvider));
+}
+
+const ICounterProvider* ReadCounterProvider(const IAnalysisSession& Session)
+{
+	return Session.ReadProvider<ICounterProvider>(CounterProviderName);
+}
+
+}
diff --git a/Engine/Source/Developer/TraceServices/Private/Modules/StatsModule.h b/Engine/Source/Developer/TraceServices/Private/Modules/StatsModule.h
new file mode 100644
--- /dev/null
+++ b/Engine/Source/Developer/TraceServices/Private/Modules/StatsModule.h
@@ -0,0 +1,22 @@
+// Copyright 1998-2019 Epic Games, Inc. All Rights Reserved.
+
+#pragma once
+
+#include "TraceServices/ModuleService.h"
+#include "UObject/NameTypes.h"
+
+namespace Trace
+{
+
+class FStatsModule
+	: public IModule
+{
+public:
+	virtual void GetModuleInfo(FModuleInfo& OutModuleInfo) override;
+	virtual void OnAnalysisBegin(IAnalysisSession& Session, bool bIsEnabled, TArray<IAnalyzer*>& OutAnalyzers) override;
+
+private:
+	static FName ModuleName;
+};
+
+}
diff --git a/Engine/Source/Developer/TraceServices/Private/Modules/TimingProfilerModule.cpp b/Engine/Source/Developer/TraceServices/Private/Modules/TimingProfilerModule.cpp
new file mode 100644
--- /dev/null
+++ b/Engine/Source/Developer/TraceServices/Private/Modules/TimingProfilerModule.cpp
@@ -0,0 +1,37 @@
+// Copyright 1998-2019 Epic Games, Inc. All Rights Reserved.
+
+#include "TimingProfilerModule.h"
+#include "Analyzers/CpuProfilerTraceAnalysis.h"
+#include "Analyzers/GpuProfilerTraceAnalysis.h"
+#include "AnalysisServicePrivate.h"
+#include "Model/TimingProfilerPrivate.h"
+
+namespace Trace
+{
+
+FName FTimingProfilerModule::ModuleName("TraceModule_TimingProfiler");
+
+void FTimingProfilerModule::GetModuleInfo(FModuleInfo& OutModuleInfo)
+{
+	OutModuleInfo.Name = ModuleName;
+	OutModuleInfo.DisplayName = TEXT("Timing");
+}
+	
+static const FName TimingProfilerProviderName("TimingProfilerProvider");
+
+void FTimingProfilerModule::OnAnalysisBegin(IAnalysisSession& InSession, bool bIsEnabled, TArray<IAnalyzer*>& OutAnalyzers)
+{
+	FAnalysisSession& Session = static_cast<FAnalysisSession&>(InSession);
+	
+	FTimingProfilerProvider* TimingProfilerProvider = new FTimingProfilerProvider(Session);
+	Session.AddProvider(TimingProfilerProviderName, TimingProfilerProvider);
+	OutAnalyzers.Add(new FCpuProfilerAnalyzer(Session, *TimingProfilerProvider));
+	OutAnalyzers.Add(new FGpuProfilerAnalyzer(Session, *TimingProfilerProvider));
+}
+
+const ITimingProfilerProvider* ReadTimingProfilerProvider(const IAnalysisSession& Session)
+{
+	return Session.ReadProvider<ITimingProfilerProvider>(TimingProfilerProviderName);
+}
+
+}
diff --git a/Engine/Source/Developer/TraceServices/Private/Modules/TimingProfilerModule.h b/Engine/Source/Developer/TraceServices/Private/Modules/TimingProfilerModule.h
new file mode 100644
--- /dev/null
+++ b/Engine/Source/Developer/TraceServices/Private/Modules/TimingProfilerModule.h
@@ -0,0 +1,22 @@
+// Copyright 1998-2019 Epic Games, Inc. All Rights Reserved.
+
+#pragma once
+
+#include "TraceServices/ModuleService.h"
+#include "UObject/NameTypes.h"
+
+namespace Trace
+{
+
+class FTimingProfilerModule
+	: public IModule
+{
+public:
+	virtual void GetModuleInfo(FModuleInfo& OutModuleInfo) override;
+	virtual void OnAnalysisBegin(IAnalysisSession& Session, bool bIsEnabled, TArray<IAnalyzer*>& OutAnalyzers) override;
+
+private:
+	static FName ModuleName;
+};
+
+}
diff --git a/Engine/Source/Developer/TraceServices/Private/TraceServicesModule.cpp b/Engine/Source/Developer/TraceServices/Private/TraceServicesModule.cpp
--- a/Engine/Source/Developer/TraceServices/Private/TraceServicesModule.cpp
+++ b/Engine/Source/Developer/TraceServices/Private/TraceServicesModule.cpp
@@ -4,6 +4,11 @@
 #include "TraceServices/ITraceServicesModule.h"
 #include "SessionServicePrivate.h"
 #include "AnalysisServicePrivate.h"
+#include "ModuleServicePrivate.h"
+#include "Features/IModularFeatures.h"
+#include "Modules/TimingProfilerModule.h"
+#include "Modules/LoadTimeProfilerModule.h"
+#include "Modules/StatsModule.h"
 
 class FTraceServicesModule
 	: public ITraceServicesModule
@@ -11,10 +16,19 @@ class FTraceServicesModule
 public:
 	virtual TSharedPtr<Trace::ISessionService> GetSessionService() override;
 	virtual TSharedPtr<Trace::IAnalysisService> GetAnalysisService() override;
+	virtual TSharedPtr<Trace::IModuleService> GetModuleService() override;
+
+	virtual void StartupModule() override;
+	virtual void ShutdownModule() override;
 
 private:
 	TSharedPtr<Trace::FSessionService> SessionService;
 	TSharedPtr<Trace::FAnalysisService> AnalysisService;
+	TSharedPtr<Trace::FModuleService> ModuleService;
+
+	Trace::FTimingProfilerModule TimingProfilerModule;
+	Trace::FLoadTimeProfilerModule LoadTimeProfilerModule;
+	Trace::FStatsModule StatsModule;
 };
 
 TSharedPtr<Trace::ISessionService> FTraceServicesModule::GetSessionService()
@@ -30,9 +44,33 @@ TSharedPtr<Trace::IAnalysisService> FTraceServicesModule::GetAnalysisService()
 {
 	if (!AnalysisService.IsValid())
 	{
-		AnalysisService = MakeShared<Trace::FAnalysisService>();
+		GetModuleService();
+		AnalysisService = MakeShared<Trace::FAnalysisService>(*ModuleService.Get());
 	}
 	return AnalysisService;
 }
 
+TSharedPtr<Trace::IModuleService> FTraceServicesModule::GetModuleService()
+{
+	if (!ModuleService.IsValid())
+	{
+		ModuleService = MakeShared<Trace::FModuleService>();
+	}
+	return ModuleService;
+}
+
+void FTraceServicesModule::StartupModule()
+{
+	IModularFeatures::Get().RegisterModularFeature(Trace::ModuleFeatureName, &TimingProfilerModule);
+	IModularFeatures::Get().RegisterModularFeature(Trace::ModuleFeatureName, &LoadTimeProfilerModule);
+	IModularFeatures::Get().RegisterModularFeature(Trace::ModuleFeatureName, &StatsModule);
+}
+
+void FTraceServicesModule::ShutdownModule()
+{
+	IModularFeatures::Get().UnregisterModularFeature(Trace::ModuleFeatureName, &StatsModule);
+	IModularFeatures::Get().UnregisterModularFeature(Trace::ModuleFeatureName, &LoadTimeProfilerModule);
+	IModularFeatures::Get().UnregisterModularFeature(Trace::ModuleFeatureName, &TimingProfilerModule);
+}
+
 IMPLEMENT_MODULE(FTraceServicesModule, TraceServices)
diff --git a/Engine/Source/Developer/TraceServices/Public/TraceServices/AnalysisService.h b/Engine/Source/Developer/TraceServices/Public/TraceServices/AnalysisService.h
--- a/Engine/Source/Developer/TraceServices/Public/TraceServices/AnalysisService.h
+++ b/Engine/Source/Developer/TraceServices/Public/TraceServices/AnalysisService.h
@@ -3,376 +3,22 @@
 #pragma once
 
 #include "CoreTypes.h"
-#include "Containers/UnrealString.h"
-#include "ProfilingDebugging/MiscTrace.h"
 #include "Templates/SharedPointer.h"
-#include "Trace/Store.h"
-#include "Delegates/Delegate.h"
-#include "GenericPlatform/GenericPlatformAffinity.h"
-#include "Serialization/LoadTimeTrace.h"
 #include "Templates/UniquePtr.h"
+#include "Delegates/Delegate.h"
+#include "Trace/DataStream.h"
+#include "TraceServices/Model/AnalysisSession.h"
+#include "TraceServices/Model/Log.h"
+#include "TraceServices/Model/Bookmarks.h"
+#include "TraceServices/Model/Frames.h"
+#include "TraceServices/Model/Threads.h"
+#include "TraceServices/Model/TimingProfiler.h"
+#include "TraceServices/Model/LoadTimeProfiler.h"
+#include "TraceServices/Model/Stats.h"
 
 namespace Trace
 {
 
-enum ETableColumnType
-{
-	TableColumnType_Invalid,
-	TableColumnType_Bool,
-	TableColumnType_Int,
-	TableColumnType_Float,
-	TableColumnType_Double,
-	TableColumnType_CString,
-};
-
-class ITableLayout
-{
-public:
-	virtual ~ITableLayout() = default;
-	virtual uint8 GetColumnCount() const = 0;
-	virtual const TCHAR* GetColumnName(uint8 ColumnIndex) const = 0;
-	virtual ETableColumnType GetColumnType(uint8 ColumnIndex) const = 0;
-};
-
-class IUntypedTableReader
-{
-public:
-	virtual ~IUntypedTableReader() = default;
-	virtual bool IsValid() const = 0;
-	virtual void NextRow() = 0;
-	virtual void SetRowIndex(uint64 RowIndex) = 0;
-	virtual bool GetValueBool(uint8 ColumnIndex) const = 0;
-	virtual int64 GetValueInt(uint8 ColumnIndex) const = 0;
-	virtual float GetValueFloat(uint8 ColumnIndex) const = 0;
-	virtual double GetValueDouble(uint8 ColumnIndex) const = 0;
-	virtual const TCHAR* GetValueCString(uint8 ColumnIndex) const = 0;
-};
-
-template<typename RowType>
-class ITableReader
-	: public IUntypedTableReader
-{
-public:
-	virtual const RowType* GetCurrentRow() const = 0;
-};
-
-class IUntypedTable
-{
-public:
-	virtual ~IUntypedTable() = default;
-	virtual const ITableLayout& GetLayout() const = 0;
-	virtual uint64 GetRowCount() const = 0;
-	virtual IUntypedTableReader* CreateReader() const = 0;
-};
-
-template<typename RowType>
-class ITable
-	: public IUntypedTable
-{
-public:
-	virtual ~ITable() = default;
-	virtual ITableReader<RowType>* CreateReader() const = 0;
-};
-
-struct FBookmark
-{
-	double Time;
-	const TCHAR* Text;
-};
-
-class IBookmarkProvider
-{
-public:
-	virtual ~IBookmarkProvider() = default;
-	virtual uint64 GetBookmarkCount() const = 0;
-	virtual void EnumerateBookmarks(double IntervalStart, double IntervalEnd, TFunctionRef<void(const FBookmark&)> Callback) const = 0;
-};
-
-struct FLogCategory
-{
-	const TCHAR* Name = nullptr;
-	ELogVerbosity::Type DefaultVerbosity;
-};
-
-struct FLogMessage
-{
-	uint64 Index;
-	double Time;
-	const FLogCategory* Category = nullptr;
-	const TCHAR* File = nullptr;
-	const TCHAR* Message = nullptr;
-	int32 Line;
-	ELogVerbosity::Type Verbosity;
-};
-
-class ILogProvider
-{
-public:
-	virtual ~ILogProvider() = default;
-	virtual uint64 GetMessageCount() const = 0;
-	virtual void EnumerateMessages(double IntervalStart, double IntervalEnd, TFunctionRef<void(const FLogMessage&)> Callback) const = 0;
-	virtual void EnumerateMessagesByIndex(uint64 Start, uint64 End, TFunctionRef<void(const FLogMessage&)> Callback) const = 0;
-	virtual bool ReadMessage(uint64 Index, TFunctionRef<void(const FLogMessage&)> Callback) const = 0;
-	virtual uint64 GetCategoryCount() const = 0;
-	virtual void EnumerateCategories(TFunctionRef<void(const FLogCategory&)> Callback) const = 0;
-	virtual const IUntypedTable& GetMessagesTable() const = 0;
-};
-
-struct FThreadInfo
-{
-	uint32 Id;
-	const TCHAR* Name;
-	const TCHAR* GroupName;
-};
-
-class IThreadProvider
-{
-public:
-	virtual ~IThreadProvider() = default;
-	virtual uint64 GetModCount() const = 0;
-	virtual void EnumerateThreads(TFunctionRef<void(const FThreadInfo&)> Callback) const = 0;
-};
-
-template<typename InEventType>
-class ITimeline
-{
-public:
-	typedef InEventType EventType;
-
-	virtual ~ITimeline() = default;
-	virtual uint64 GetModCount() const = 0;
-	virtual uint64 GetEventCount() const = 0;
-	virtual void EnumerateEventsDownSampled(double IntervalStart, double IntervalEnd, double Resolution, TFunctionRef<void(bool, double, const EventType&)> Callback) const = 0;
-	virtual void EnumerateEventsDownSampled(double IntervalStart, double IntervalEnd, double Resolution, TFunctionRef<void(double, double, uint32, const EventType&)> Callback) const = 0;
-	virtual void EnumerateEvents(double IntervalStart, double IntervalEnd, TFunctionRef<void(bool, double, const EventType&)> Callback) const = 0;
-	virtual void EnumerateEvents(double IntervalStart, double IntervalEnd, TFunctionRef<void(double, double, uint32, const EventType&)> Callback) const = 0;
-};
-
-struct FTimingProfilerTimer
-{
-	const TCHAR* Name;
-	uint32 Id;
-	uint32 NameHash;
-	bool IsGpuTimer;
-};
-
-struct FTimingProfilerEvent
-{
-	uint32 TimerIndex;
-};
-
-struct FAggregatedTimingStats
-{
-	uint64 InstanceCount = 0;
-	double TotalInclusiveTime = 0.0;
-	double MinInclusiveTime = DBL_MAX;
-	double MaxInclusiveTime = -DBL_MAX;
-	double AverageInclusiveTime = 0.0;
-	double MedianInclusiveTime = 0.0;
-	double TotalExclusiveTime = 0.0;
-	double MinExclusiveTime = DBL_MAX;
-	double MaxExclusiveTime = -DBL_MAX;
-	double AverageExclusiveTime = 0.0;
-	double MedianExclusiveTime = 0.0;
-};
-
-struct FTimingProfilerAggregatedStats
-	: public FAggregatedTimingStats
-{
-	const FTimingProfilerTimer* Timer = nullptr;
-};
-
-class ITimingProfilerProvider
-{
-public:
-	typedef ITimeline<FTimingProfilerEvent> Timeline;
-
-	virtual ~ITimingProfilerProvider() = default;
-	virtual bool GetCpuThreadTimelineIndex(uint32 ThreadId, uint32& OutTimelineIndex) const = 0;
-	virtual bool GetGpuTimelineIndex(uint32& OutTimelineIndex) const = 0;
-	virtual bool ReadTimeline(uint32 Index, TFunctionRef<void(const Timeline&)> Callback) const = 0;
-	virtual uint64 GetTimelineCount() const = 0;
-	virtual void EnumerateTimelines(TFunctionRef<void(const Timeline&)> Callback) const = 0;
-	virtual void ReadTimers(TFunctionRef<void(const FTimingProfilerTimer*, uint64)> Callback) const = 0;
-	virtual ITable<FTimingProfilerAggregatedStats>* CreateAggregation(double IntervalStart, double IntervalEnd, TFunctionRef<bool(uint32)> CpuThreadFilter, bool IncludeGpu) const = 0;
-};
-
-struct FFileInfo
-{
-	uint32 Id;
-	const TCHAR* Path;
-};
-
-enum EFileActivityType
-{
-	FileActivityType_Open,
-	FileActivityType_Close,
-	FileActivityType_Read,
-	FileActivityType_Write,
-
-	FileActivityType_Count,
-};
-
-struct FFileActivity
-{
-	uint64 Offset;
-	uint64 Size;
-	EFileActivityType ActivityType;
-	bool Failed;
-};
-
-class IFileActivityProvider
-{
-public:
-	typedef ITimeline<FFileActivity> Timeline;
-
-	virtual ~IFileActivityProvider() = default;
-	virtual void EnumerateFileActivity(TFunctionRef<bool(const FFileInfo&, const Timeline&)> Callback) const = 0;
-};
-
-struct FFrame
-{
-	uint64 Index;
-	double StartTime;
-	double EndTime;
-};
-
-class IFrameProvider
-{
-public:
-	virtual ~IFrameProvider() = default;
-	virtual uint64 GetFrameCount(ETraceFrameType FrameType) const = 0;
-	virtual void EnumerateFrames(ETraceFrameType FrameType, uint64 Start, uint64 End, TFunctionRef<void(const FFrame&)> Callback) const = 0;
-};
-
-struct FPackageSummaryInfo
-{
-	uint32 TotalHeaderSize = 0;
-	uint32 NameCount = 0;
-	uint32 ImportCount = 0;
-	uint32 ExportCount = 0;
-};
-
-struct FClassInfo
-{
-	const TCHAR* Name;
-};
-
-struct FPackageExportInfo
-{
-	uint32 Id;
-	const FClassInfo* Class = nullptr;
-	uint64 SerialOffset = 0;
-	uint64 SerialSize = 0;
-	bool IsAsset = false;
-};
-
-struct FPackageInfo
-{
-	uint32 Id;
-	const TCHAR* Name = nullptr;
-	FPackageSummaryInfo Summary;
-	TArray<FPackageExportInfo*> Exports;
-};
-
-struct FLoadTimeProfilerCpuEvent
-{
-	const FPackageInfo* Package = nullptr;
-	const FPackageExportInfo* Export = nullptr;
-	ELoadTimeProfilerPackageEventType PackageEventType = LoadTimeProfilerPackageEventType_None;
-	ELoadTimeProfilerObjectEventType ExportEventType = LoadTimeProfilerObjectEventType_None;
-};
-
-struct FLoadTimeProfilerAggregatedStats
-{
-	const TCHAR* Name;
-	uint64 Count;
-	double Total;
-	double Min;
-	double Max;
-	double Average;
-	double Median;
-};
-
-class ILoadTimeProfilerProvider
-{
-public:
-	typedef ITimeline<FLoadTimeProfilerCpuEvent> CpuTimeline;
-
-	virtual ~ILoadTimeProfilerProvider() = default;
-	virtual uint64 GetPackageCount() const = 0;
-	virtual void EnumeratePackages(TFunctionRef<void(const FPackageInfo&)> Callback) const = 0; 
-	virtual void ReadMainThreadCpuTimeline(TFunctionRef<void(const CpuTimeline&)> Callback) const = 0;
-	virtual void ReadAsyncLoadingThreadCpuTimeline(TFunctionRef<void(const CpuTimeline&)> Callback) const = 0;
-	virtual ITable<FLoadTimeProfilerAggregatedStats>* CreateEventAggregation(double IntervalStart, double IntervalEnd) const = 0;
-	virtual ITable<FLoadTimeProfilerAggregatedStats>* CreateObjectTypeAggregation(double IntervalStart, double IntervalEnd) const = 0;
-};
-
-enum ECounterDisplayHint
-{
-	CounterDisplayHint_None,
-	CounterDisplayHint_Memory,
-	CounterDisplayHint_FloatingPoint,
-};
-
-class ICounter
-{
-public:
-	virtual ~ICounter() = default;
-
-	virtual const TCHAR* GetName() const = 0;
-	virtual uint32 GetId() const = 0;
-	virtual ECounterDisplayHint GetDisplayHint() const = 0;
-	virtual void EnumerateValues(double IntervalStart, double IntervalEnd, TFunctionRef<void(double, int64)> Callback) const = 0;
-	virtual void EnumerateFloatValues(double IntervalStart, double IntervalEnd, TFunctionRef<void(double, double)> Callback) const = 0;
-};
-
-class ICounterProvider
-{
-public:
-	virtual ~ICounterProvider() = default;
-	virtual uint64 GetCounterCount() const = 0;
-	virtual void EnumerateCounters(TFunctionRef<void(const ICounter&)> Callback) const = 0;
-};
-
-class IAnalysisSession
-{
-public:
-	virtual ~IAnalysisSession() = default;
-	
-	virtual const TCHAR* GetName() const = 0;
-	virtual bool IsAnalysisComplete() const = 0;
-	virtual double GetDurationSeconds() const = 0;
-
-	virtual void BeginRead() const = 0;
-	virtual void EndRead() const = 0;
-	virtual const IBookmarkProvider& ReadBookmarkProvider() const = 0;
-	virtual const ILogProvider& ReadLogProvider() const = 0;
-	virtual const IThreadProvider& ReadThreadProvider() const = 0;
-	virtual const IFrameProvider& ReadFrameProvider() const = 0;
-	virtual const ITimingProfilerProvider& ReadTimingProfilerProvider() const = 0;
-	virtual const IFileActivityProvider& ReadFileActivityProvider() const = 0;
-	virtual const ILoadTimeProfilerProvider& ReadLoadTimeProfilerProvider() const = 0;
-	virtual const ICounterProvider& ReadCounterProvider() const = 0;
-};
-
-struct FAnalysisSessionReadScope
-{
-	FAnalysisSessionReadScope(const IAnalysisSession& InAnalysisSession)
-		: AnalysisSession(InAnalysisSession)
-	{
-		AnalysisSession.BeginRead();
-	}
-
-	~FAnalysisSessionReadScope()
-	{
-		AnalysisSession.EndRead();
-	}
-
-private:
-	const IAnalysisSession& AnalysisSession;
-};
-
 class IAnalysisService
 {
 public:
diff --git a/Engine/Source/Developer/TraceServices/Public/TraceServices/Containers/Allocators.h b/Engine/Source/Developer/TraceServices/Public/TraceServices/Containers/Allocators.h
new file mode 100644
--- /dev/null
+++ b/Engine/Source/Developer/TraceServices/Public/TraceServices/Containers/Allocators.h
@@ -0,0 +1,13 @@
+#pragma once
+
+namespace Trace
+{
+	
+class ILinearAllocator
+{
+public:
+	virtual ~ILinearAllocator() = default;
+	virtual void* Allocate(uint64 Size) = 0;
+};
+
+}
\ No newline at end of file
diff --git a/Engine/Source/Developer/TraceServices/Public/TraceServices/Containers/Tables.h b/Engine/Source/Developer/TraceServices/Public/TraceServices/Containers/Tables.h
new file mode 100644
--- /dev/null
+++ b/Engine/Source/Developer/TraceServices/Public/TraceServices/Containers/Tables.h
@@ -0,0 +1,65 @@
+#pragma once
+
+namespace Trace
+{
+
+enum ETableColumnType
+{
+	TableColumnType_Invalid,
+	TableColumnType_Bool,
+	TableColumnType_Int,
+	TableColumnType_Float,
+	TableColumnType_Double,
+	TableColumnType_CString,
+};
+
+class ITableLayout
+{
+public:
+	virtual ~ITableLayout() = default;
+	virtual uint8 GetColumnCount() const = 0;
+	virtual const TCHAR* GetColumnName(uint8 ColumnIndex) const = 0;
+	virtual ETableColumnType GetColumnType(uint8 ColumnIndex) const = 0;
+};
+
+class IUntypedTableReader
+{
+public:
+	virtual ~IUntypedTableReader() = default;
+	virtual bool IsValid() const = 0;
+	virtual void NextRow() = 0;
+	virtual void SetRowIndex(uint64 RowIndex) = 0;
+	virtual bool GetValueBool(uint8 ColumnIndex) const = 0;
+	virtual int64 GetValueInt(uint8 ColumnIndex) const = 0;
+	virtual float GetValueFloat(uint8 ColumnIndex) const = 0;
+	virtual double GetValueDouble(uint8 ColumnIndex) const = 0;
+	virtual const TCHAR* GetValueCString(uint8 ColumnIndex) const = 0;
+};
+
+template<typename RowType>
+class ITableReader
+	: public IUntypedTableReader
+{
+public:
+	virtual const RowType* GetCurrentRow() const = 0;
+};
+
+class IUntypedTable
+{
+public:
+	virtual ~IUntypedTable() = default;
+	virtual const ITableLayout& GetLayout() const = 0;
+	virtual uint64 GetRowCount() const = 0;
+	virtual IUntypedTableReader* CreateReader() const = 0;
+};
+
+template<typename RowType>
+class ITable
+	: public IUntypedTable
+{
+public:
+	virtual ~ITable() = default;
+	virtual ITableReader<RowType>* CreateReader() const = 0;
+};
+
+}
\ No newline at end of file
diff --git a/Engine/Source/Developer/TraceServices/Public/TraceServices/Containers/Timelines.h b/Engine/Source/Developer/TraceServices/Public/TraceServices/Containers/Timelines.h
new file mode 100644
--- /dev/null
+++ b/Engine/Source/Developer/TraceServices/Public/TraceServices/Containers/Timelines.h
@@ -0,0 +1,38 @@
+#pragma once
+
+#include "Templates/Function.h"
+
+namespace Trace
+{
+
+template<typename InEventType>
+class ITimeline
+{
+public:
+	typedef InEventType EventType;
+
+	virtual ~ITimeline() = default;
+	virtual uint64 GetModCount() const = 0;
+	virtual uint64 GetEventCount() const = 0;
+	virtual void EnumerateEventsDownSampled(double IntervalStart, double IntervalEnd, double Resolution, TFunctionRef<void(bool, double, const EventType&)> Callback) const = 0;
+	virtual void EnumerateEventsDownSampled(double IntervalStart, double IntervalEnd, double Resolution, TFunctionRef<void(double, double, uint32, const EventType&)> Callback) const = 0;
+	virtual void EnumerateEvents(double IntervalStart, double IntervalEnd, TFunctionRef<void(bool, double, const EventType&)> Callback) const = 0;
+	virtual void EnumerateEvents(double IntervalStart, double IntervalEnd, TFunctionRef<void(double, double, uint32, const EventType&)> Callback) const = 0;
+};
+
+struct FAggregatedTimingStats
+{
+	uint64 InstanceCount = 0;
+	double TotalInclusiveTime = 0.0;
+	double MinInclusiveTime = DBL_MAX;
+	double MaxInclusiveTime = -DBL_MAX;
+	double AverageInclusiveTime = 0.0;
+	double MedianInclusiveTime = 0.0;
+	double TotalExclusiveTime = 0.0;
+	double MinExclusiveTime = DBL_MAX;
+	double MaxExclusiveTime = -DBL_MAX;
+	double AverageExclusiveTime = 0.0;
+	double MedianExclusiveTime = 0.0;
+};
+
+}
\ No newline at end of file
diff --git a/Engine/Source/Developer/TraceServices/Public/TraceServices/ITraceServicesModule.h b/Engine/Source/Developer/TraceServices/Public/TraceServices/ITraceServicesModule.h
--- a/Engine/Source/Developer/TraceServices/Public/TraceServices/ITraceServicesModule.h
+++ b/Engine/Source/Developer/TraceServices/Public/TraceServices/ITraceServicesModule.h
@@ -11,6 +11,7 @@ namespace Trace
 
 class IAnalysisService;
 class ISessionService;
+class IModuleService;
 
 }
 
@@ -20,6 +21,7 @@ class ITraceServicesModule
 public:
 	virtual TSharedPtr<Trace::ISessionService> GetSessionService() = 0;
 	virtual TSharedPtr<Trace::IAnalysisService> GetAnalysisService() = 0;
+	virtual TSharedPtr<Trace::IModuleService> GetModuleService() = 0;
 
-	virtual ~ITraceServicesModule() { }
+	virtual ~ITraceServicesModule() = default;
 };
diff --git a/Engine/Source/Developer/TraceServices/Public/TraceServices/Model/AnalysisSession.h b/Engine/Source/Developer/TraceServices/Public/TraceServices/Model/AnalysisSession.h
new file mode 100644
--- /dev/null
+++ b/Engine/Source/Developer/TraceServices/Public/TraceServices/Model/AnalysisSession.h
@@ -0,0 +1,83 @@
+#pragma once
+
+#include "CoreTypes.h"
+
+class FName;
+
+namespace Trace
+{
+
+class ILinearAllocator;
+
+class IProvider
+{
+public:
+	virtual ~IProvider() = default;
+};
+
+class IAnalysisSession
+{
+public:
+	virtual ~IAnalysisSession() = default;
+	
+	virtual const TCHAR* GetName() const = 0;
+	virtual bool IsAnalysisComplete() const = 0;
+	virtual double GetDurationSeconds() const = 0;
+	virtual void UpdateDurationSeconds(double Duration) = 0;
+
+	virtual void BeginRead() const = 0;
+	virtual void EndRead() const = 0;
+	virtual void ReadAccessCheck() const = 0;
+
+	virtual void BeginEdit() = 0;
+	virtual void EndEdit() = 0;
+	virtual void WriteAccessCheck() = 0;
+	
+	virtual ILinearAllocator& GetLinearAllocator() = 0;
+	virtual const TCHAR* StoreString(const TCHAR* String) = 0;
+	
+	virtual void AddProvider(const FName& Name, IProvider* Provider) = 0;
+	template<typename ProviderType>
+	const ProviderType* ReadProvider(const FName& Name) const
+	{
+		return static_cast<const ProviderType*>(ReadProviderPrivate(Name));
+	}
+
+private:
+	virtual const IProvider* ReadProviderPrivate(const FName& Name) const = 0;
+};
+
+struct FAnalysisSessionReadScope
+{
+	FAnalysisSessionReadScope(const IAnalysisSession& InAnalysisSession)
+		: AnalysisSession(InAnalysisSession)
+	{
+		AnalysisSession.BeginRead();
+	}
+
+	~FAnalysisSessionReadScope()
+	{
+		AnalysisSession.EndRead();
+	}
+
+private:
+	const IAnalysisSession& AnalysisSession;
+};
+	
+struct FAnalysisSessionEditScope
+{
+	FAnalysisSessionEditScope(IAnalysisSession& InAnalysisSession)
+		: AnalysisSession(InAnalysisSession)
+	{
+		AnalysisSession.BeginEdit();
+	}
+
+	~FAnalysisSessionEditScope()
+	{
+		AnalysisSession.EndEdit();
+	}
+
+	IAnalysisSession& AnalysisSession;
+};
+	
+}
diff --git a/Engine/Source/Developer/TraceServices/Public/TraceServices/Model/Bookmarks.h b/Engine/Source/Developer/TraceServices/Public/TraceServices/Model/Bookmarks.h
new file mode 100644
--- /dev/null
+++ b/Engine/Source/Developer/TraceServices/Public/TraceServices/Model/Bookmarks.h
@@ -0,0 +1,28 @@
+// Copyright 1998-2019 Epic Games, Inc. All Rights Reserved.
+
+#pragma  once
+
+#include "TraceServices/Model/AnalysisSession.h"
+#include "Templates/Function.h"
+
+namespace Trace
+{
+
+struct FBookmark
+{
+	double Time;
+	const TCHAR* Text;
+};
+
+class IBookmarkProvider
+	: public IProvider
+{
+public:
+	virtual ~IBookmarkProvider() = default;
+	virtual uint64 GetBookmarkCount() const = 0;
+	virtual void EnumerateBookmarks(double IntervalStart, double IntervalEnd, TFunctionRef<void(const FBookmark&)> Callback) const = 0;
+};
+
+TRACESERVICES_API const IBookmarkProvider& ReadBookmarkProvider(const IAnalysisSession& Session);
+
+}
\ No newline at end of file
diff --git a/Engine/Source/Developer/TraceServices/Public/TraceServices/Model/Frames.h b/Engine/Source/Developer/TraceServices/Public/TraceServices/Model/Frames.h
new file mode 100644
--- /dev/null
+++ b/Engine/Source/Developer/TraceServices/Public/TraceServices/Model/Frames.h
@@ -0,0 +1,30 @@
+// Copyright 1998-2019 Epic Games, Inc. All Rights Reserved.
+
+#pragma  once
+
+#include "TraceServices/Model/AnalysisSession.h"
+#include "ProfilingDebugging/MiscTrace.h"
+#include "Templates/Function.h"
+
+namespace Trace
+{
+
+struct FFrame
+{
+	uint64 Index;
+	double StartTime;
+	double EndTime;
+};
+
+class IFrameProvider
+	: public IProvider
+{
+public:
+	virtual ~IFrameProvider() = default;
+	virtual uint64 GetFrameCount(ETraceFrameType FrameType) const = 0;
+	virtual void EnumerateFrames(ETraceFrameType FrameType, uint64 Start, uint64 End, TFunctionRef<void(const FFrame&)> Callback) const = 0;
+};
+
+TRACESERVICES_API const IFrameProvider& ReadFrameProvider(const IAnalysisSession& Session);
+
+}
\ No newline at end of file
diff --git a/Engine/Source/Developer/TraceServices/Public/TraceServices/Model/LoadTimeProfiler.h b/Engine/Source/Developer/TraceServices/Public/TraceServices/Model/LoadTimeProfiler.h
new file mode 100644
--- /dev/null
+++ b/Engine/Source/Developer/TraceServices/Public/TraceServices/Model/LoadTimeProfiler.h
@@ -0,0 +1,113 @@
+#pragma once
+
+#include "TraceServices/Model/AnalysisSession.h"
+#include "Serialization/LoadTimeTrace.h"
+#include "TraceServices/Containers/Timelines.h"
+#include "TraceServices/Containers/Tables.h"
+#include "Containers/Array.h"
+
+namespace Trace
+{
+
+struct FFileInfo
+{
+	uint32 Id;
+	const TCHAR* Path;
+};
+
+enum EFileActivityType
+{
+	FileActivityType_Open,
+	FileActivityType_Close,
+	FileActivityType_Read,
+	FileActivityType_Write,
+
+	FileActivityType_Count,
+};
+
+struct FFileActivity
+{
+	uint64 Offset;
+	uint64 Size;
+	EFileActivityType ActivityType;
+	bool Failed;
+};
+
+class IFileActivityProvider
+	: public IProvider
+{
+public:
+	typedef ITimeline<FFileActivity> Timeline;
+
+	virtual ~IFileActivityProvider() = default;
+	virtual void EnumerateFileActivity(TFunctionRef<bool(const FFileInfo&, const Timeline&)> Callback) const = 0;
+};
+
+struct FPackageSummaryInfo
+{
+	uint32 TotalHeaderSize = 0;
+	uint32 NameCount = 0;
+	uint32 ImportCount = 0;
+	uint32 ExportCount = 0;
+};
+
+struct FClassInfo
+{
+	const TCHAR* Name;
+};
+
+struct FPackageExportInfo
+{
+	uint32 Id;
+	const FClassInfo* Class = nullptr;
+	uint64 SerialOffset = 0;
+	uint64 SerialSize = 0;
+	bool IsAsset = false;
+};
+
+struct FPackageInfo
+{
+	uint32 Id;
+	const TCHAR* Name = nullptr;
+	FPackageSummaryInfo Summary;
+	TArray<FPackageExportInfo*> Exports;
+};
+
+struct FLoadTimeProfilerCpuEvent
+{
+	const FPackageInfo* Package = nullptr;
+	const FPackageExportInfo* Export = nullptr;
+	ELoadTimeProfilerPackageEventType PackageEventType = LoadTimeProfilerPackageEventType_None;
+	ELoadTimeProfilerObjectEventType ExportEventType = LoadTimeProfilerObjectEventType_None;
+};
+
+struct FLoadTimeProfilerAggregatedStats
+{
+	const TCHAR* Name;
+	uint64 Count;
+	double Total;
+	double Min;
+	double Max;
+	double Average;
+	double Median;
+};
+
+class ILoadTimeProfilerProvider
+	: public IProvider
+{
+public:
+	typedef ITimeline<FLoadTimeProfilerCpuEvent> CpuTimeline;
+
+	virtual ~ILoadTimeProfilerProvider() = default;
+	virtual uint64 GetPackageCount() const = 0;
+	virtual void EnumeratePackages(TFunctionRef<void(const FPackageInfo&)> Callback) const = 0; 
+	virtual void ReadMainThreadCpuTimeline(TFunctionRef<void(const CpuTimeline&)> Callback) const = 0;
+	virtual void ReadAsyncLoadingThreadCpuTimeline(TFunctionRef<void(const CpuTimeline&)> Callback) const = 0;
+	virtual ITable<FLoadTimeProfilerAggregatedStats>* CreateEventAggregation(double IntervalStart, double IntervalEnd) const = 0;
+	virtual ITable<FLoadTimeProfilerAggregatedStats>* CreateObjectTypeAggregation(double IntervalStart, double IntervalEnd) const = 0;
+};
+
+TRACESERVICES_API const ILoadTimeProfilerProvider* ReadLoadTimeProfilerProvider(const IAnalysisSession& Session);
+TRACESERVICES_API const IFileActivityProvider* ReadFileActivityProvider(const IAnalysisSession& Session);
+
+}
\ No newline at end of file
diff --git a/Engine/Source/Developer/TraceServices/Public/TraceServices/Model/Log.h b/Engine/Source/Developer/TraceServices/Public/TraceServices/Model/Log.h
new file mode 100644
--- /dev/null
+++ b/Engine/Source/Developer/TraceServices/Public/TraceServices/Model/Log.h
@@ -0,0 +1,46 @@
+// Copyright 1998-2019 Epic Games, Inc. All Rights Reserved.
+
+#pragma once
+
+#include "CoreTypes.h"
+#include "TraceServices/Model/AnalysisSession.h"
+#include "TraceServices/Containers/Tables.h"
+#include "Templates/Function.h"
+
+namespace Trace
+{
+
+struct FLogCategory
+{
+	const TCHAR* Name = nullptr;
+	ELogVerbosity::Type DefaultVerbosity;
+};
+
+struct FLogMessage
+{
+	uint64 Index;
+	double Time;
+	const FLogCategory* Category = nullptr;
+	const TCHAR* File = nullptr;
+	const TCHAR* Message = nullptr;
+	int32 Line;
+	ELogVerbosity::Type Verbosity;
+};
+
+class ILogProvider
+	: public IProvider
+{
+public:
+	virtual ~ILogProvider() = default;
+	virtual uint64 GetMessageCount() const = 0;
+	virtual void EnumerateMessages(double IntervalStart, double IntervalEnd, TFunctionRef<void(const FLogMessage&)> Callback) const = 0;
+	virtual void EnumerateMessagesByIndex(uint64 Start, uint64 End, TFunctionRef<void(const FLogMessage&)> Callback) const = 0;
+	virtual bool ReadMessage(uint64 Index, TFunctionRef<void(const FLogMessage&)> Callback) const = 0;
+	virtual uint64 GetCategoryCount() const = 0;
+	virtual void EnumerateCategories(TFunctionRef<void(const FLogCategory&)> Callback) const = 0;
+	virtual const IUntypedTable& GetMessagesTable() const = 0;
+};
+
+TRACESERVICES_API const ILogProvider& ReadLogProvider(const IAnalysisSession& Session);
+
+}
\ No newline at end of file
diff --git a/Engine/Source/Developer/TraceServices/Public/TraceServices/Model/Stats.h b/Engine/Source/Developer/TraceServices/Public/TraceServices/Model/Stats.h
new file mode 100644
--- /dev/null
+++ b/Engine/Source/Developer/TraceServices/Public/TraceServices/Model/Stats.h
@@ -0,0 +1,39 @@
+#pragma once
+
+#include "TraceServices/Model/AnalysisSession.h"
+
+namespace Trace
+{
+
+enum ECounterDisplayHint
+{
+	CounterDisplayHint_None,
+	CounterDisplayHint_Memory,
+	CounterDisplayHint_FloatingPoint,
+};
+
+class ICounter
+{
+public:
+	virtual ~ICounter() = default;
+
+	virtual const TCHAR* GetName() const = 0;
+	virtual uint32 GetId() const = 0;
+	virtual ECounterDisplayHint GetDisplayHint() const = 0;
+	virtual void EnumerateValues(double IntervalStart, double IntervalEnd, TFunctionRef<void(double, int64)> Callback) const = 0;
+	virtual void EnumerateFloatValues(double IntervalStart, double IntervalEnd, TFunctionRef<void(double, double)> Callback) const = 0;
+};
+
+class ICounterProvider
+	: public IProvider
+{
+public:
+	static const FName ProviderName;
+	virtual ~ICounterProvider() = default;
+	virtual uint64 GetCounterCount() const = 0;
+	virtual void EnumerateCounters(TFunctionRef<void(const ICounter&)> Callback) const = 0;
+};
+
+TRACESERVICES_API const ICounterProvider* ReadCounterProvider(const IAnalysisSession& Session);
+
+}
diff --git a/Engine/Source/Developer/TraceServices/Public/TraceServices/Model/Threads.h b/Engine/Source/Developer/TraceServices/Public/TraceServices/Model/Threads.h
new file mode 100644
--- /dev/null
+++ b/Engine/Source/Developer/TraceServices/Public/TraceServices/Model/Threads.h
@@ -0,0 +1,29 @@
+// Copyright 1998-2019 Epic Games, Inc. All Rights Reserved.
+
+#pragma once
+
+#include "TraceServices/Model/AnalysisSession.h"
+#include "Templates/Function.h"
+
+namespace Trace
+{
+
+struct FThreadInfo
+{
+	uint32 Id;
+	const TCHAR* Name;
+	const TCHAR* GroupName;
+};
+
+class IThreadProvider
+	: public IProvider
+{
+public:
+	virtual ~IThreadProvider() = default;
+	virtual uint64 GetModCount() const = 0;
+	virtual void EnumerateThreads(TFunctionRef<void(const FThreadInfo&)> Callback) const = 0;
+};
+
+TRACESERVICES_API const IThreadProvider& ReadThreadProvider(const IAnalysisSession& Session);
+
+}
\ No newline at end of file
diff --git a/Engine/Source/Developer/TraceServices/Public/TraceServices/Model/TimingProfiler.h b/Engine/Source/Developer/TraceServices/Public/TraceServices/Model/TimingProfiler.h
new file mode 100644
--- /dev/null
+++ b/Engine/Source/Developer/TraceServices/Public/TraceServices/Model/TimingProfiler.h
@@ -0,0 +1,49 @@
+// Copyright 1998-2019 Epic Games, Inc. All Rights Reserved.
+
+#pragma once
+
+#include "TraceServices/Model/AnalysisSession.h"
+#include "TraceServices/Containers/Timelines.h"
+#include "TraceServices/Containers/Tables.h"
+
+namespace Trace
+{
+
+struct FTimingProfilerTimer
+{
+	const TCHAR* Name;
+	uint32 Id;
+	uint32 NameHash;
+	bool IsGpuTimer;
+};
+
+struct FTimingProfilerEvent
+{
+	uint32 TimerIndex;
+};
+
+struct FTimingProfilerAggregatedStats
+	: public FAggregatedTimingStats
+{
+	const FTimingProfilerTimer* Timer = nullptr;
+};
+
+class ITimingProfilerProvider
+	: public IProvider
+{
+public:
+	typedef ITimeline<FTimingProfilerEvent> Timeline;
+
+	virtual ~ITimingProfilerProvider() = default;
+	virtual bool GetCpuThreadTimelineIndex(uint32 ThreadId, uint32& OutTimelineIndex) const = 0;
+	virtual bool GetGpuTimelineIndex(uint32& OutTimelineIndex) const = 0;
+	virtual bool ReadTimeline(uint32 Index, TFunctionRef<void(const Timeline&)> Callback) const = 0;
+	virtual uint64 GetTimelineCount() const = 0;
+	virtual void EnumerateTimelines(TFunctionRef<void(const Timeline&)> Callback) const = 0;
+	virtual void ReadTimers(TFunctionRef<void(const FTimingProfilerTimer*, uint64)> Callback) const = 0;
+	virtual ITable<FTimingProfilerAggregatedStats>* CreateAggregation(double IntervalStart, double IntervalEnd, TFunctionRef<bool(uint32)> CpuThreadFilter, bool IncludeGpu) const = 0;
+};
+
+TRACESERVICES_API const ITimingProfilerProvider* ReadTimingProfilerProvider(const IAnalysisSession& Session);
+
+}
diff --git a/Engine/Source/Developer/TraceServices/Public/TraceServices/ModuleService.h b/Engine/Source/Developer/TraceServices/Public/TraceServices/ModuleService.h
new file mode 100644
--- /dev/null
+++ b/Engine/Source/Developer/TraceServices/Public/TraceServices/ModuleService.h
@@ -0,0 +1,39 @@
+// Copyright 1998-2019 Epic Games, Inc. All Rights Reserved.
+
+#pragma once
+
+#include "CoreMinimal.h"
+#include "Features/IModularFeature.h"
+
+class FName;
+
+namespace Trace
+{
+
+class IAnalyzer;
+class IAnalysisSession;
+extern const FName ModuleFeatureName;
+
+struct FModuleInfo
+{
+	FName Name;
+	const TCHAR* DisplayName = nullptr;
+};
+
+class IModule
+	: public IModularFeature
+{
+public:
+	virtual void GetModuleInfo(FModuleInfo& OutModuleInfo) = 0;
+	virtual void OnAnalysisBegin(IAnalysisSession& Session, bool bIsEnabled, TArray<IAnalyzer*>& OutAnalyzers) = 0;
+};
+
+class IModuleService
+{
+public:
+	virtual ~IModuleService() = default;
+	virtual void GetAvailableModules(TArray<FModuleInfo>& OutModules) = 0;
+	virtual void SetModuleEnabled(const FName& ModuleName, bool bEnabled) = 0;
+};
+
+}
