UnrealInsights:
Traced thread group names are now free to be any string they want

diff --git a/Engine/Source/Developer/TraceServices/Private/Analyzers/MiscTraceAnalysis.cpp b/Engine/Source/Developer/TraceServices/Private/Analyzers/MiscTraceAnalysis.cpp
--- a/Engine/Source/Developer/TraceServices/Private/Analyzers/MiscTraceAnalysis.cpp
+++ b/Engine/Source/Developer/TraceServices/Private/Analyzers/MiscTraceAnalysis.cpp
@@ -37,51 +37,6 @@ void FMiscTraceAnalyzer::OnAnalysisBegin(const FOnAnalysisContext& Context)
 	Builder.RouteEvent(RouteId_EndFrame, "Misc", "EndFrame");
 }
 
-static ETraceThreadGroup GetThreadGroupFromName(const char* GroupName)
-{
-	if (!strcmp(GroupName, "Render"))
-	{
-		return TraceThreadGroup_Render;
-	}
-	else if (!strcmp(GroupName, "AsyncLoading"))
-	{
-		return TraceThreadGroup_AsyncLoading;
-	}
-	else if (!strcmp(GroupName, "TaskGraphHigh"))
-	{
-		return TraceThreadGroup_TaskGraphHigh;
-	}
-	else if (!strcmp(GroupName, "TaskGraphNormal"))
-	{
-		return TraceThreadGroup_TaskGraphNormal;
-	}
-	else if (!strcmp(GroupName, "TaskGraphLow"))
-	{
-		return TraceThreadGroup_TaskGraphLow;
-	}
-	else if (!strcmp(GroupName, "ThreadPool"))
-	{
-		return TraceThreadGroup_ThreadPool;
-	}
-	else if (!strcmp(GroupName, "BackgroundThreadPool"))
-	{
-		return TraceThreadGroup_BackgroundThreadPool;
-	}
-	else if (!strcmp(GroupName, "LargeThreadPool"))
-	{
-		return TraceThreadGroup_LargeThreadPool;
-	}
-	else if (!strcmp(GroupName, "IOThreadPool"))
-	{
-		return TraceThreadGroup_IOThreadPool;
-	}
-	else
-	{
-		check(false);
-		return ETraceThreadGroup(0);
-	}
-}
-
 void FMiscTraceAnalyzer::OnEvent(uint16 RouteId, const FOnEventContext& Context)
 {
 	Trace::FAnalysisSessionEditScope _(Session);
@@ -110,19 +65,17 @@ void FMiscTraceAnalyzer::OnEvent(uint16 RouteId, const FOnEventContext& Context)
 	}
 	case RouteId_SetThreadGroup:
 	{
-		const char* GroupName = reinterpret_cast<const char*>(EventData.GetAttachment());
-		ETraceThreadGroup ThreadGroup = GetThreadGroupFromName(GroupName);
+		const TCHAR* GroupName = Session.StoreString(ANSI_TO_TCHAR(reinterpret_cast<const char*>(EventData.GetAttachment())));
 		uint32 ThreadId = EventData.GetValue("ThreadId").As<uint32>();
-		ThreadProvider.SetThreadGroup(ThreadId, ThreadGroup);
+		ThreadProvider.SetThreadGroup(ThreadId, GroupName);
 		break;
 	}
 	case RouteId_BeginThreadGroupScope:
 	{
-		const char* GroupName = reinterpret_cast<const char*>(EventData.GetAttachment());
-		ETraceThreadGroup ThreadGroup = GetThreadGroupFromName(GroupName);
+		const TCHAR* GroupName = Session.StoreString(ANSI_TO_TCHAR(reinterpret_cast<const char*>(EventData.GetAttachment())));
 		uint32 CurrentThreadId = EventData.GetValue("CurrentThreadId").As<uint32>();
 		FThreadState* ThreadState = GetThreadState(CurrentThreadId);
-		ThreadState->ThreadGroupStack.Push(ThreadGroup);
+		ThreadState->ThreadGroupStack.Push(GroupName);
 		break;
 	}
 	case RouteId_EndThreadGroupScope:
diff --git a/Engine/Source/Developer/TraceServices/Private/Analyzers/MiscTraceAnalysis.h b/Engine/Source/Developer/TraceServices/Private/Analyzers/MiscTraceAnalysis.h
--- a/Engine/Source/Developer/TraceServices/Private/Analyzers/MiscTraceAnalysis.h
+++ b/Engine/Source/Developer/TraceServices/Private/Analyzers/MiscTraceAnalysis.h
@@ -46,7 +46,7 @@ private:
 
 	struct FThreadState
 	{
-		TArray<ETraceThreadGroup> ThreadGroupStack;
+		TArray<const TCHAR*> ThreadGroupStack;
 	};
 
 	FThreadState* GetThreadState(uint32 ThreadId);
diff --git a/Engine/Source/Developer/TraceServices/Private/Model/Threads.cpp b/Engine/Source/Developer/TraceServices/Private/Model/Threads.cpp
--- a/Engine/Source/Developer/TraceServices/Private/Model/Threads.cpp
+++ b/Engine/Source/Developer/TraceServices/Private/Model/Threads.cpp
@@ -75,40 +75,14 @@ void FThreadProvider::SetThreadPriority(uint32 Id, EThreadPriority Priority)
 	++ModCount;
 }
 
-void FThreadProvider::SetThreadGroup(uint32 Id, ETraceThreadGroup Group)
+void FThreadProvider::SetThreadGroup(uint32 Id, const TCHAR* GroupName)
 {
 	Session.WriteAccessCheck();
 
 	check(ThreadMap.Contains(Id));
 	FThreadInfoInternal* ThreadInfo = ThreadMap[Id];
-	ThreadInfo->GroupSortOrder = GetGroupSortOrder(Group);
-	switch (Group)
-	{
-	case TraceThreadGroup_Render:
-		ThreadInfo->GroupName = TEXT("Render");
-		break;
-	case TraceThreadGroup_TaskGraphHigh:
-		ThreadInfo->GroupName = TEXT("TaskGraph (High Priority)");
-		break;
-	case TraceThreadGroup_TaskGraphNormal:
-		ThreadInfo->GroupName = TEXT("TaskGraph (Normal Priority)");
-		break;
-	case TraceThreadGroup_TaskGraphLow:
-		ThreadInfo->GroupName = TEXT("TaskGraph (Low Priority)");
-		break;
-	case TraceThreadGroup_ThreadPool:
-		ThreadInfo->GroupName = TEXT("ThreadPool");
-		break;
-	case TraceThreadGroup_BackgroundThreadPool:
-		ThreadInfo->GroupName = TEXT("ThreadPool (Background)");
-		break;
-	case TraceThreadGroup_LargeThreadPool:
-		ThreadInfo->GroupName = TEXT("ThreadPool (Large)");
-		break;
-	case TraceThreadGroup_IOThreadPool:
-		ThreadInfo->GroupName = TEXT("ThreadPool (IO)");
-		break;
-	}
+	ThreadInfo->GroupName = GroupName;
+	ThreadInfo->GroupSortOrder = GetGroupSortOrder(GroupName);
 	SortThreads();
 	++ModCount;
 }
@@ -135,9 +109,52 @@ void FThreadProvider::SortThreads()
 	});
 }
 
-uint32 FThreadProvider::GetGroupSortOrder(ETraceThreadGroup ThreadGroup)
+uint32 FThreadProvider::GetGroupSortOrder(const TCHAR* GroupName)
 {
-	return uint32(ThreadGroup);
+	if (!GroupName)
+	{
+		return uint32(-1);
+	}
+	else if (!FCString::Strcmp(GroupName, TEXT("Render")))
+	{
+		return 0;
+	}
+	else if (!FCString::Strcmp(GroupName, TEXT("AsyncLoading")))
+	{
+		return 1;
+	}
+	else if (!FCString::Strcmp(GroupName, TEXT("TaskGraphHigh")))
+	{
+		return 2;
+	}
+	else if (!FCString::Strcmp(GroupName, TEXT("TaskGraphNormal")))
+	{
+		return 3;
+	}
+	else if (!FCString::Strcmp(GroupName, TEXT("TaskGraphLow")))
+	{
+		return 4;
+	}
+	else if (!FCString::Strcmp(GroupName, TEXT("LargeThreadPool")))
+	{
+		return 5;
+	}
+	else if (!FCString::Strcmp(GroupName, TEXT("ThreadPool")))
+	{
+		return 6;
+	}
+	else if (!FCString::Strcmp(GroupName, TEXT("BackgroundThreadPool")))
+	{
+		return 7;
+	}
+	else if (!FCString::Strcmp(GroupName, TEXT("IOThreadPool")))
+	{
+		return 8;
+	}
+	else
+	{
+		return GetTypeHash(GroupName);
+	}
 }
 
 uint32 FThreadProvider::GetPrioritySortOrder(EThreadPriority ThreadPriority)
diff --git a/Engine/Source/Developer/TraceServices/Private/Model/ThreadsPrivate.h b/Engine/Source/Developer/TraceServices/Private/Model/ThreadsPrivate.h
--- a/Engine/Source/Developer/TraceServices/Private/Model/ThreadsPrivate.h
+++ b/Engine/Source/Developer/TraceServices/Private/Model/ThreadsPrivate.h
@@ -20,7 +20,7 @@ public:
 	void AddGameThread(uint32 Id);
 	void AddThread(uint32 Id, const TCHAR* Name, EThreadPriority Priority);
 	void SetThreadPriority(uint32 Id, EThreadPriority Priority);
-	void SetThreadGroup(uint32 Id, ETraceThreadGroup Group);
+	void SetThreadGroup(uint32 Id, const TCHAR* GroupName);
 	virtual uint64 GetModCount() const override { return ModCount; }
 	virtual void EnumerateThreads(TFunctionRef<void(const FThreadInfo&)> Callback) const override;
 
@@ -40,7 +40,7 @@ private:
 
 	void SortThreads();
 	static uint32 GetPrioritySortOrder(EThreadPriority ThreadPriority);
-	static uint32 GetGroupSortOrder(ETraceThreadGroup ThreadGroup);
+	static uint32 GetGroupSortOrder(const TCHAR* GroupName);
 
 	IAnalysisSession& Session;
 	uint64 ModCount = 0;
diff --git a/Engine/Source/Runtime/Core/Private/Async/TaskGraph.cpp b/Engine/Source/Runtime/Core/Private/Async/TaskGraph.cpp
--- a/Engine/Source/Runtime/Core/Private/Async/TaskGraph.cpp
+++ b/Engine/Source/Runtime/Core/Private/Async/TaskGraph.cpp
@@ -1211,14 +1211,14 @@ public:
 		for (int32 ThreadIndex = LastExternalThread + 1; ThreadIndex < NumThreads; ThreadIndex++)
 		{
 			FString Name;
-			ETraceThreadGroup Group = TraceThreadGroup_TaskGraphNormal;
+			const ANSICHAR* GroupName = "TaskGraphNormal";
 			int32 Priority = ThreadIndexToPriorityIndex(ThreadIndex);
 			EThreadPriority ThreadPri;
 			uint64 Affinity = FPlatformAffinity::GetTaskGraphThreadMask();
 			if (Priority == 1)
 			{
 				Name = FString::Printf(TEXT("TaskGraphThreadHP %d"), ThreadIndex - (LastExternalThread + 1));
-				Group = TraceThreadGroup_TaskGraphHigh;
+				GroupName = "TaskGraphHigh";
 				ThreadPri = TPri_SlightlyBelowNormal; // we want even hi priority tasks below the normal threads
 
 				// If the platform defines FPlatformAffinity::GetTaskGraphHighPriorityTaskMask then use it
@@ -1230,7 +1230,7 @@ public:
 			else if (Priority == 2)
 			{
 				Name = FString::Printf(TEXT("TaskGraphThreadBP %d"), ThreadIndex - (LastExternalThread + 1));
-				Group = TraceThreadGroup_TaskGraphLow;
+				GroupName = "TaskGraphLow";
 				ThreadPri = TPri_Lowest;
 				// If the platform defines FPlatformAffinity::GetTaskGraphBackgroundTaskMask then use it
 				if ( FPlatformAffinity::GetTaskGraphBackgroundTaskMask() != 0xFFFFFFFFFFFFFFFF )
@@ -1254,7 +1254,7 @@ public:
 			WorkerThreads[ThreadIndex].bAttached = true;
 			if (WorkerThreads[ThreadIndex].RunnableThread)
 			{
-				TRACE_SET_THREAD_GROUP(WorkerThreads[ThreadIndex].RunnableThread->GetThreadID(), Group);
+				TRACE_SET_THREAD_GROUP(WorkerThreads[ThreadIndex].RunnableThread->GetThreadID(), GroupName);
 			}
 		}
 	}
diff --git a/Engine/Source/Runtime/Core/Private/ProfilingDebugging/MiscTrace.cpp b/Engine/Source/Runtime/Core/Private/ProfilingDebugging/MiscTrace.cpp
--- a/Engine/Source/Runtime/Core/Private/ProfilingDebugging/MiscTrace.cpp
+++ b/Engine/Source/Runtime/Core/Private/ProfilingDebugging/MiscTrace.cpp
@@ -67,45 +67,16 @@ void FMiscTrace::OutputCreateThread(uint32 Id, const TCHAR* Name, uint32 Priorit
 		<< CreateThread.Attachment(Name, NameSize);
 }
 
-static const char* GetThreadGroupName(ETraceThreadGroup Group)
+void FMiscTrace::OutputSetThreadGroup(uint32 Id, const ANSICHAR* GroupName)
 {
-	switch (Group)
-	{
-	case TraceThreadGroup_Render:
-		return "Render";
-	case TraceThreadGroup_TaskGraphHigh:
-		return "TaskGraphHigh";
-	case TraceThreadGroup_TaskGraphNormal:
-		return "TaskGraphNormal";
-	case TraceThreadGroup_TaskGraphLow:
-		return "TaskGraphLow";
-	case TraceThreadGroup_ThreadPool:
-		return "ThreadPool";
-	case TraceThreadGroup_BackgroundThreadPool:
-		return "BackgroundThreadPool";
-	case TraceThreadGroup_LargeThreadPool:
-		return "LargeThreadPool";
-	case TraceThreadGroup_AsyncLoading:
-		return "AsyncLoading";
-	case TraceThreadGroup_IOThreadPool:
-		return "IOThreadPool";
-	}
-	check(false);
-	return nullptr;
-}
-
-void FMiscTrace::OutputSetThreadGroup(uint32 Id, ETraceThreadGroup Group)
-{
-	const char* GroupName = GetThreadGroupName(Group);
 	uint16 NameSize = strlen(GroupName) + 1;
 	UE_TRACE_LOG(Misc, SetThreadGroup, NameSize)
 		<< SetThreadGroup.ThreadId(Id)
 		<< SetThreadGroup.Attachment(GroupName, NameSize);
 }
 
-void FMiscTrace::OutputBeginThreadGroupScope(ETraceThreadGroup Group)
+void FMiscTrace::OutputBeginThreadGroupScope(const ANSICHAR* GroupName)
 {
-	const char* GroupName = GetThreadGroupName(Group);
 	uint16 NameSize = strlen(GroupName) + 1;
 	UE_TRACE_LOG(Misc, BeginThreadGroupScope, NameSize)
 		<< BeginThreadGroupScope.CurrentThreadId(FPlatformTLS::GetCurrentThreadId())
diff --git a/Engine/Source/Runtime/Core/Public/ProfilingDebugging/MiscTrace.h b/Engine/Source/Runtime/Core/Public/ProfilingDebugging/MiscTrace.h
--- a/Engine/Source/Runtime/Core/Public/ProfilingDebugging/MiscTrace.h
+++ b/Engine/Source/Runtime/Core/Public/ProfilingDebugging/MiscTrace.h
@@ -10,21 +10,6 @@
 #define MISCTRACE_ENABLED 0
 #endif
 
-enum ETraceThreadGroup
-{
-	TraceThreadGroup_Render,
-	TraceThreadGroup_AsyncLoading,
-	TraceThreadGroup_TaskGraphHigh,
-	TraceThreadGroup_TaskGraphNormal,
-	TraceThreadGroup_TaskGraphLow,
-	TraceThreadGroup_LargeThreadPool,
-	TraceThreadGroup_ThreadPool,
-	TraceThreadGroup_BackgroundThreadPool,
-	TraceThreadGroup_IOThreadPool,
-
-	TraceThreadGroup_Count
-};
-
 enum ETraceFrameType
 {
 	TraceFrameType_Game,
@@ -61,8 +46,8 @@ struct FMiscTrace
 {
 	CORE_API static void OutputRegisterGameThread(uint32 Id);
 	CORE_API static void OutputCreateThread(uint32 Id, const TCHAR* Name, uint32 Priority);
-	CORE_API static void OutputSetThreadGroup(uint32 Id, ETraceThreadGroup Group);
-	CORE_API static void OutputBeginThreadGroupScope(ETraceThreadGroup Group);
+	CORE_API static void OutputSetThreadGroup(uint32 Id, const ANSICHAR* GroupName);
+	CORE_API static void OutputBeginThreadGroupScope(const ANSICHAR* GroupName);
 	CORE_API static void OutputEndThreadGroupScope();
 	CORE_API static void OutputBookmarkSpec(const void* BookmarkPoint, const ANSICHAR* File, int32 Line, const TCHAR* Format);
 	template <typename... Types>
@@ -81,9 +66,9 @@ struct FMiscTrace
 
 	struct FThreadGroupScope
 	{
-		FThreadGroupScope(ETraceThreadGroup Group)
+		FThreadGroupScope(const ANSICHAR* GroupName)
 		{
-			FMiscTrace::OutputBeginThreadGroupScope(Group);
+			FMiscTrace::OutputBeginThreadGroupScope(GroupName);
 		}
 
 		~FThreadGroupScope()
diff --git a/Engine/Source/Runtime/CoreUObject/Private/Serialization/AsyncLoading.cpp b/Engine/Source/Runtime/CoreUObject/Private/Serialization/AsyncLoading.cpp
--- a/Engine/Source/Runtime/CoreUObject/Private/Serialization/AsyncLoading.cpp
+++ b/Engine/Source/Runtime/CoreUObject/Private/Serialization/AsyncLoading.cpp
@@ -5214,7 +5214,7 @@ void FAsyncLoadingThread::StartThread()
 		Thread = FRunnableThread::Create(this, TEXT("FAsyncLoadingThread"), 0, TPri_Normal);
 		if (Thread)
 		{
-			TRACE_SET_THREAD_GROUP(Thread->GetThreadID(), TraceThreadGroup_AsyncLoading);
+			TRACE_SET_THREAD_GROUP(Thread->GetThreadID(), "AsyncLoading");
 		}
 	}
 }
diff --git a/Engine/Source/Runtime/Launch/Private/LaunchEngineLoop.cpp b/Engine/Source/Runtime/Launch/Private/LaunchEngineLoop.cpp
--- a/Engine/Source/Runtime/Launch/Private/LaunchEngineLoop.cpp
+++ b/Engine/Source/Runtime/Launch/Private/LaunchEngineLoop.cpp
@@ -1775,7 +1775,7 @@ int32 FEngineLoop::PreInit(const TCHAR* CmdLine)
 
 
 		{
-			TRACE_THREAD_GROUP_SCOPE(TraceThreadGroup_ThreadPool);
+			TRACE_THREAD_GROUP_SCOPE("ThreadPool");
 			GThreadPool = FQueuedThreadPool::Allocate();
 			int32 NumThreadsInThreadPool = FPlatformMisc::NumberOfWorkerThreadsToSpawn();
 
@@ -1787,7 +1787,7 @@ int32 FEngineLoop::PreInit(const TCHAR* CmdLine)
 			verify(GThreadPool->Create(NumThreadsInThreadPool, StackSize * 1024, TPri_SlightlyBelowNormal));
 		}
 		{
-			TRACE_THREAD_GROUP_SCOPE(TraceThreadGroup_BackgroundThreadPool);
+			TRACE_THREAD_GROUP_SCOPE("BackgroundThreadPool");
 			GBackgroundPriorityThreadPool = FQueuedThreadPool::Allocate();
 			int32 NumThreadsInThreadPool = 2;
 			if (FPlatformProperties::IsServerOnly())
@@ -1800,7 +1800,7 @@ int32 FEngineLoop::PreInit(const TCHAR* CmdLine)
 
 #if WITH_EDITOR
 		{
-			TRACE_THREAD_GROUP_SCOPE(TraceThreadGroup_LargeThreadPool);
+			TRACE_THREAD_GROUP_SCOPE("LargeThreadPool");
 			// when we are in the editor we like to do things like build lighting and such
 			// this thread pool can be used for those purposes
 			GLargeThreadPool = FQueuedThreadPool::Allocate();
@@ -1857,7 +1857,7 @@ int32 FEngineLoop::PreInit(const TCHAR* CmdLine)
 	if (FPlatformProcess::SupportsMultithreading())
 	{
 		{
-			TRACE_THREAD_GROUP_SCOPE(TraceThreadGroup_IOThreadPool);
+			TRACE_THREAD_GROUP_SCOPE("IOThreadPool");
 			SCOPED_BOOT_TIMING("GIOThreadPool->Create");
 			GIOThreadPool = FQueuedThreadPool::Allocate();
 			int32 NumThreadsInThreadPool = FPlatformMisc::NumberOfIOWorkerThreadsToSpawn();
diff --git a/Engine/Source/Runtime/RenderCore/Private/RenderingThread.cpp b/Engine/Source/Runtime/RenderCore/Private/RenderingThread.cpp
--- a/Engine/Source/Runtime/RenderCore/Private/RenderingThread.cpp
+++ b/Engine/Source/Runtime/RenderCore/Private/RenderingThread.cpp
@@ -702,7 +702,7 @@ void StartRenderingThread()
 	GRenderingThreadRunnable = new FRenderingThread();
 
 	GRenderingThread = FRunnableThread::Create(GRenderingThreadRunnable, *BuildRenderingThreadName(ThreadCount), 0, FPlatformAffinity::GetRenderingThreadPriority(), FPlatformAffinity::GetRenderingThreadMask());
-	TRACE_SET_THREAD_GROUP(GRenderingThread->GetThreadID(), TraceThreadGroup_Render);
+	TRACE_SET_THREAD_GROUP(GRenderingThread->GetThreadID(), "Render");
 
 	// Wait for render thread to have taskgraph bound before we dispatch any tasks for it.
 	((FRenderingThread*)GRenderingThreadRunnable)->TaskGraphBoundSyncEvent->Wait();
