UnrealInsights:
Fixed current session time being displayed incorrectly

diff --git a/Engine/Source/Developer/TraceInsights/Private/Insights/ViewModels/TimingTrackViewport.h b/Engine/Source/Developer/TraceInsights/Private/Insights/ViewModels/TimingTrackViewport.h
--- a/Engine/Source/Developer/TraceInsights/Private/Insights/ViewModels/TimingTrackViewport.h
+++ b/Engine/Source/Developer/TraceInsights/Private/Insights/ViewModels/TimingTrackViewport.h
@@ -34,7 +34,7 @@ public:
 		Height = 0.0f;
 
 		MinValidTime = 0.0;
-		MaxValidTime = 60.0; // 1m
+		MaxValidTime = 0.0; // 1m
 		StartTime = 0.0;
 		EndTime = 0.0;
 		MinScaleX = (5 * 20) / 3600.0; // 1h between major tick marks
diff --git a/Engine/Source/Developer/TraceServices/Private/AnalysisService.cpp b/Engine/Source/Developer/TraceServices/Private/AnalysisService.cpp
--- a/Engine/Source/Developer/TraceServices/Private/AnalysisService.cpp
+++ b/Engine/Source/Developer/TraceServices/Private/AnalysisService.cpp
@@ -20,7 +20,7 @@ namespace Trace
 
 void FAnalysisSessionLock::ReadAccessCheck() const
 {
-	checkf(IsReadOnly, TEXT("Trying to read from session while not in read only mode"));
+	checkf(IsReadOnly || FPlatformTLS::GetCurrentThreadId() == OwnerThread, TEXT("Trying to read from session while someone else is writing"));
 }
 
 void FAnalysisSessionLock::WriteAccessCheck() const
diff --git a/Engine/Source/Developer/TraceServices/Private/Analyzers/GpuProfilerTraceAnalysis.cpp b/Engine/Source/Developer/TraceServices/Private/Analyzers/GpuProfilerTraceAnalysis.cpp
--- a/Engine/Source/Developer/TraceServices/Private/Analyzers/GpuProfilerTraceAnalysis.cpp
+++ b/Engine/Source/Developer/TraceServices/Private/Analyzers/GpuProfilerTraceAnalysis.cpp
@@ -44,11 +44,13 @@ void FGpuProfilerAnalyzer::OnEvent(uint16 RouteId, const FOnEventContext& Contex
 		uint32 CurrentDepth = 0;
 
 		uint64 LastTimestamp = EventData.GetValue("TimestampBase").As<uint64>();
+		double LastTime = 0.0;
 		while (BufferPtr < BufferEnd)
 		{
 			uint64 DecodedTimestamp = FTraceAnalyzerUtils::Decode7bit(BufferPtr);
 			uint64 ActualTimestamp = (DecodedTimestamp >> 1) + LastTimestamp;
 			LastTimestamp = ActualTimestamp;
+			LastTime = GpuTimestampToSessionTime(ActualTimestamp);
 			if (DecodedTimestamp & 1ull)
 			{
 				uint64 EventType = *reinterpret_cast<const uint64*>(BufferPtr);
@@ -56,16 +58,17 @@ void FGpuProfilerAnalyzer::OnEvent(uint16 RouteId, const FOnEventContext& Contex
 				check(EventTypeMap.Contains(EventType));
 				Trace::FTimingProfilerEvent Event;
 				Event.TimerIndex = EventTypeMap[EventType];
-				Timeline.AppendBeginEvent(GpuTimestampToSessionTime(ActualTimestamp), Event);
+				Timeline.AppendBeginEvent(LastTime, Event);
 				++CurrentDepth;
 			}
 			else
 			{
 				check(CurrentDepth > 0);
 				--CurrentDepth;
-				Timeline.AppendEndEvent(GpuTimestampToSessionTime(ActualTimestamp));
+				Timeline.AppendEndEvent(LastTime);
 			}
 		}
+		Session.UpdateDurationSeconds(LastTime);
 		check(BufferPtr == BufferEnd);
 		check(CurrentDepth == 0);
 		break;
diff --git a/Engine/Source/Developer/TraceServices/Private/Model/Bookmarks.cpp b/Engine/Source/Developer/TraceServices/Private/Model/Bookmarks.cpp
--- a/Engine/Source/Developer/TraceServices/Private/Model/Bookmarks.cpp
+++ b/Engine/Source/Developer/TraceServices/Private/Model/Bookmarks.cpp
@@ -40,6 +40,7 @@ void FBookmarkProvider::AppendBookmark(double Time, uint64 BookmarkPoint, const
 	FFormatArgsHelper::Format(FormatBuffer, FormatBufferSize - 1, Spec.FormatString, FormatArgs);
 	Bookmark->Text = Session.StoreString(FormatBuffer);
 	Bookmarks.Add(Bookmark);
+	Session.UpdateDurationSeconds(Time);
 }
 
 void FBookmarkProvider::EnumerateBookmarks(double IntervalStart, double IntervalEnd, TFunctionRef<void(const FBookmark &)> Callback) const
@@ -77,4 +78,9 @@ void FBookmarkProvider::EnumerateBookmarks(double IntervalStart, double Interval
 	}
 }
 
+const IBookmarkProvider& ReadBookmarkProvider(const IAnalysisSession& Session)
+{
+	return *Session.ReadProvider<IBookmarkProvider>(FBookmarkProvider::ProviderName);
+}
+
 }
\ No newline at end of file
diff --git a/Engine/Source/Developer/TraceServices/Private/Model/Frames.cpp b/Engine/Source/Developer/TraceServices/Private/Model/Frames.cpp
--- a/Engine/Source/Developer/TraceServices/Private/Model/Frames.cpp
+++ b/Engine/Source/Developer/TraceServices/Private/Model/Frames.cpp
@@ -52,6 +52,7 @@ void FFrameProvider::BeginFrame(ETraceFrameType FrameType, double Time)
 	Frame.StartTime = Time;
 	Frame.EndTime = std::numeric_limits<double>::infinity();
 	Frame.Index = Index;
+	Session.UpdateDurationSeconds(Time);
 }
 
 void FFrameProvider::EndFrame(ETraceFrameType FrameType, double Time)
@@ -59,6 +60,7 @@ void FFrameProvider::EndFrame(ETraceFrameType FrameType, double Time)
 	Session.WriteAccessCheck();
 	FFrame& Frame = Frames[FrameType][Frames[FrameType].Num() - 1];
 	Frame.EndTime = Time;
+	Session.UpdateDurationSeconds(Time);
 }
 
 const IFrameProvider& ReadFrameProvider(const IAnalysisSession& Session)
diff --git a/Engine/Source/Developer/TraceServices/Private/Model/Log.cpp b/Engine/Source/Developer/TraceServices/Private/Model/Log.cpp
--- a/Engine/Source/Developer/TraceServices/Private/Model/Log.cpp
+++ b/Engine/Source/Developer/TraceServices/Private/Model/Log.cpp
@@ -59,6 +59,7 @@ void FLogProvider::AppendMessage(uint64 LogPoint, double Time, const uint8* Form
 	InternalMessage.Spec = SpecMap[LogPoint];
 	FFormatArgsHelper::Format(FormatBuffer, FormatBufferSize - 1, InternalMessage.Spec->FormatString, FormatArgs);
 	InternalMessage.Message = Session.StoreString(FormatBuffer);
+	Session.UpdateDurationSeconds(Time);
 }
 
 uint64 FLogProvider::GetMessageCount() const
