Unreal Insights:
Exposing thread ids from the LoadTimeProfiler provider

diff --git a/Engine/Source/Developer/TraceServices/Private/Analyzers/LoadTimeTraceAnalysis.cpp b/Engine/Source/Developer/TraceServices/Private/Analyzers/LoadTimeTraceAnalysis.cpp
--- a/Engine/Source/Developer/TraceServices/Private/Analyzers/LoadTimeTraceAnalysis.cpp
+++ b/Engine/Source/Developer/TraceServices/Private/Analyzers/LoadTimeTraceAnalysis.cpp
@@ -27,6 +27,7 @@ TSharedRef<FAsyncLoadingTraceAnalyzer::FThreadState> FAsyncLoadingTraceAnalyzer:
 		if (MainThreadId == -1)
 		{
 			MainThreadId = ThreadId;
+			LoadTimeProfilerProvider.SetMainThreadId(MainThreadId);
 			ThreadState->CpuTimeline = LoadTimeProfilerProvider.EditMainThreadCpuTimeline();
 		}
 		else if (ThreadId == AsyncLoadingThreadId)
@@ -103,7 +104,9 @@ void FAsyncLoadingTraceAnalyzer::OnEvent(uint16 RouteId, const FOnEventContext&
 	{
 	case RouteId_StartAsyncLoading:
 	{
+		Trace::FAnalysisSessionEditScope _(Session);
 		AsyncLoadingThreadId = EventData.GetValue("ThreadId").As<uint32>();
+		LoadTimeProfilerProvider.SetAsyncLoadingThreadId(AsyncLoadingThreadId);
 		break;
 	}
 	case RouteId_NewLinker:
diff --git a/Engine/Source/Developer/TraceServices/Private/Model/LoadTimeProfilerPrivate.h b/Engine/Source/Developer/TraceServices/Private/Model/LoadTimeProfilerPrivate.h
--- a/Engine/Source/Developer/TraceServices/Private/Model/LoadTimeProfilerPrivate.h
+++ b/Engine/Source/Developer/TraceServices/Private/Model/LoadTimeProfilerPrivate.h
@@ -33,6 +33,10 @@ public:
 	FPackageExportInfo& CreateExport();
 	TSharedRef<CpuTimelineInternal> EditMainThreadCpuTimeline() { return MainThreadCpuTimeline; }
 	TSharedRef<CpuTimelineInternal> EditAsyncLoadingThreadCpuTimeline() { return AsyncLoadingThreadCpuTimeline; }
+	virtual uint32 GetMainThreadId() const override { return MainThreadId; }
+	void SetMainThreadId(uint32 ThreadId) { MainThreadId = ThreadId; }
+	virtual uint32 GetAsyncLoadingThreadId() const override { return AsyncLoadingThreadId; }
+	void SetAsyncLoadingThreadId(uint32 ThreadId) { AsyncLoadingThreadId = ThreadId; }
 
 private:
 	IAnalysisSession& Session;
@@ -41,6 +45,8 @@ private:
 	TPagedArray<FPackageExportInfo> Exports;
 	TSharedRef<CpuTimelineInternal> MainThreadCpuTimeline;
 	TSharedRef<CpuTimelineInternal> AsyncLoadingThreadCpuTimeline;
+	uint32 MainThreadId = uint32(-1);
+	uint32 AsyncLoadingThreadId = uint32(-1);
 
 	UE_TRACE_TABLE_LAYOUT_BEGIN(FAggregatedStatsTableLayout, FLoadTimeProfilerAggregatedStats)
 		UE_TRACE_TABLE_COLUMN(Name, TEXT("Name"))
diff --git a/Engine/Source/Developer/TraceServices/Public/TraceServices/Model/LoadTimeProfiler.h b/Engine/Source/Developer/TraceServices/Public/TraceServices/Model/LoadTimeProfiler.h
--- a/Engine/Source/Developer/TraceServices/Public/TraceServices/Model/LoadTimeProfiler.h
+++ b/Engine/Source/Developer/TraceServices/Public/TraceServices/Model/LoadTimeProfiler.h
@@ -102,8 +102,10 @@ public:
 
 	virtual ~ILoadTimeProfilerProvider() = default;
 	virtual uint64 GetPackageCount() const = 0;
-	virtual void EnumeratePackages(TFunctionRef<void(const FPackageInfo&)> Callback) const = 0; 
+	virtual void EnumeratePackages(TFunctionRef<void(const FPackageInfo&)> Callback) const = 0;
+	virtual uint32 GetMainThreadId() const = 0;
 	virtual void ReadMainThreadCpuTimeline(TFunctionRef<void(const CpuTimeline&)> Callback) const = 0;
+	virtual uint32 GetAsyncLoadingThreadId() const = 0;
 	virtual void ReadAsyncLoadingThreadCpuTimeline(TFunctionRef<void(const CpuTimeline&)> Callback) const = 0;
 	virtual ITable<FLoadTimeProfilerAggregatedStats>* CreateEventAggregation(double IntervalStart, double IntervalEnd) const = 0;
 	virtual ITable<FLoadTimeProfilerAggregatedStats>* CreateObjectTypeAggregation(double IntervalStart, double IntervalEnd) const = 0;
