UnrealInsights:
Added functionality to toggle modules and their events on or off per session
Added ISessionService::ConnectSession for connecting running game instances

diff --git a/Engine/Source/Developer/TraceServices/Private/ModuleService.cpp b/Engine/Source/Developer/TraceServices/Private/ModuleService.cpp
--- a/Engine/Source/Developer/TraceServices/Private/ModuleService.cpp
+++ b/Engine/Source/Developer/TraceServices/Private/ModuleService.cpp
@@ -54,6 +54,11 @@ void FModuleService::SetModuleEnabled(const FName& ModuleName, bool bEnabled)
 	{
 		return;
 	}
+	bool bWasEnabled = EnabledModules.Find(*FindIt);
+	if (bEnabled == bWasEnabled)
+	{
+		return;
+	}
 	if (bEnabled)
 	{
 		EnabledModules.Add(*FindIt);
@@ -75,4 +80,17 @@ void FModuleService::OnAnalysisBegin(IAnalysisSession& Session, TArray<IAnalyzer
 	}
 }
 
+bool FModuleService::GetModuleLoggers(const FName& ModuleName, TArray<const TCHAR *>& OutLoggers)
+{
+	FScopeLock Lock(&CriticalSection);
+	Initialize();
+	IModule** FindIt = ModulesMap.Find(ModuleName);
+	if (!FindIt)
+	{
+		return false;
+	}
+	(*FindIt)->GetLoggers(OutLoggers);
+	return true;
+}
+
 }
diff --git a/Engine/Source/Developer/TraceServices/Private/ModuleServicePrivate.h b/Engine/Source/Developer/TraceServices/Private/ModuleServicePrivate.h
--- a/Engine/Source/Developer/TraceServices/Private/ModuleServicePrivate.h
+++ b/Engine/Source/Developer/TraceServices/Private/ModuleServicePrivate.h
@@ -20,6 +20,7 @@ public:
 	virtual void GetAvailableModules(TArray<FModuleInfo>& OutModules) override;
 	virtual void SetModuleEnabled(const FName& ModuleName, bool bEnabled) override;
 	void OnAnalysisBegin(IAnalysisSession& Session, TArray<IAnalyzer*>& OutAnalyzers);
+	bool GetModuleLoggers(const FName& ModuleName, TArray<const TCHAR*>& OutLoggers);
 
 private:
 	void Initialize();
diff --git a/Engine/Source/Developer/TraceServices/Private/Modules/LoadTimeProfilerModule.cpp b/Engine/Source/Developer/TraceServices/Private/Modules/LoadTimeProfilerModule.cpp
--- a/Engine/Source/Developer/TraceServices/Private/Modules/LoadTimeProfilerModule.cpp
+++ b/Engine/Source/Developer/TraceServices/Private/Modules/LoadTimeProfilerModule.cpp
@@ -30,6 +30,12 @@ void FLoadTimeProfilerModule::OnAnalysisBegin(IAnalysisSession& Session, bool bI
 	OutAnalyzers.Add(new FPlatformFileTraceAnalyzer(Session, *FileActivityProvider));
 }
 
+void FLoadTimeProfilerModule::GetLoggers(TArray<const TCHAR *>& OutLoggers)
+{
+	OutLoggers.Add(TEXT("LoadTime"));
+	OutLoggers.Add(TEXT("PlatformFile"));
+}
+
 const ILoadTimeProfilerProvider* ReadLoadTimeProfilerProvider(const IAnalysisSession& Session)
 {
 	return Session.ReadProvider<ILoadTimeProfilerProvider>(LoadTimeProfilerProviderName);
diff --git a/Engine/Source/Developer/TraceServices/Private/Modules/LoadTimeProfilerModule.h b/Engine/Source/Developer/TraceServices/Private/Modules/LoadTimeProfilerModule.h
--- a/Engine/Source/Developer/TraceServices/Private/Modules/LoadTimeProfilerModule.h
+++ b/Engine/Source/Developer/TraceServices/Private/Modules/LoadTimeProfilerModule.h
@@ -13,6 +13,7 @@ class FLoadTimeProfilerModule
 public:
 	virtual void GetModuleInfo(FModuleInfo& OutModuleInfo) override;
 	virtual void OnAnalysisBegin(IAnalysisSession& Session, bool bIsEnabled, TArray<IAnalyzer*>& OutAnalyzers) override;
+	virtual void GetLoggers(TArray<const TCHAR *>& OutLoggers) override;
 
 private:
 	static FName ModuleName;
diff --git a/Engine/Source/Developer/TraceServices/Private/Modules/StatsModule.cpp b/Engine/Source/Developer/TraceServices/Private/Modules/StatsModule.cpp
--- a/Engine/Source/Developer/TraceServices/Private/Modules/StatsModule.cpp
+++ b/Engine/Source/Developer/TraceServices/Private/Modules/StatsModule.cpp
@@ -27,6 +27,11 @@ void FStatsModule::OnAnalysisBegin(IAnalysisSession& InSession, bool bIsEnabled,
 	OutAnalyzers.Add(new FStatsAnalyzer(Session, *CounterProvider));
 }
 
+void FStatsModule::GetLoggers(TArray<const TCHAR *>& OutLoggers)
+{
+	OutLoggers.Add(TEXT("Stats"));
+}
+
 const ICounterProvider* ReadCounterProvider(const IAnalysisSession& Session)
 {
 	return Session.ReadProvider<ICounterProvider>(CounterProviderName);
diff --git a/Engine/Source/Developer/TraceServices/Private/Modules/StatsModule.h b/Engine/Source/Developer/TraceServices/Private/Modules/StatsModule.h
--- a/Engine/Source/Developer/TraceServices/Private/Modules/StatsModule.h
+++ b/Engine/Source/Developer/TraceServices/Private/Modules/StatsModule.h
@@ -14,6 +14,7 @@ class FStatsModule
 public:
 	virtual void GetModuleInfo(FModuleInfo& OutModuleInfo) override;
 	virtual void OnAnalysisBegin(IAnalysisSession& Session, bool bIsEnabled, TArray<IAnalyzer*>& OutAnalyzers) override;
+	virtual void GetLoggers(TArray<const TCHAR *>& OutLoggers) override;
 
 private:
 	static FName ModuleName;
diff --git a/Engine/Source/Developer/TraceServices/Private/Modules/TimingProfilerModule.cpp b/Engine/Source/Developer/TraceServices/Private/Modules/TimingProfilerModule.cpp
--- a/Engine/Source/Developer/TraceServices/Private/Modules/TimingProfilerModule.cpp
+++ b/Engine/Source/Developer/TraceServices/Private/Modules/TimingProfilerModule.cpp
@@ -29,6 +29,12 @@ void FTimingProfilerModule::OnAnalysisBegin(IAnalysisSession& InSession, bool bI
 	OutAnalyzers.Add(new FGpuProfilerAnalyzer(Session, *TimingProfilerProvider));
 }
 
+void FTimingProfilerModule::GetLoggers(TArray<const TCHAR *>& OutLoggers)
+{
+	OutLoggers.Add(TEXT("CpuProfiler"));
+	OutLoggers.Add(TEXT("GpuProfiler"));
+}
+
 const ITimingProfilerProvider* ReadTimingProfilerProvider(const IAnalysisSession& Session)
 {
 	return Session.ReadProvider<ITimingProfilerProvider>(TimingProfilerProviderName);
diff --git a/Engine/Source/Developer/TraceServices/Private/Modules/TimingProfilerModule.h b/Engine/Source/Developer/TraceServices/Private/Modules/TimingProfilerModule.h
--- a/Engine/Source/Developer/TraceServices/Private/Modules/TimingProfilerModule.h
+++ b/Engine/Source/Developer/TraceServices/Private/Modules/TimingProfilerModule.h
@@ -14,6 +14,7 @@ class FTimingProfilerModule
 public:
 	virtual void GetModuleInfo(FModuleInfo& OutModuleInfo) override;
 	virtual void OnAnalysisBegin(IAnalysisSession& Session, bool bIsEnabled, TArray<IAnalyzer*>& OutAnalyzers) override;
+	virtual void GetLoggers(TArray<const TCHAR *>& OutLoggers) override;
 
 private:
 	static FName ModuleName;
diff --git a/Engine/Source/Developer/TraceServices/Private/SessionService.cpp b/Engine/Source/Developer/TraceServices/Private/SessionService.cpp
--- a/Engine/Source/Developer/TraceServices/Private/SessionService.cpp
+++ b/Engine/Source/Developer/TraceServices/Private/SessionService.cpp
@@ -5,12 +5,15 @@
 #include "Trace/DataStream.h"
 #include "Misc/Paths.h"
 #include "Misc/OutputDeviceRedirector.h"
+#include "SocketSubsystem.h"
 #include "Trace/ControlClient.h"
+#include "IPAddress.h"
 
 namespace Trace
 {
 
-FSessionService::FSessionService()
+FSessionService::FSessionService(FModuleService& InModuleService)
+	: ModuleService(InModuleService)
 {
 	LocalSessionDirectory = FPaths::ProjectSavedDir() / TEXT("TraceSessions");
 	TraceStore = Store_Create(*LocalSessionDirectory);
@@ -91,6 +94,101 @@ Trace::IInDataStream* FSessionService::OpenSessionFromFile(const TCHAR* FilePath
 	return Trace::DataStream_ReadFile(FilePath);
 }
 
+void FSessionService::SetModuleEnabled(FSessionHandle SessionHandle, const FName& ModuleName, bool bState)
+{
+	FScopeLock Lock(&SessionsCS);
+	FSessionInfoInternal* FindIt = Sessions.Find(SessionHandle);
+	if (!FindIt)
+	{
+		return;
+	}
+	if (bState)
+	{
+		TArray<const TCHAR*> Loggers;
+		if (!ModuleService.GetModuleLoggers(ModuleName, Loggers))
+		{
+			return;
+		}
+		TSet<FString>& EnabledLoggers = FindIt->EnabledModuleLoggersMap.FindOrAdd(ModuleName);
+		for (const TCHAR* Logger : Loggers)
+		{
+			EnabledLoggers.Add(Logger);
+		}
+		if (FindIt->RecorderSessionHandle)
+		{
+			for (const FString& Logger : EnabledLoggers)
+			{
+				TraceRecorder->ToggleEvent(FindIt->RecorderSessionHandle, *Logger, true);
+			}
+		}
+	}
+	else
+	{
+		TSet<FString>* EnabledLoggers = FindIt->EnabledModuleLoggersMap.Find(ModuleName);
+		if (EnabledLoggers)
+		{
+			if (FindIt->RecorderSessionHandle)
+			{
+				for (const FString& Logger : *EnabledLoggers)
+				{
+					TraceRecorder->ToggleEvent(FindIt->RecorderSessionHandle, *Logger, false);
+				}
+			}
+			FindIt->EnabledModuleLoggersMap.Remove(ModuleName);
+		}
+	}
+}
+
+bool FSessionService::IsModuleEnabled(Trace::FSessionHandle SessionHandle, const FName& ModuleName) const
+{
+	FScopeLock Lock(&SessionsCS);
+	const FSessionInfoInternal* FindIt = Sessions.Find(SessionHandle);
+	if (!FindIt)
+	{
+		return false;
+	}
+	return FindIt->EnabledModuleLoggersMap.Contains(ModuleName);
+}
+
+bool FSessionService::ConnectSession(const TCHAR* ControlClientAddress)
+{
+	ISocketSubsystem* Sockets = ISocketSubsystem::Get();
+	if (!Sockets)
+	{
+		return false;
+	}
+	bool bCanBindAll = false;
+	TSharedPtr<FInternetAddr> RecorderAddr = Sockets->GetLocalHostAddr(*GLog, bCanBindAll);
+	if (!RecorderAddr.IsValid())
+	{
+		return false;
+	}
+	uint16 Port = 1985;
+	FString AddressString = ControlClientAddress;
+	const int32 LastColonIndex = AddressString.Find(":", ESearchCase::IgnoreCase, ESearchDir::FromEnd);
+	if (INDEX_NONE != LastColonIndex)
+	{
+		FString PortString = AddressString.RightChop(LastColonIndex + 1);
+		Port = FCString::Atoi(*PortString);
+		AddressString = AddressString.Left(LastColonIndex);
+	}
+	TSharedPtr<FInternetAddr> ClientAddr = Sockets->CreateInternetAddr();
+	ESocketErrors SocketErrors = Sockets->GetHostByName(TCHAR_TO_ANSI(*AddressString), *ClientAddr);
+	if (SocketErrors != SE_NO_ERROR)
+	{
+		return false;
+	}
+	ClientAddr->SetPort(Port);
+	FControlClient ControlClient;
+	if (!ControlClient.Connect(*ClientAddr))
+	{
+		return false;
+	}
+	ControlClient.SendConnect(*RecorderAddr->ToString(false));
+	ControlClient.Disconnect();
+	return true;
+}
+
 bool FSessionService::Tick(float DeltaTime)
 {
 	UpdateSessions();
diff --git a/Engine/Source/Developer/TraceServices/Private/SessionServicePrivate.h b/Engine/Source/Developer/TraceServices/Private/SessionServicePrivate.h
--- a/Engine/Source/Developer/TraceServices/Private/SessionServicePrivate.h
+++ b/Engine/Source/Developer/TraceServices/Private/SessionServicePrivate.h
@@ -7,6 +7,7 @@
 #include "Containers/Ticker.h"
 #include "Trace/Recorder.h"
 #include "Trace/Store.h"
+#include "ModuleServicePrivate.h"
 
 namespace Trace
 {
@@ -15,7 +16,7 @@ class FSessionService
 	: public ISessionService
 {
 public:
-	FSessionService();
+	FSessionService(FModuleService& ModuleService);
 	virtual ~FSessionService();
 	virtual bool StartRecorderServer() override;
 	virtual bool IsRecorderServerRunning() const override;
@@ -26,6 +27,9 @@ public:
 	virtual bool GetSessionInfo(Trace::FSessionHandle SessionHandle, Trace::FSessionInfo& OutSessionInfo) const override;
 	virtual Trace::IInDataStream* OpenSessionStream(Trace::FSessionHandle SessionHandle) override;
 	virtual Trace::IInDataStream* OpenSessionFromFile(const TCHAR* FilePath) override;
+	virtual void SetModuleEnabled(Trace::FSessionHandle SessionHandle, const FName& ModuleName, bool bState) override;
+	virtual bool IsModuleEnabled(Trace::FSessionHandle SessionHandle, const FName& ModuleName) const override;
+	virtual bool ConnectSession(const TCHAR* ControlClientAddress) override;
 
 private:
 	bool Tick(float DeltaTime);
@@ -36,9 +40,11 @@ private:
 		FRecorderSessionHandle RecorderSessionHandle;
 		const TCHAR* Uri;
 		const TCHAR* Name;
+		TMap<FName, TSet<FString>> EnabledModuleLoggersMap;
 		bool bIsLive;
 	};
 
+	FModuleService& ModuleService;
 	FString LocalSessionDirectory;
 	TSharedPtr<IStore> TraceStore;
 	TSharedPtr<IRecorder> TraceRecorder;
diff --git a/Engine/Source/Developer/TraceServices/Private/TraceServicesModule.cpp b/Engine/Source/Developer/TraceServices/Private/TraceServicesModule.cpp
--- a/Engine/Source/Developer/TraceServices/Private/TraceServicesModule.cpp
+++ b/Engine/Source/Developer/TraceServices/Private/TraceServicesModule.cpp
@@ -35,7 +35,8 @@ TSharedPtr<Trace::ISessionService> FTraceServicesModule::GetSessionService()
 {
 	if (!SessionService.IsValid())
 	{
-		SessionService = MakeShared<Trace::FSessionService>();
+		GetModuleService();
+		SessionService = MakeShared<Trace::FSessionService>(*ModuleService.Get());
 	}
 	return SessionService;
 }
diff --git a/Engine/Source/Developer/TraceServices/Public/TraceServices/ModuleService.h b/Engine/Source/Developer/TraceServices/Public/TraceServices/ModuleService.h
--- a/Engine/Source/Developer/TraceServices/Public/TraceServices/ModuleService.h
+++ b/Engine/Source/Developer/TraceServices/Public/TraceServices/ModuleService.h
@@ -26,6 +26,7 @@ class IModule
 public:
 	virtual void GetModuleInfo(FModuleInfo& OutModuleInfo) = 0;
 	virtual void OnAnalysisBegin(IAnalysisSession& Session, bool bIsEnabled, TArray<IAnalyzer*>& OutAnalyzers) = 0;
+	virtual void GetLoggers(TArray<const TCHAR*>& OutLoggers) = 0;
 };
 
 class IModuleService
diff --git a/Engine/Source/Developer/TraceServices/Public/TraceServices/SessionService.h b/Engine/Source/Developer/TraceServices/Public/TraceServices/SessionService.h
--- a/Engine/Source/Developer/TraceServices/Public/TraceServices/SessionService.h
+++ b/Engine/Source/Developer/TraceServices/Public/TraceServices/SessionService.h
@@ -32,6 +32,9 @@ public:
 	virtual bool GetSessionInfo(FSessionHandle SessionHandle, FSessionInfo& OutSessionInfo) const = 0;
 	virtual Trace::IInDataStream* OpenSessionStream(FSessionHandle SessionHandle) = 0;
 	virtual Trace::IInDataStream* OpenSessionFromFile(const TCHAR* FilePath) = 0;
+	virtual void SetModuleEnabled(Trace::FSessionHandle SessionHandle, const FName& ModuleName, bool bState) = 0;
+	virtual bool IsModuleEnabled(Trace::FSessionHandle SessionHandle, const FName& ModuleName) const = 0;
+	virtual bool ConnectSession(const TCHAR* ControlClientAddress) = 0;
 };
 
 }
diff --git a/Engine/Source/Developer/TraceServices/TraceServices.Build.cs b/Engine/Source/Developer/TraceServices/TraceServices.Build.cs
--- a/Engine/Source/Developer/TraceServices/TraceServices.Build.cs
+++ b/Engine/Source/Developer/TraceServices/TraceServices.Build.cs
@@ -10,6 +10,7 @@ public class TraceServices : ModuleRules
 			new string[] {
 				"Core",
 				"CoreUObject",
+				"Sockets",
 				"TraceAnalysis",
 			});
 
