Use a named semaphore to indicate to other local processes that a the trace recorder is up and running.

diff --git a/Engine/Source/Developer/TraceServices/Private/SessionService.cpp b/Engine/Source/Developer/TraceServices/Private/SessionService.cpp
--- a/Engine/Source/Developer/TraceServices/Private/SessionService.cpp
+++ b/Engine/Source/Developer/TraceServices/Private/SessionService.cpp
@@ -34,7 +34,12 @@ FSessionService::~FSessionService()
 
 bool FSessionService::StartRecorderServer()
 {
-	return TraceRecorder->StartRecording();
+	bool bOk = TraceRecorder->StartRecording();
+	if (bOk && RecorderMutex == nullptr)
+	{
+		RecorderMutex = FPlatformProcess::NewInterprocessSynchObject("UnrealInsightsRecorder", true);
+	}
+	return bOk;
 }
 
 bool FSessionService::IsRecorderServerRunning() const
@@ -44,6 +49,11 @@ bool FSessionService::IsRecorderServerRunning() const
 
 void FSessionService::StopRecorderServer()
 {
+	if (RecorderMutex != nullptr)
+	{
+		FPlatformProcess::DeleteInterprocessSynchObject(RecorderMutex);
+		RecorderMutex = nullptr;
+	}
 	TraceRecorder->StopRecording();
 }
 
diff --git a/Engine/Source/Developer/TraceServices/Private/SessionServicePrivate.h b/Engine/Source/Developer/TraceServices/Private/SessionServicePrivate.h
--- a/Engine/Source/Developer/TraceServices/Private/SessionServicePrivate.h
+++ b/Engine/Source/Developer/TraceServices/Private/SessionServicePrivate.h
@@ -3,6 +3,7 @@
 #pragma once
 
 #include "TraceServices/SessionService.h"
+#include "HAL/PlatformProcess.h"
 #include "Misc/ScopeLock.h"
 #include "Containers/Ticker.h"
 #include "Trace/Recorder.h"
@@ -45,6 +46,7 @@ private:
 	};
 
 	FModuleService& ModuleService;
+	FPlatformProcess::FSemaphore* RecorderMutex = nullptr;
 	FString LocalSessionDirectory;
 	TSharedPtr<IStore> TraceStore;
 	TSharedPtr<IRecorder> TraceRecorder;
diff --git a/Engine/Source/Runtime/Launch/Private/LaunchEngineLoop.cpp b/Engine/Source/Runtime/Launch/Private/LaunchEngineLoop.cpp
--- a/Engine/Source/Runtime/Launch/Private/LaunchEngineLoop.cpp
+++ b/Engine/Source/Runtime/Launch/Private/LaunchEngineLoop.cpp
@@ -1214,15 +1214,22 @@ int32 FEngineLoop::PreInit(const TCHAR* CmdLine)
 	}
 
 	{
-		const TCHAR* TraceHostPtr = TEXT("127.0.0.1");
-
 		FString TraceHost;
 		if (FParse::Value(CmdLine, TEXT("tracehost"), TraceHost))
 		{
-			TraceHostPtr = *TraceHost;
+			Trace::Connect(*TraceHost);
 		}
 
-		Trace::Connect(TraceHostPtr);
+#if PLATFORM_WINDOWS && !UE_BUILD_SHIPPING
+		else
+		{
+			if (FPlatformProcess::FSemaphore* Mutex = FPlatformProcess::NewInterprocessSynchObject("UnrealInsightsRecorder", false))
+			{
+				Trace::Connect(TEXT("127.0.0.1"));
+				FPlatformProcess::DeleteInterprocessSynchObject(Mutex);
+			}
+		}
+#endif // PLATFORM_WINDOWS
 	}
 
 #if WITH_ENGINE
