Static analysis fixes

diff --git a/Engine/Source/Developer/TraceServices/Private/Analyzers/MiscTraceAnalysis.cpp b/Engine/Source/Developer/TraceServices/Private/Analyzers/MiscTraceAnalysis.cpp
--- a/Engine/Source/Developer/TraceServices/Private/Analyzers/MiscTraceAnalysis.cpp
+++ b/Engine/Source/Developer/TraceServices/Private/Analyzers/MiscTraceAnalysis.cpp
@@ -137,29 +137,26 @@ void FMiscTraceAnalyzer::OnEvent(uint16 RouteId, const FOnEventContext& Context)
 	case RouteId_BeginRenderFrame:
 	case RouteId_EndRenderFrame:
 	{
-		ETraceFrameType FrameType = TraceFrameType_Count;
+		ETraceFrameType FrameType;
 		if (RouteId == RouteId_BeginGameFrame || RouteId == RouteId_EndGameFrame)
 		{
 			FrameType = TraceFrameType_Game;
 		}
-		else if (RouteId == RouteId_BeginRenderFrame || RouteId == RouteId_EndRenderFrame)
+		else
 		{
 			FrameType = TraceFrameType_Rendering;
 		}
-		if (FrameType != TraceFrameType_Count)
+		const uint8* BufferPtr = EventData.GetAttachment();
+		uint64 CycleDiff = FTraceAnalyzerUtils::Decode7bit(BufferPtr);
+		uint64 Cycle = LastFrameCycle[FrameType] + CycleDiff;
+		LastFrameCycle[FrameType] = Cycle;
+		if (RouteId == RouteId_BeginGameFrame || RouteId == RouteId_BeginRenderFrame)
 		{
-			const uint8* BufferPtr = EventData.GetAttachment();
-			uint64 CycleDiff = FTraceAnalyzerUtils::Decode7bit(BufferPtr);
-			uint64 Cycle = LastFrameCycle[FrameType] + CycleDiff;
-			LastFrameCycle[FrameType] = Cycle;
-			if (RouteId == RouteId_BeginGameFrame || RouteId == RouteId_BeginRenderFrame)
-			{
-				FrameProvider.BeginFrame(FrameType, Context.SessionContext.TimestampFromCycle(Cycle));
-			}
-			else
-			{
-				FrameProvider.EndFrame(FrameType, Context.SessionContext.TimestampFromCycle(Cycle));
-			}
+			FrameProvider.BeginFrame(FrameType, Context.SessionContext.TimestampFromCycle(Cycle));
+		}
+		else
+		{
+			FrameProvider.EndFrame(FrameType, Context.SessionContext.TimestampFromCycle(Cycle));
 		}
 		break;
 	}
diff --git a/Engine/Source/Runtime/Core/Private/ProfilingDebugging/CpuProfilerTrace.cpp b/Engine/Source/Runtime/Core/Private/ProfilingDebugging/CpuProfilerTrace.cpp
--- a/Engine/Source/Runtime/Core/Private/ProfilingDebugging/CpuProfilerTrace.cpp
+++ b/Engine/Source/Runtime/Core/Private/ProfilingDebugging/CpuProfilerTrace.cpp
@@ -81,7 +81,7 @@ void FCpuProfilerTrace::OutputBeginEvent(uint16 SpecId)
 	}
 	++ThreadState->Depth;
 	bool bEventEnabled = UE_TRACE_EVENT_IS_ENABLED(CpuProfiler, EventBatch);
-	if (!bEventEnabled & (bEventEnabled == ThreadState->bEnabled))
+	if ((!bEventEnabled) & (bEventEnabled == ThreadState->bEnabled))
 	{
 		return;
 	}
@@ -109,7 +109,7 @@ void FCpuProfilerTrace::OutputEndEvent()
 	FCpuProfilerTraceInternal::FThreadState* ThreadState = FCpuProfilerTraceInternal::ThreadLocalThreadState;
 	--ThreadState->Depth;
 	bool bEventEnabled = UE_TRACE_EVENT_IS_ENABLED(CpuProfiler, EventBatch);
-	if (!bEventEnabled & (bEventEnabled == ThreadState->bEnabled))
+	if ((!bEventEnabled) & (bEventEnabled == ThreadState->bEnabled))
 	{
 		return;
 	}
