UnrealInsights:
Made it possible to toggle load time tracing on or off dynamically
Load time tracing compiled in by default
Added -loadtimetrace command line option to enable load time tracing from boot
Added missing DeferredPostLoad load time trace event type

diff --git a/Engine/Source/Developer/TraceInsights/Private/Insights/Widgets/STimingView.cpp b/Engine/Source/Developer/TraceInsights/Private/Insights/Widgets/STimingView.cpp
--- a/Engine/Source/Developer/TraceInsights/Private/Insights/Widgets/STimingView.cpp
+++ b/Engine/Source/Developer/TraceInsights/Private/Insights/Widgets/STimingView.cpp
@@ -1377,6 +1377,7 @@ const TCHAR* GetName(ELoadTimeProfilerPackageEventType Type)
 		case LoadTimeProfilerPackageEventType_StartPostLoad:			return TEXT("StartPostLoad");
 		case LoadTimeProfilerPackageEventType_Tick:						return TEXT("Tick");
 		case LoadTimeProfilerPackageEventType_Finish:					return TEXT("Finish");
+		case LoadTimeProfilerPackageEventType_DeferredPostLoad:			return TEXT("DeferredPostLoad");
 		case LoadTimeProfilerPackageEventType_None:						return TEXT("None");
 		default:														return TEXT("");
 	}
diff --git a/Engine/Source/Developer/TraceServices/Private/Analyzers/LoadTimeTraceAnalysis.cpp b/Engine/Source/Developer/TraceServices/Private/Analyzers/LoadTimeTraceAnalysis.cpp
--- a/Engine/Source/Developer/TraceServices/Private/Analyzers/LoadTimeTraceAnalysis.cpp
+++ b/Engine/Source/Developer/TraceServices/Private/Analyzers/LoadTimeTraceAnalysis.cpp
@@ -127,12 +127,11 @@ void FAsyncLoadingTraceAnalyzer::OnEvent(uint16 RouteId, const FOnEventContext&
 	case RouteId_PackageSummary:
 	{
 		uint64 LinkerPtr = EventData.GetValue("Linker").As<uint64>();
-		//check(ActiveLinkersMap.Contains(LinkerPtr));
-		TSharedRef<FLinkerState> LinkerState = ActiveLinkersMap[LinkerPtr];
-		if (LinkerState->PackageInfo)
+		TSharedRef<FLinkerState>* LinkerState = ActiveLinkersMap.Find(LinkerPtr);
+		if (LinkerState && (*LinkerState)->PackageInfo)
 		{
 			Trace::FAnalysisSessionEditScope _(Session);
-			Trace::FPackageSummaryInfo& Summary = LinkerState->PackageInfo->Summary;
+			Trace::FPackageSummaryInfo& Summary = (*LinkerState)->PackageInfo->Summary;
 			Summary.TotalHeaderSize = EventData.GetValue("TotalHeaderSize").As<uint32>();
 			Summary.NameCount = EventData.GetValue("NameCount").As<uint32>();
 			Summary.ImportCount = EventData.GetValue("ImportCount").As<uint32>();
@@ -144,35 +143,36 @@ void FAsyncLoadingTraceAnalyzer::OnEvent(uint16 RouteId, const FOnEventContext&
 	{
 		Trace::FAnalysisSessionEditScope _(Session);
 
-		uint64 LinkerPtr = EventData.GetValue("Linker").As<uint64>();
-		//check(ActiveLinkersMap.Contains(LinkerPtr));
-		TSharedRef<FLinkerState> LinkerState = ActiveLinkersMap[LinkerPtr];
-		uint32 ThreadId = EventData.GetValue("ThreadId").As<uint32>();
-		TSharedRef<FThreadState> ThreadState = GetThreadState(ThreadId);
 		Trace::FPackageExportInfo& Export = LoadTimeProfilerProvider.CreateExport();
 		Export.SerialOffset = EventData.GetValue("SerialOffset").As<uint64>();
 		Export.SerialSize = EventData.GetValue("SerialSize").As<uint64>();
 		Export.IsAsset = EventData.GetValue("IsAsset").As<bool>();
-		if (LinkerState->PackageInfo)
+		uint32 ThreadId = EventData.GetValue("ThreadId").As<uint32>();
+		TSharedRef<FThreadState> ThreadState = GetThreadState(ThreadId);
+		ThreadState->EnterExportScope(Context.SessionContext.TimestampFromCycle(EventData.GetValue("Cycle").As<uint64>()), &Export, LoadTimeProfilerObjectEventType_Create);
+		uint64 LinkerPtr = EventData.GetValue("Linker").As<uint64>();
+		TSharedRef<FLinkerState>* LinkerState = ActiveLinkersMap.Find(LinkerPtr);
+		if (LinkerState && (*LinkerState)->PackageInfo)
 		{
-			LinkerState->PackageInfo->Exports.Add(&Export);
+			(*LinkerState)->PackageInfo->Exports.Add(&Export);
 		}
-		ThreadState->EnterExportScope(Context.SessionContext.TimestampFromCycle(EventData.GetValue("Cycle").As<uint64>()), &Export, LoadTimeProfilerObjectEventType_Create);
 		break;
 	}
 	case RouteId_EndCreateExport:
 	{
+		Trace::FAnalysisSessionEditScope _(Session);
+
 		uint32 ThreadId = EventData.GetValue("ThreadId").As<uint32>();
 		TSharedRef<FThreadState> ThreadState = GetThreadState(ThreadId);
 		uint64 ObjectPtr = EventData.GetValue("Object").As<uint64>();
-		const Trace::FClassInfo* ObjectClass = GetClassInfo(EventData.GetValue("Class").As<uint64>());
 		Trace::FPackageExportInfo* Export = ThreadState->GetCurrentExportScope();
-		ExportsMap.Add(ObjectPtr, Export);
+		if (Export)
 		{
-			Trace::FAnalysisSessionEditScope _(Session);
+			ExportsMap.Add(ObjectPtr, Export);
+			const Trace::FClassInfo* ObjectClass = GetClassInfo(EventData.GetValue("Class").As<uint64>());
 			Export->Class = ObjectClass;
-			ThreadState->LeaveScope(Context.SessionContext.TimestampFromCycle(EventData.GetValue("Cycle").As<uint64>()));
 		}
+		ThreadState->LeaveScope(Context.SessionContext.TimestampFromCycle(EventData.GetValue("Cycle").As<uint64>()));
 		break;
 	}
 	case RouteId_BeginObjectScope:
@@ -201,7 +201,7 @@ void FAsyncLoadingTraceAnalyzer::OnEvent(uint16 RouteId, const FOnEventContext&
 	case RouteId_BeginRequest:
 	{
 		uint64 RequestId = EventData.GetValue("RequestId").As<uint64>();
-		check(!ActiveRequestsMap.Contains(RequestId));
+		//check(!ActiveRequestsMap.Contains(RequestId));
 		TSharedRef<FRequestState> RequestState = MakeShared<FRequestState>();
 		RequestState->Id = RequestId;
 		RequestState->WallTimeStartCycle = EventData.GetValue("Cycle").As<uint64>();
@@ -225,9 +225,11 @@ void FAsyncLoadingTraceAnalyzer::OnEvent(uint16 RouteId, const FOnEventContext&
 	case RouteId_EndRequest:
 	{
 		uint64 RequestId = EventData.GetValue("RequestId").As<uint64>();
-		check(ActiveRequestsMap .Contains(RequestId));
-		TSharedRef<FRequestState> RequestState = ActiveRequestsMap[RequestId];
-		RequestState->WallTimeEndCycle = EventData.GetValue("Cycle").As<uint64>();
+		TSharedRef<FRequestState>* RequestState = ActiveRequestsMap.Find(RequestId);
+		if (RequestState)
+		{
+			(*RequestState)->WallTimeEndCycle = EventData.GetValue("Cycle").As<uint64>();
+		}
 		break;
 	}
 	case RouteId_BeginLoadMap:
@@ -247,16 +249,17 @@ void FAsyncLoadingTraceAnalyzer::OnEvent(uint16 RouteId, const FOnEventContext&
 	{
 		uint32 ThreadId = EventData.GetValue("ThreadId").As<uint32>();
 		TSharedRef<FThreadState> ThreadState = GetThreadState(ThreadId);
-		TSharedRef<FLoadMapState> LoadMapState = ThreadState->ActiveLoadMap.ToSharedRef();
-		ThreadState->ActiveLoadMap = nullptr;
-		LoadMapState->WallTimeEndCycle = EventData.GetValue("Cycle").As<uint64>();
+		if (ThreadState->ActiveLoadMap)
+		{
+			ThreadState->ActiveLoadMap->WallTimeEndCycle = EventData.GetValue("Cycle").As<uint64>();
+			ThreadState->ActiveLoadMap = nullptr;
+		}
 		break;
 	}
-	
 	case RouteId_NewStreamableHandle:
 	{
 		uint64 Ptr = EventData.GetValue("StreamableHandle").As<uint64>();
-		check(!ActiveStreamableHandlesMap.Contains(Ptr));
+		//check(!ActiveStreamableHandlesMap.Contains(Ptr));
 		TSharedRef<FStreamableHandleState> StreamableHandleState = MakeShared<FStreamableHandleState>();
 		StreamableHandleState->Id = NextStreamableHandleId++;
 		StreamableHandleState->DebugName = FString((const char*)EventData.GetAttachment(), EventData.GetValue("DebugNameSize").As<uint16>());
@@ -275,36 +278,40 @@ void FAsyncLoadingTraceAnalyzer::OnEvent(uint16 RouteId, const FOnEventContext&
 	case RouteId_DestroyStreamableHandle:
 	{
 		uint64 Ptr = EventData.GetValue("StreamableHandle").As<uint64>();
-		check(ActiveStreamableHandlesMap.Contains(Ptr));
 		ActiveStreamableHandlesMap.Remove(Ptr);
 		break;
 	}
 	case RouteId_BeginLoadStreamableHandle:
 	{
 		uint64 Ptr = EventData.GetValue("StreamableHandle").As<uint64>();
-		check(ActiveStreamableHandlesMap.Contains(Ptr));
-		TSharedRef<FStreamableHandleState> StreamableHandleState = ActiveStreamableHandlesMap[Ptr];
-		StreamableHandleState->WallTimeStartCycle = EventData.GetValue("Cycle").As<uint64>();
+		TSharedRef<FStreamableHandleState>* StreamableHandleState = ActiveStreamableHandlesMap.Find(Ptr);
+		if (StreamableHandleState)
+		{
+			(*StreamableHandleState)->WallTimeStartCycle = EventData.GetValue("Cycle").As<uint64>();
+		}
 		break;
 	}
 	case RouteId_EndLoadStreamableHandle:
 	{
 		uint64 Ptr = EventData.GetValue("StreamableHandle").As<uint64>();
-		check(ActiveStreamableHandlesMap.Contains(Ptr));
-		TSharedRef<FStreamableHandleState> StreamableHandleState = ActiveStreamableHandlesMap[Ptr];
-		StreamableHandleState->WallTimeEndCycle = EventData.GetValue("Cycle").As<uint64>();
+		TSharedRef<FStreamableHandleState>* StreamableHandleState = ActiveStreamableHandlesMap.Find(Ptr);
+		if (StreamableHandleState)
+		{
+			(*StreamableHandleState)->WallTimeEndCycle = EventData.GetValue("Cycle").As<uint64>();
+		}
 		break;
 	}
 	case RouteId_BeginWaitForStreamableHandle:
 	{
 		uint64 Ptr = EventData.GetValue("StreamableHandle").As<uint64>();
-		check(ActiveStreamableHandlesMap.Contains(Ptr));
-		TSharedRef<FStreamableHandleState> StreamableHandleState = ActiveStreamableHandlesMap[Ptr];
-		uint32 ThreadId = EventData.GetValue("ThreadId").As<uint32>();
-		TSharedRef<FThreadState> ThreadState = GetThreadState(ThreadId);
-		check(!ThreadState->WaitForStreamableHandleHandle);
-		ThreadState->WaitForStreamableHandleHandle = StreamableHandleState;
-		ThreadState->WaitForStreamableHandleStartCycle = EventData.GetValue("Cycle").As<uint64>();
+		TSharedRef<FStreamableHandleState>* StreamableHandleState = ActiveStreamableHandlesMap.Find(Ptr);
+		if (StreamableHandleState)
+		{
+			uint32 ThreadId = EventData.GetValue("ThreadId").As<uint32>();
+			TSharedRef<FThreadState> ThreadState = GetThreadState(ThreadId);
+			ThreadState->WaitForStreamableHandleHandle = *StreamableHandleState;
+			ThreadState->WaitForStreamableHandleStartCycle = EventData.GetValue("Cycle").As<uint64>();
+		}
 		/*uint64 EventTypeId = StreamableHandlesTimelineEventTypeMap[ThreadState->WaitForStreamableHandleHandle->Id];
 
 		{
@@ -318,7 +325,6 @@ void FAsyncLoadingTraceAnalyzer::OnEvent(uint16 RouteId, const FOnEventContext&
 	{
 		uint32 ThreadId = EventData.GetValue("ThreadId").As<uint32>();
 		TSharedRef<FThreadState> ThreadState = GetThreadState(ThreadId);
-		check(ThreadState->WaitForStreamableHandleHandle);
 		ThreadState->WaitForStreamableHandleHandle = nullptr;
 
 		/*{
@@ -350,35 +356,40 @@ void FAsyncLoadingTraceAnalyzer::OnEvent(uint16 RouteId, const FOnEventContext&
 	case RouteId_StreamableHandleRequestAssociation:
 	{
 		uint64 StreamableHandlePtr = EventData.GetValue("StreamableHandle").As<uint64>();
-		check(ActiveStreamableHandlesMap.Contains(StreamableHandlePtr));
+		TSharedRef<FStreamableHandleState>* StreamableHandleState = ActiveStreamableHandlesMap.Find(StreamableHandlePtr);
 		uint64 RequestId = EventData.GetValue("RequestId").As<uint64>();
-		check(ActiveRequestsMap.Contains(RequestId));
-		ActiveStreamableHandlesMap[StreamableHandlePtr]->Requests.Add(ActiveRequestsMap[RequestId]);
+		TSharedRef<FRequestState>* RequestState = ActiveRequestsMap.Find(RequestId);
+		if (StreamableHandleState && RequestState)
+		{
+			(*StreamableHandleState)->Requests.Add(*RequestState);
+		}
 		break;
 	}
 	case RouteId_AsyncPackageRequestAssociation:
 	{
 		uint64 AsyncPackagePtr = EventData.GetValue("AsyncPackage").As<uint64>();
-		//check(ActiveAsyncPackagesMap.Contains(AsyncPackagePtr));
-		TSharedRef<FAsyncPackageState> AsyncPackageState = ActiveAsyncPackagesMap[AsyncPackagePtr];
+		TSharedRef<FAsyncPackageState>* AsyncPackageState = ActiveAsyncPackagesMap.Find(AsyncPackagePtr);
 		uint64 RequestId = EventData.GetValue("RequestId").As<uint64>();
-		check(ActiveRequestsMap.Contains(RequestId));
-		TSharedRef<FRequestState> RequestState = ActiveRequestsMap[RequestId];
-		RequestState->AsyncPackages.Add(AsyncPackageState);
-		AsyncPackageState->Requests.Add(RequestState);
+		TSharedRef<FRequestState>* RequestState = ActiveRequestsMap.Find(RequestId);
+		if (AsyncPackageState && RequestState)
+		{
+			(*RequestState)->AsyncPackages.Add(*AsyncPackageState);
+			(*AsyncPackageState)->Requests.Add(*RequestState);
+		}
 		break;
 	}
 	case RouteId_AsyncPackageLinkerAssociation:
 	{
 		uint64 LinkerPtr = EventData.GetValue("Linker").As<uint64>();
-		//check(ActiveLinkersMap.Contains(LinkerPtr));
-		TSharedRef<FLinkerState> LinkerState = ActiveLinkersMap[LinkerPtr];
+		TSharedRef<FLinkerState>* LinkerState = ActiveLinkersMap.Find(LinkerPtr);
 		uint64 AsyncPackagePtr = EventData.GetValue("AsyncPackage").As<uint64>();
-		//check(ActiveAsyncPackagesMap.Contains(AsyncPackagePtr));
-		TSharedRef<FAsyncPackageState> AsyncPackageState = ActiveAsyncPackagesMap[AsyncPackagePtr];
-		AsyncPackageState->LinkerState = LinkerState;
-		LinkerState->AsyncPackageState = AsyncPackageState;
-		LinkerState->PackageInfo = AsyncPackageState->PackageInfo;
+		TSharedRef<FAsyncPackageState>* AsyncPackageState = ActiveAsyncPackagesMap.Find(AsyncPackagePtr);
+		if (LinkerState && AsyncPackageState)
+		{
+			(*AsyncPackageState)->LinkerState = *LinkerState;
+			(*LinkerState)->AsyncPackageState = *AsyncPackageState;
+			(*LinkerState)->PackageInfo = (*AsyncPackageState)->PackageInfo;
+		}
 		break;
 	}
 	case RouteId_LinkerArchiveAssociation:
@@ -392,12 +403,11 @@ void FAsyncLoadingTraceAnalyzer::OnEvent(uint16 RouteId, const FOnEventContext&
 		Trace::FAnalysisSessionEditScope _(Session);
 
 		uint64 AsyncPackagePtr = EventData.GetValue("AsyncPackage").As<uint64>();
-		//check(ActiveAsyncPackagesMap.Contains(AsyncPackagePtr));
-		TSharedRef<FAsyncPackageState> AsyncPackageState = ActiveAsyncPackagesMap[AsyncPackagePtr];
+		TSharedRef<FAsyncPackageState>* AsyncPackageState = ActiveAsyncPackagesMap.Find(AsyncPackagePtr);
 		uint32 ThreadId = EventData.GetValue("ThreadId").As<uint32>();
 		TSharedRef<FThreadState> ThreadState = GetThreadState(ThreadId);
 		ELoadTimeProfilerPackageEventType EventType = static_cast<ELoadTimeProfilerPackageEventType>(EventData.GetValue("EventType").As<uint8>());
-		ThreadState->EnterPackageScope(Context.SessionContext.TimestampFromCycle(EventData.GetValue("Cycle").As<uint64>()), AsyncPackageState->PackageInfo, EventType);
+		ThreadState->EnterPackageScope(Context.SessionContext.TimestampFromCycle(EventData.GetValue("Cycle").As<uint64>()), AsyncPackageState ? (*AsyncPackageState)->PackageInfo : nullptr, EventType);
 		break;
 	}
 	case RouteId_EndAsyncPackageScope:
@@ -418,7 +428,7 @@ void FAsyncLoadingTraceAnalyzer::OnEvent(uint16 RouteId, const FOnEventContext&
 		FlushAsyncLoadingStartCycle = EventData.GetValue("Cycle").As<uint64>();
 		if (FlushAsyncLoadingRequestId != INDEX_NONE)
 		{
-			check(ActiveRequestsMap.Contains(FlushAsyncLoadingRequestId));
+			//check(ActiveRequestsMap.Contains(FlushAsyncLoadingRequestId));
 			/*uint64 EventTypeId = RequestTimelineEventTypeMap[FlushAsyncLoadingRequestId];
 			BlockingRequestsTimeline->AppendBeginEvent(Context.SessionContext.TimestampFromCycle(FlushAsyncLoadingStartCycle), EventTypeId);*/
 		}
diff --git a/Engine/Source/Developer/TraceServices/Private/Model/LoadTimeProfiler.cpp b/Engine/Source/Developer/TraceServices/Private/Model/LoadTimeProfiler.cpp
--- a/Engine/Source/Developer/TraceServices/Private/Model/LoadTimeProfiler.cpp
+++ b/Engine/Source/Developer/TraceServices/Private/Model/LoadTimeProfiler.cpp
@@ -92,6 +92,9 @@ ITable<FLoadTimeProfilerAggregatedStats>* FLoadTimeProfilerProvider::CreateEvent
 		case LoadTimeProfilerPackageEventType_Tick:
 			Row.Name = TEXT("Tick");
 			break;
+		case LoadTimeProfilerPackageEventType_DeferredPostLoad:
+			Row.Name = TEXT("DeferredPostLoad");
+			break;
 		case LoadTimeProfilerPackageEventType_Finish:
 			Row.Name = TEXT("Finish");
 			break;
diff --git a/Engine/Source/Runtime/CoreUObject/Private/Serialization/AsyncLoading.cpp b/Engine/Source/Runtime/CoreUObject/Private/Serialization/AsyncLoading.cpp
--- a/Engine/Source/Runtime/CoreUObject/Private/Serialization/AsyncLoading.cpp
+++ b/Engine/Source/Runtime/CoreUObject/Private/Serialization/AsyncLoading.cpp
@@ -4832,6 +4832,7 @@ EAsyncPackageState::Type FAsyncLoadingThread::ProcessLoadedPackages(bool bUseTim
 		FAsyncPackage* Package = LoadedPackagesToProcess[PackageIndex];
 		if (Package->GetDependencyRefCount() == 0)
 		{
+			TRACE_LOADTIME_ASYNC_PACKAGE_SCOPE(Package, LoadTimeProfilerPackageEventType_DeferredPostLoad);
 			SCOPED_LOADTIMER(ProcessLoadedPackagesTime);
 
 			Result = Package->PostLoadDeferredObjects(TickStartTime, bUseTimeLimit, TimeLimit);
@@ -5128,6 +5129,10 @@ FAsyncLoadingThread::FAsyncLoadingThread()
 			TEXT("Event driven async loader is NOT being used but it seems to be enabled in project settings."));
 	}
 
+#if LOADTIMEPROFILERTRACE_ENABLED
+	FLoadTimeProfilerTracePrivate::Init(FParse::Param(FCommandLine::Get(), TEXT("loadtimetrace")));
+#endif
+
 	QueuedRequestsEvent = FPlatformProcess::GetSynchEventFromPool();
 	CancelLoadingEvent = FPlatformProcess::GetSynchEventFromPool();
 	ThreadSuspendedEvent = FPlatformProcess::GetSynchEventFromPool();
diff --git a/Engine/Source/Runtime/CoreUObject/Private/Serialization/LoadTimeTrace.cpp b/Engine/Source/Runtime/CoreUObject/Private/Serialization/LoadTimeTrace.cpp
--- a/Engine/Source/Runtime/CoreUObject/Private/Serialization/LoadTimeTrace.cpp
+++ b/Engine/Source/Runtime/CoreUObject/Private/Serialization/LoadTimeTrace.cpp
@@ -12,7 +12,7 @@
 #include "UObject/Object.h"
 #include "UObject/UObjectGlobals.h"
 
-UE_TRACE_EVENT_BEGIN(LoadTime, StartAsyncLoading)
+UE_TRACE_EVENT_BEGIN(LoadTime, StartAsyncLoading, Always)
 	UE_TRACE_EVENT_FIELD(uint64, Cycle)
 	UE_TRACE_EVENT_FIELD(uint32, ThreadId)
 UE_TRACE_EVENT_END()
@@ -201,10 +201,54 @@ UE_TRACE_EVENT_BEGIN(LoadTime, QueueEvent)
 	UE_TRACE_EVENT_FIELD(uint8, EventType)
 UE_TRACE_EVENT_END();
 
-UE_TRACE_EVENT_BEGIN(LoadTime, ClassInfo)
+UE_TRACE_EVENT_BEGIN(LoadTime, ClassInfo, Always)
 	UE_TRACE_EVENT_FIELD(const UClass*, Class)
 UE_TRACE_EVENT_END()
 
+void FLoadTimeProfilerTracePrivate::Init(bool bStartEnabled)
+{
+	UE_TRACE_EVENT_IS_ENABLED(LoadTime, StartAsyncLoading);
+	UE_TRACE_EVENT_IS_ENABLED(LoadTime, SuspendAsyncLoading);
+	UE_TRACE_EVENT_IS_ENABLED(LoadTime, ResumeAsyncLoading);
+	UE_TRACE_EVENT_IS_ENABLED(LoadTime, NewPackage);
+	UE_TRACE_EVENT_IS_ENABLED(LoadTime, NewLinker);
+	UE_TRACE_EVENT_IS_ENABLED(LoadTime, DestroyLinker);
+	UE_TRACE_EVENT_IS_ENABLED(LoadTime, PackageSummary);
+	UE_TRACE_EVENT_IS_ENABLED(LoadTime, BeginLinkerScope);
+	UE_TRACE_EVENT_IS_ENABLED(LoadTime, EndLinkerScope);
+	UE_TRACE_EVENT_IS_ENABLED(LoadTime, BeginCreateExport);
+	UE_TRACE_EVENT_IS_ENABLED(LoadTime, EndCreateExport);
+	UE_TRACE_EVENT_IS_ENABLED(LoadTime, BeginObjectScope);
+	UE_TRACE_EVENT_IS_ENABLED(LoadTime, EndObjectScope);
+	UE_TRACE_EVENT_IS_ENABLED(LoadTime, LinkerArchiveAssociation);
+	UE_TRACE_EVENT_IS_ENABLED(LoadTime, BeginRequest);
+	UE_TRACE_EVENT_IS_ENABLED(LoadTime, EndRequest);
+	UE_TRACE_EVENT_IS_ENABLED(LoadTime, BeginLoadMap);
+	UE_TRACE_EVENT_IS_ENABLED(LoadTime, EndLoadMap);
+	UE_TRACE_EVENT_IS_ENABLED(LoadTime, NewStreamableHandle);
+	UE_TRACE_EVENT_IS_ENABLED(LoadTime, DestroyStreamableHandle);
+	UE_TRACE_EVENT_IS_ENABLED(LoadTime, BeginLoadStreamableHandle);
+	UE_TRACE_EVENT_IS_ENABLED(LoadTime, EndLoadStreamableHandle);
+	UE_TRACE_EVENT_IS_ENABLED(LoadTime, BeginWaitForStreamableHandle);
+	UE_TRACE_EVENT_IS_ENABLED(LoadTime, EndWaitForStreamableHandle);
+	UE_TRACE_EVENT_IS_ENABLED(LoadTime, NewAsyncPackage);
+	UE_TRACE_EVENT_IS_ENABLED(LoadTime, DestroyAsyncPackage);
+	UE_TRACE_EVENT_IS_ENABLED(LoadTime, StreamableHandleRequestAssociation);
+	UE_TRACE_EVENT_IS_ENABLED(LoadTime, AsyncPackageRequestAssociation);
+	UE_TRACE_EVENT_IS_ENABLED(LoadTime, LoadImportDependency);
+	UE_TRACE_EVENT_IS_ENABLED(LoadTime, AsyncPackageLinkerAssociation);
+	UE_TRACE_EVENT_IS_ENABLED(LoadTime, BeginAsyncPackageScope);
+	UE_TRACE_EVENT_IS_ENABLED(LoadTime, EndAsyncPackageScope);
+	UE_TRACE_EVENT_IS_ENABLED(LoadTime, BeginFlushAsyncLoading);
+	UE_TRACE_EVENT_IS_ENABLED(LoadTime, EndFlushAsyncLoading);
+	UE_TRACE_EVENT_IS_ENABLED(LoadTime, QueueEvent);
+	UE_TRACE_EVENT_IS_ENABLED(LoadTime, ClassInfo);
+	if (bStartEnabled)
+	{
+		Trace::ToggleEvent(TEXT("LoadTime"), true);
+	}
+}
+
 void FLoadTimeProfilerTracePrivate::OutputStartAsyncLoading()
 {
 	UE_TRACE_LOG(LoadTime, StartAsyncLoading)
diff --git a/Engine/Source/Runtime/CoreUObject/Private/Serialization/LoadTimeTracePrivate.h b/Engine/Source/Runtime/CoreUObject/Private/Serialization/LoadTimeTracePrivate.h
--- a/Engine/Source/Runtime/CoreUObject/Private/Serialization/LoadTimeTracePrivate.h
+++ b/Engine/Source/Runtime/CoreUObject/Private/Serialization/LoadTimeTracePrivate.h
@@ -16,6 +16,7 @@ class UPackage;
 
 struct FLoadTimeProfilerTracePrivate
 {
+	static void Init(bool bStartEnabled);
 	static void OutputStartAsyncLoading();
 	static void OutputSuspendAsyncLoading();
 	static void OutputResumeAsyncLoading();
@@ -106,8 +107,8 @@ struct FLoadTimeProfilerTracePrivate
 #define TRACE_LOADTIME_ASYNC_PACKAGE_LINKER_ASSOCIATION(AsyncPackage, Linker) \
 	FLoadTimeProfilerTracePrivate::OutputAsyncPackageLinkerAssociation(AsyncPackage, Linker);
 
-#define TRACE_LOADTIME_QUEUE_EVENT(AsyncPackage, EventType, UserPriority, PackageSerialNumber, EventSystemPriority) \
-	FLoadTimeProfilerTracePrivate::OutputQueueEvent(AsyncPackage, EventType, UserPriority, PackageSerialNumber, EventSystemPriority);
+#define TRACE_LOADTIME_QUEUE_EVENT(AsyncPackage, EventType, UserPriority, PackageSerialNumber, EventSystemPriority)
+//	FLoadTimeProfilerTracePrivate::OutputQueueEvent(AsyncPackage, EventType, UserPriority, PackageSerialNumber, EventSystemPriority);
 
 #define TRACE_LOADTIME_ASYNC_PACKAGE_SCOPE(AsyncPackage, EventType) \
 	FLoadTimeProfilerTracePrivate::FAsyncPackageScope __LoadTimeTraceAsyncPackageScope(AsyncPackage, EventType);
diff --git a/Engine/Source/Runtime/CoreUObject/Public/Serialization/LoadTimeTrace.h b/Engine/Source/Runtime/CoreUObject/Public/Serialization/LoadTimeTrace.h
--- a/Engine/Source/Runtime/CoreUObject/Public/Serialization/LoadTimeTrace.h
+++ b/Engine/Source/Runtime/CoreUObject/Public/Serialization/LoadTimeTrace.h
@@ -3,9 +3,10 @@
 #pragma once
 
 #include "CoreTypes.h"
+#include "Trace/Trace.h"
 
-#if !IS_PROGRAM && !UE_BUILD_SHIPPING
-#define LOADTIMEPROFILERTRACE_ENABLED 0
+#if UE_TRACE_ENABLED && !UE_BUILD_SHIPPING
+#define LOADTIMEPROFILERTRACE_ENABLED 1
 #else
 #define LOADTIMEPROFILERTRACE_ENABLED 0
 #endif
@@ -23,6 +24,7 @@ enum ELoadTimeProfilerPackageEventType
 	LoadTimeProfilerPackageEventType_StartPostLoad,
 	LoadTimeProfilerPackageEventType_Tick,
 	LoadTimeProfilerPackageEventType_Finish,
+	LoadTimeProfilerPackageEventType_DeferredPostLoad,
 	LoadTimeProfilerPackageEventType_None,
 };
 
