Use named events instead of semaphores as they're lighter weight. As failing to open a named event is an expected possibility, we'll directly use the Windows API so as to not fill the TTY log with warnings.

diff --git a/Engine/Source/Developer/TraceServices/Private/SessionService.cpp b/Engine/Source/Developer/TraceServices/Private/SessionService.cpp
--- a/Engine/Source/Developer/TraceServices/Private/SessionService.cpp
+++ b/Engine/Source/Developer/TraceServices/Private/SessionService.cpp
@@ -9,6 +9,12 @@
 #include "Trace/ControlClient.h"
 #include "IPAddress.h"
 
+#if PLATFORM_WINDOWS
+	#include "Windows/AllowWindowsPlatformTypes.h"
+	#include "Windows/MinWindows.h"
+	#include "Windows/HideWindowsPlatformTypes.h"
+#endif
+
 namespace Trace
 {
 
@@ -35,10 +41,14 @@ FSessionService::~FSessionService()
 bool FSessionService::StartRecorderServer()
 {
 	bool bOk = TraceRecorder->StartRecording();
-	if (bOk && RecorderMutex == nullptr)
+#if PLATFORM_WINDOWS
+	// Create a named event that other processes can use detect a running
+	// recorder and connect to it automatically
+	if (bOk && RecorderEvent == nullptr)
 	{
-		RecorderMutex = FPlatformProcess::NewInterprocessSynchObject("UnrealInsightsRecorder", true);
+		RecorderEvent = ::CreateEvent(nullptr, true, false, TEXT("Local\\UnrealInsightsRecorder"));
 	}
+#endif // PLATFORM_WINDOWS
 	return bOk;
 }
 
@@ -49,11 +59,13 @@ bool FSessionService::IsRecorderServerRunning() const
 
 void FSessionService::StopRecorderServer()
 {
-	if (RecorderMutex != nullptr)
+#if PLATFORM_WINDOWS
+	if (RecorderEvent != nullptr)
 	{
-		FPlatformProcess::DeleteInterprocessSynchObject(RecorderMutex);
-		RecorderMutex = nullptr;
+		::CloseHandle(RecorderEvent);
+		RecorderEvent = nullptr;
 	}
+#endif // PLATFORM_WINDOWS
 	TraceRecorder->StopRecording();
 }
 
diff --git a/Engine/Source/Developer/TraceServices/Private/SessionServicePrivate.h b/Engine/Source/Developer/TraceServices/Private/SessionServicePrivate.h
--- a/Engine/Source/Developer/TraceServices/Private/SessionServicePrivate.h
+++ b/Engine/Source/Developer/TraceServices/Private/SessionServicePrivate.h
@@ -3,7 +3,6 @@
 #pragma once
 
 #include "TraceServices/SessionService.h"
-#include "HAL/PlatformProcess.h"
 #include "Misc/ScopeLock.h"
 #include "Containers/Ticker.h"
 #include "Trace/Recorder.h"
@@ -46,7 +45,7 @@ private:
 	};
 
 	FModuleService& ModuleService;
-	FPlatformProcess::FSemaphore* RecorderMutex = nullptr;
+	void* RecorderEvent = nullptr;
 	FString LocalSessionDirectory;
 	TSharedPtr<IStore> TraceStore;
 	TSharedPtr<IRecorder> TraceRecorder;
diff --git a/Engine/Source/Runtime/Launch/Private/LaunchEngineLoop.cpp b/Engine/Source/Runtime/Launch/Private/LaunchEngineLoop.cpp
--- a/Engine/Source/Runtime/Launch/Private/LaunchEngineLoop.cpp
+++ b/Engine/Source/Runtime/Launch/Private/LaunchEngineLoop.cpp
@@ -1223,10 +1223,12 @@ int32 FEngineLoop::PreInit(const TCHAR* CmdLine)
 #if PLATFORM_WINDOWS && !UE_BUILD_SHIPPING
 		else
 		{
-			if (FPlatformProcess::FSemaphore* Mutex = FPlatformProcess::NewInterprocessSynchObject("UnrealInsightsRecorder", false))
+			// If we can detect a named event then we can try and auto-connect to UnrealInsights.
+			HANDLE KnownEvent = ::OpenEvent(EVENT_ALL_ACCESS, false, TEXT("Local\\UnrealInsightsRecorder"));
+			if (KnownEvent != nullptr)
 			{
 				Trace::Connect(TEXT("127.0.0.1"));
-				FPlatformProcess::DeleteInterprocessSynchObject(Mutex);
+				::CloseHandle(KnownEvent);
 			}
 		}
 #endif // PLATFORM_WINDOWS
