Less verbose IAnalyzer.Event.GetValue() API for fetching a field's value, along with support to coerce between all types.

diff --git a/Engine/Source/Developer/TraceAnalysis/Private/Analysis/Engine.cpp b/Engine/Source/Developer/TraceAnalysis/Private/Analysis/Engine.cpp
--- a/Engine/Source/Developer/TraceAnalysis/Private/Analysis/Engine.cpp
+++ b/Engine/Source/Developer/TraceAnalysis/Private/Analysis/Engine.cpp
@@ -217,7 +217,7 @@ struct FAnalysisEngine::FEventDataImpl
 	: public FEventData
 {
 	virtual					~FEventDataImpl() = default;
-	virtual const FValue&	GetValue(const ANSICHAR* FieldName) const override;
+	virtual const void*		GetValueImpl(const ANSICHAR* FieldName, uint16& Type) const override;
 	virtual const uint8*	GetAttachment() const override;
 	virtual uint16			GetAttachmentSize() const override;
 	const FDispatch*		Dispatch;
@@ -226,10 +226,8 @@ struct FAnalysisEngine::FEventDataImpl
 };
 
 ////////////////////////////////////////////////////////////////////////////////
-const IAnalyzer::FValue& FAnalysisEngine::FEventDataImpl::GetValue(const ANSICHAR* FieldName) const
+const void* FAnalysisEngine::FEventDataImpl::GetValueImpl(const ANSICHAR* FieldName, uint16& Type) const
 {
-	UPTRINT Ret = ~0ull;
-
 	FFnv1aHash Hash;
 	Hash.Add(FieldName);
 	uint32 NameHash = Hash.Get();
@@ -239,13 +237,12 @@ const IAnalyzer::FValue& FAnalysisEngine::FEventDataImpl::GetValue(const ANSICHA
 		const auto& Field = Dispatch->Fields[i];
 		if (Field.Hash == NameHash)
 		{
-			Ret = UPTRINT(Ptr + Field.Offset);
-			Ret |= UPTRINT(Field.TypeInfo) << 48;
-			break;
+			Type = Field.TypeInfo;
+			return (Ptr + Field.Offset);
 		}
 	}
 
-	return *(FValue*)Ret;
+	return nullptr;
 }
 
 ////////////////////////////////////////////////////////////////////////////////
@@ -416,8 +413,8 @@ void FAnalysisEngine::OnNewTrace(const FOnEventContext& Context)
 ////////////////////////////////////////////////////////////////////////////////
 void FAnalysisEngine::OnTiming(const FOnEventContext& Context)
 {
-	SessionContext.StartCycle = Context.EventData.GetValue("StartCycle").As<uint64>();
-	SessionContext.CycleFrequency = Context.EventData.GetValue("CycleFrequency").As<uint64>();
+	SessionContext.StartCycle = Context.EventData.GetValue<uint64>("StartCycle");
+	SessionContext.CycleFrequency = Context.EventData.GetValue<uint64>("CycleFrequency");
 }
 
 ////////////////////////////////////////////////////////////////////////////////
diff --git a/Engine/Source/Developer/TraceAnalysis/Public/Trace/Analyzer.h b/Engine/Source/Developer/TraceAnalysis/Public/Trace/Analyzer.h
--- a/Engine/Source/Developer/TraceAnalysis/Public/Trace/Analyzer.h
+++ b/Engine/Source/Developer/TraceAnalysis/Public/Trace/Analyzer.h
@@ -42,17 +42,14 @@ public:
 		FInterfaceBuilder&		InterfaceBuilder;
 	};
 
-	struct FValue
-	{
-		template <typename AsType>
-		AsType As() const;
-	};
-
 	struct FEventData
 	{
-		virtual const FValue&	GetValue(const ANSICHAR* FieldName) const = 0;
+		template <typename ValueType> ValueType GetValue(const ANSICHAR* FieldName) const;
 		virtual const uint8*	GetAttachment() const = 0;
 		virtual uint16			GetAttachmentSize() const = 0;
+
+	private:
+		virtual const void*		GetValueImpl(const ANSICHAR* FieldName, uint16& Type) const = 0;
 	};
 
 	struct FOnEventContext
@@ -68,23 +65,39 @@ public:
 };
 
 ////////////////////////////////////////////////////////////////////////////////
-template <typename AsType>
-AsType IAnalyzer::FValue::As() const
+template <typename ValueType>
+ValueType IAnalyzer::FEventData::GetValue(const ANSICHAR* FieldName) const
 {
-	uint16 FieldType = uint16(UPTRINT(this) >> 48);
-	if (FieldType == 0xffff)
+	uint16 FieldTypeId;
+	const void* Addr = GetValueImpl(FieldName, FieldTypeId);
+	if (Addr == nullptr)
 	{
-		return AsType(0);
+		return ValueType(0);
 	}
 
-	check((FieldType & _Field_Float) == 0);
+	uint32 Pow2Size = (FieldTypeId & _Field_Pow2SizeMask);
+	uint32 Category = (FieldTypeId & _Field_CategoryMask);
 
-	const void* Addr = (void*)(UPTRINT(this) & ((1ull << 48) - 1));
-	UPTRINT Value = 0;
-	uint32 Size = 1 << (FieldType & _Field_Pow2SizeMask);
-	memcpy(&Value, Addr, Size);
+	if (Category == _Field_Float)
+	{
+		switch (Pow2Size)
+		{
+		case _Field_32: return ValueType(*(const float*)(Addr));
+		case _Field_64: return ValueType(*(const double*)(Addr));
+		}
+	}
+	else if (Category == _Field_Integer)
+	{
+		switch (Pow2Size)
+		{
+		case _Field_8:  return ValueType(*(const uint8*)(Addr));
+		case _Field_16: return ValueType(*(const uint16*)(Addr));
+		case _Field_32: return ValueType(*(const uint32*)(Addr));
+		case _Field_64: return ValueType(*(const uint64*)(Addr));
+		}
+	}
 
-	return AsType(Value);
+	return ValueType(0);
 }
 
 } // namespace Trace
diff --git a/Engine/Source/Developer/TraceServices/Private/Analyzers/CpuProfilerTraceAnalysis.cpp b/Engine/Source/Developer/TraceServices/Private/Analyzers/CpuProfilerTraceAnalysis.cpp
--- a/Engine/Source/Developer/TraceServices/Private/Analyzers/CpuProfilerTraceAnalysis.cpp
+++ b/Engine/Source/Developer/TraceServices/Private/Analyzers/CpuProfilerTraceAnalysis.cpp
@@ -33,7 +33,7 @@ void FCpuProfilerAnalyzer::OnEvent(uint16 RouteId, const FOnEventContext& Contex
 	{
 	case RouteId_EventSpec:
 	{
-		uint16 Id = EventData.GetValue("Id").As<uint16>();
+		uint16 Id = EventData.GetValue<uint16>("Id");
 		if (ScopeIdToEventIdMap.Contains(Id))
 		{
 			TimingProfilerProvider.SetTimerName(ScopeIdToEventIdMap[Id], reinterpret_cast<const TCHAR*>(EventData.GetAttachment()));
@@ -48,7 +48,7 @@ void FCpuProfilerAnalyzer::OnEvent(uint16 RouteId, const FOnEventContext& Contex
 	case RouteId_EndCapture:
 	{
 		TotalEventSize += EventData.GetAttachmentSize();
-		uint32 ThreadId = EventData.GetValue("ThreadId").As<uint32>();
+		uint32 ThreadId = EventData.GetValue<uint32>("ThreadId");
 		TSharedRef<FThreadState> ThreadState = GetThreadState(ThreadId);
 		uint64 LastCycle = ThreadState->LastCycle;
 		uint64 BufferSize = EventData.GetAttachmentSize();
diff --git a/Engine/Source/Developer/TraceServices/Private/Analyzers/GpuProfilerTraceAnalysis.cpp b/Engine/Source/Developer/TraceServices/Private/Analyzers/GpuProfilerTraceAnalysis.cpp
--- a/Engine/Source/Developer/TraceServices/Private/Analyzers/GpuProfilerTraceAnalysis.cpp
+++ b/Engine/Source/Developer/TraceServices/Private/Analyzers/GpuProfilerTraceAnalysis.cpp
@@ -30,8 +30,8 @@ void FGpuProfilerAnalyzer::OnEvent(uint16 RouteId, const FOnEventContext& Contex
 	{
 	case RouteId_EventSpec:
 	{
-		uint64 EventType = EventData.GetValue("EventType").As<uint64>();
-		FString EventName(reinterpret_cast<const TCHAR*>(EventData.GetAttachment()), EventData.GetValue("NameLength").As<uint16>());
+		uint64 EventType = EventData.GetValue<uint64>("EventType");
+		FString EventName(reinterpret_cast<const TCHAR*>(EventData.GetAttachment()), EventData.GetValue<uint16>("NameLength"));
 		EventTypeMap.Add(EventType, TimingProfilerProvider.AddGpuTimer(*EventName));
 		break;
 	}
@@ -43,7 +43,7 @@ void FGpuProfilerAnalyzer::OnEvent(uint16 RouteId, const FOnEventContext& Contex
 
 		uint32 CurrentDepth = 0;
 
-		uint64 LastTimestamp = EventData.GetValue("TimestampBase").As<uint64>();
+		uint64 LastTimestamp = EventData.GetValue<uint64>("TimestampBase");
 		double LastTime = 0.0;
 		while (BufferPtr < BufferEnd)
 		{
diff --git a/Engine/Source/Developer/TraceServices/Private/Analyzers/LoadTimeTraceAnalysis.cpp b/Engine/Source/Developer/TraceServices/Private/Analyzers/LoadTimeTraceAnalysis.cpp
--- a/Engine/Source/Developer/TraceServices/Private/Analyzers/LoadTimeTraceAnalysis.cpp
+++ b/Engine/Source/Developer/TraceServices/Private/Analyzers/LoadTimeTraceAnalysis.cpp
@@ -105,13 +105,13 @@ void FAsyncLoadingTraceAnalyzer::OnEvent(uint16 RouteId, const FOnEventContext&
 	case RouteId_StartAsyncLoading:
 	{
 		Trace::FAnalysisSessionEditScope _(Session);
-		AsyncLoadingThreadId = EventData.GetValue("ThreadId").As<uint32>();
+		AsyncLoadingThreadId = EventData.GetValue<uint32>("ThreadId");
 		LoadTimeProfilerProvider.SetAsyncLoadingThreadId(AsyncLoadingThreadId);
 		break;
 	}
 	case RouteId_NewLinker:
 	{
-		uint64 LinkerPtr = EventData.GetValue("Linker").As<uint64>();
+		uint64 LinkerPtr = EventData.GetValue<uint64>("Linker");
 		//check(!ActiveLinkersMap.Contains(LinkerPtr));
 		TSharedRef<FLinkerState> LinkerState = MakeShared<FLinkerState>();
 		ActiveLinkersMap.Add(LinkerPtr, LinkerState);
@@ -119,23 +119,23 @@ void FAsyncLoadingTraceAnalyzer::OnEvent(uint16 RouteId, const FOnEventContext&
 	}
 	case RouteId_DestroyLinker:
 	{
-		uint64 Ptr = EventData.GetValue("Linker").As<uint64>();
+		uint64 Ptr = EventData.GetValue<uint64>("Linker");
 		//check(ActiveLinkersMap.Contains(Ptr));
 		ActiveLinkersMap.Remove(Ptr);
 		break;
 	}
 	case RouteId_PackageSummary:
 	{
-		uint64 LinkerPtr = EventData.GetValue("Linker").As<uint64>();
+		uint64 LinkerPtr = EventData.GetValue<uint64>("Linker");
 		TSharedRef<FLinkerState>* LinkerState = ActiveLinkersMap.Find(LinkerPtr);
 		if (LinkerState && (*LinkerState)->PackageInfo)
 		{
 			Trace::FAnalysisSessionEditScope _(Session);
 			Trace::FPackageSummaryInfo& Summary = (*LinkerState)->PackageInfo->Summary;
-			Summary.TotalHeaderSize = EventData.GetValue("TotalHeaderSize").As<uint32>();
-			Summary.NameCount = EventData.GetValue("NameCount").As<uint32>();
-			Summary.ImportCount = EventData.GetValue("ImportCount").As<uint32>();
-			Summary.ExportCount = EventData.GetValue("ExportCount").As<uint32>();
+			Summary.TotalHeaderSize = EventData.GetValue<uint32>("TotalHeaderSize");
+			Summary.NameCount = EventData.GetValue<uint32>("NameCount");
+			Summary.ImportCount = EventData.GetValue<uint32>("ImportCount");
+			Summary.ExportCount = EventData.GetValue<uint32>("ExportCount");
 		}
 		break;
 	}
@@ -144,13 +144,13 @@ void FAsyncLoadingTraceAnalyzer::OnEvent(uint16 RouteId, const FOnEventContext&
 		Trace::FAnalysisSessionEditScope _(Session);
 
 		Trace::FPackageExportInfo& Export = LoadTimeProfilerProvider.CreateExport();
-		Export.SerialOffset = EventData.GetValue("SerialOffset").As<uint64>();
-		Export.SerialSize = EventData.GetValue("SerialSize").As<uint64>();
-		Export.IsAsset = EventData.GetValue("IsAsset").As<bool>();
-		uint32 ThreadId = EventData.GetValue("ThreadId").As<uint32>();
+		Export.SerialOffset = EventData.GetValue<uint64>("SerialOffset");
+		Export.SerialSize = EventData.GetValue<uint64>("SerialSize");
+		Export.IsAsset = EventData.GetValue<bool>("IsAsset");
+		uint32 ThreadId = EventData.GetValue<uint32>("ThreadId");
 		TSharedRef<FThreadState> ThreadState = GetThreadState(ThreadId);
-		ThreadState->EnterExportScope(Context.SessionContext.TimestampFromCycle(EventData.GetValue("Cycle").As<uint64>()), &Export, LoadTimeProfilerObjectEventType_Create);
-		uint64 LinkerPtr = EventData.GetValue("Linker").As<uint64>();
+		ThreadState->EnterExportScope(Context.SessionContext.TimestampFromCycle(EventData.GetValue<uint64>("Cycle")), &Export, LoadTimeProfilerObjectEventType_Create);
+		uint64 LinkerPtr = EventData.GetValue<uint64>("Linker");
 		TSharedRef<FLinkerState>* LinkerState = ActiveLinkersMap.Find(LinkerPtr);
 		if (LinkerState && (*LinkerState)->PackageInfo)
 		{
@@ -162,53 +162,53 @@ void FAsyncLoadingTraceAnalyzer::OnEvent(uint16 RouteId, const FOnEventContext&
 	{
 		Trace::FAnalysisSessionEditScope _(Session);
 
-		uint32 ThreadId = EventData.GetValue("ThreadId").As<uint32>();
+		uint32 ThreadId = EventData.GetValue<uint32>("ThreadId");
 		TSharedRef<FThreadState> ThreadState = GetThreadState(ThreadId);
-		uint64 ObjectPtr = EventData.GetValue("Object").As<uint64>();
+		uint64 ObjectPtr = EventData.GetValue<uint64>("Object");
 		Trace::FPackageExportInfo* Export = ThreadState->GetCurrentExportScope();
 		if (Export)
 		{
 			ExportsMap.Add(ObjectPtr, Export);
-			const Trace::FClassInfo* ObjectClass = GetClassInfo(EventData.GetValue("Class").As<uint64>());
+			const Trace::FClassInfo* ObjectClass = GetClassInfo(EventData.GetValue<uint64>("Class"));
 			Export->Class = ObjectClass;
 		}
-		ThreadState->LeaveScope(Context.SessionContext.TimestampFromCycle(EventData.GetValue("Cycle").As<uint64>()));
+		ThreadState->LeaveScope(Context.SessionContext.TimestampFromCycle(EventData.GetValue<uint64>("Cycle")));
 		break;
 	}
 	case RouteId_BeginObjectScope:
 	{
-		uint64 ObjectPtr = EventData.GetValue("Object").As<uint64>();
+		uint64 ObjectPtr = EventData.GetValue<uint64>("Object");
 		Trace::FPackageExportInfo** Export = ExportsMap.Find(ObjectPtr);
-		uint32 ThreadId = EventData.GetValue("ThreadId").As<uint32>();
+		uint32 ThreadId = EventData.GetValue<uint32>("ThreadId");
 		TSharedRef<FThreadState> ThreadState = GetThreadState(ThreadId);
-		ELoadTimeProfilerObjectEventType EventType = static_cast<ELoadTimeProfilerObjectEventType>(EventData.GetValue("EventType").As<uint8>());
+		ELoadTimeProfilerObjectEventType EventType = static_cast<ELoadTimeProfilerObjectEventType>(EventData.GetValue<uint8>("EventType"));
 		{
 			Trace::FAnalysisSessionEditScope _(Session);
-			ThreadState->EnterExportScope(Context.SessionContext.TimestampFromCycle(EventData.GetValue("Cycle").As<uint64>()), Export ? *Export : nullptr, EventType);
+			ThreadState->EnterExportScope(Context.SessionContext.TimestampFromCycle(EventData.GetValue<uint64>("Cycle")), Export ? *Export : nullptr, EventType);
 		}
 		break;
 	}
 	case RouteId_EndObjectScope:
 	{
-		uint32 ThreadId = EventData.GetValue("ThreadId").As<uint32>();
+		uint32 ThreadId = EventData.GetValue<uint32>("ThreadId");
 		TSharedRef<FThreadState> ThreadState = GetThreadState(ThreadId);
 		{
 			Trace::FAnalysisSessionEditScope _(Session);
-			ThreadState->LeaveScope(Context.SessionContext.TimestampFromCycle(EventData.GetValue("Cycle").As<uint64>()));
+			ThreadState->LeaveScope(Context.SessionContext.TimestampFromCycle(EventData.GetValue<uint64>("Cycle")));
 		}
 		break;
 	}
 	case RouteId_BeginRequest:
 	{
-		uint64 RequestId = EventData.GetValue("RequestId").As<uint64>();
+		uint64 RequestId = EventData.GetValue<uint64>("RequestId");
 		//check(!ActiveRequestsMap.Contains(RequestId));
 		TSharedRef<FRequestState> RequestState = MakeShared<FRequestState>();
 		RequestState->Id = RequestId;
-		RequestState->WallTimeStartCycle = EventData.GetValue("Cycle").As<uint64>();
+		RequestState->WallTimeStartCycle = EventData.GetValue<uint64>("Cycle");
 		RequestState->WallTimeEndCycle = 0;
 		Requests.Add(RequestState);
 		ActiveRequestsMap.Add(RequestId, RequestState);
-		uint32 ThreadId = EventData.GetValue("ThreadId").As<uint32>();
+		uint32 ThreadId = EventData.GetValue<uint32>("ThreadId");
 		TSharedRef<FThreadState> ThreadState = GetThreadState(ThreadId);
 		if (ThreadState->ActiveLoadMap)
 		{
@@ -224,11 +224,11 @@ void FAsyncLoadingTraceAnalyzer::OnEvent(uint16 RouteId, const FOnEventContext&
 	}
 	case RouteId_EndRequest:
 	{
-		uint64 RequestId = EventData.GetValue("RequestId").As<uint64>();
+		uint64 RequestId = EventData.GetValue<uint64>("RequestId");
 		TSharedRef<FRequestState>* RequestState = ActiveRequestsMap.Find(RequestId);
 		if (RequestState)
 		{
-			(*RequestState)->WallTimeEndCycle = EventData.GetValue("Cycle").As<uint64>();
+			(*RequestState)->WallTimeEndCycle = EventData.GetValue<uint64>("Cycle");
 		}
 		break;
 	}
@@ -236,33 +236,33 @@ void FAsyncLoadingTraceAnalyzer::OnEvent(uint16 RouteId, const FOnEventContext&
 	{
 		TSharedRef<FLoadMapState> LoadMapState = MakeShared<FLoadMapState>();
 		LoadMapState->Id = NextMapId++;
-		LoadMapState->WallTimeStartCycle = EventData.GetValue("Cycle").As<uint64>();
+		LoadMapState->WallTimeStartCycle = EventData.GetValue<uint64>("Cycle");
 		LoadMapState->WallTimeEndCycle = 0;
 		LoadMapState->Name = reinterpret_cast<const TCHAR*>(EventData.GetAttachment());
 		Maps.Add(LoadMapState);
-		uint32 ThreadId = EventData.GetValue("ThreadId").As<uint32>();
+		uint32 ThreadId = EventData.GetValue<uint32>("ThreadId");
 		TSharedRef<FThreadState> ThreadState = GetThreadState(ThreadId);
 		ThreadState->ActiveLoadMap = LoadMapState;
 		break;
 	}
 	case RouteId_EndLoadMap:
 	{
-		uint32 ThreadId = EventData.GetValue("ThreadId").As<uint32>();
+		uint32 ThreadId = EventData.GetValue<uint32>("ThreadId");
 		TSharedRef<FThreadState> ThreadState = GetThreadState(ThreadId);
 		if (ThreadState->ActiveLoadMap)
 		{
-			ThreadState->ActiveLoadMap->WallTimeEndCycle = EventData.GetValue("Cycle").As<uint64>();
+			ThreadState->ActiveLoadMap->WallTimeEndCycle = EventData.GetValue<uint64>("Cycle");
 			ThreadState->ActiveLoadMap = nullptr;
 		}
 		break;
 	}
 	case RouteId_NewStreamableHandle:
 	{
-		uint64 Ptr = EventData.GetValue("StreamableHandle").As<uint64>();
+		uint64 Ptr = EventData.GetValue<uint64>("StreamableHandle");
 		//check(!ActiveStreamableHandlesMap.Contains(Ptr));
 		TSharedRef<FStreamableHandleState> StreamableHandleState = MakeShared<FStreamableHandleState>();
 		StreamableHandleState->Id = NextStreamableHandleId++;
-		StreamableHandleState->DebugName = FString((const char*)EventData.GetAttachment(), EventData.GetValue("DebugNameSize").As<uint16>());
+		StreamableHandleState->DebugName = FString((const char*)EventData.GetAttachment(), EventData.GetValue<uint16>("DebugNameSize"));
 		StreamableHandleState->WallTimeStartCycle = 0;
 		StreamableHandleState->WallTimeEndCycle = 0;
 		StreamableHandles.Add(StreamableHandleState);
@@ -277,40 +277,40 @@ void FAsyncLoadingTraceAnalyzer::OnEvent(uint16 RouteId, const FOnEventContext&
 	}
 	case RouteId_DestroyStreamableHandle:
 	{
-		uint64 Ptr = EventData.GetValue("StreamableHandle").As<uint64>();
+		uint64 Ptr = EventData.GetValue<uint64>("StreamableHandle");
 		ActiveStreamableHandlesMap.Remove(Ptr);
 		break;
 	}
 	case RouteId_BeginLoadStreamableHandle:
 	{
-		uint64 Ptr = EventData.GetValue("StreamableHandle").As<uint64>();
+		uint64 Ptr = EventData.GetValue<uint64>("StreamableHandle");
 		TSharedRef<FStreamableHandleState>* StreamableHandleState = ActiveStreamableHandlesMap.Find(Ptr);
 		if (StreamableHandleState)
 		{
-			(*StreamableHandleState)->WallTimeStartCycle = EventData.GetValue("Cycle").As<uint64>();
+			(*StreamableHandleState)->WallTimeStartCycle = EventData.GetValue<uint64>("Cycle");
 		}
 		break;
 	}
 	case RouteId_EndLoadStreamableHandle:
 	{
-		uint64 Ptr = EventData.GetValue("StreamableHandle").As<uint64>();
+		uint64 Ptr = EventData.GetValue<uint64>("StreamableHandle");
 		TSharedRef<FStreamableHandleState>* StreamableHandleState = ActiveStreamableHandlesMap.Find(Ptr);
 		if (StreamableHandleState)
 		{
-			(*StreamableHandleState)->WallTimeEndCycle = EventData.GetValue("Cycle").As<uint64>();
+			(*StreamableHandleState)->WallTimeEndCycle = EventData.GetValue<uint64>("Cycle");
 		}
 		break;
 	}
 	case RouteId_BeginWaitForStreamableHandle:
 	{
-		uint64 Ptr = EventData.GetValue("StreamableHandle").As<uint64>();
+		uint64 Ptr = EventData.GetValue<uint64>("StreamableHandle");
 		TSharedRef<FStreamableHandleState>* StreamableHandleState = ActiveStreamableHandlesMap.Find(Ptr);
 		if (StreamableHandleState)
 		{
-			uint32 ThreadId = EventData.GetValue("ThreadId").As<uint32>();
+			uint32 ThreadId = EventData.GetValue<uint32>("ThreadId");
 			TSharedRef<FThreadState> ThreadState = GetThreadState(ThreadId);
 			ThreadState->WaitForStreamableHandleHandle = *StreamableHandleState;
-			ThreadState->WaitForStreamableHandleStartCycle = EventData.GetValue("Cycle").As<uint64>();
+			ThreadState->WaitForStreamableHandleStartCycle = EventData.GetValue<uint64>("Cycle");
 		}
 		/*uint64 EventTypeId = StreamableHandlesTimelineEventTypeMap[ThreadState->WaitForStreamableHandleHandle->Id];
 
@@ -323,13 +323,13 @@ void FAsyncLoadingTraceAnalyzer::OnEvent(uint16 RouteId, const FOnEventContext&
 	}
 	case RouteId_EndWaitForStreamableHandle:
 	{
-		uint32 ThreadId = EventData.GetValue("ThreadId").As<uint32>();
+		uint32 ThreadId = EventData.GetValue<uint32>("ThreadId");
 		TSharedRef<FThreadState> ThreadState = GetThreadState(ThreadId);
 		ThreadState->WaitForStreamableHandleHandle = nullptr;
 
 		/*{
 			Trace::FAnalysisSessionEditScope _(Session.Get());
-			BlockingRequestsTimeline->AppendEndEvent(Context.SessionContext.TimestampFromCycle(EventData.GetValue("Cycle").As<uint64>()));
+			BlockingRequestsTimeline->AppendEndEvent(Context.SessionContext.TimestampFromCycle(EventData.GetValue<uint64>("Cycle")));
 		}*/
 
 		break;
@@ -338,7 +338,7 @@ void FAsyncLoadingTraceAnalyzer::OnEvent(uint16 RouteId, const FOnEventContext&
 	{
 		Trace::FAnalysisSessionEditScope _(Session);
 
-		uint64 AsyncPackagePtr = EventData.GetValue("AsyncPackage").As<uint64>();
+		uint64 AsyncPackagePtr = EventData.GetValue<uint64>("AsyncPackage");
 		//check(!ActivePackagesMap.Contains(AsyncPackagePtr));
 		TSharedRef<FAsyncPackageState> AsyncPackageState = MakeShared<FAsyncPackageState>();
 		AsyncPackageState->PackageInfo = &LoadTimeProfilerProvider.CreatePackage(reinterpret_cast<const TCHAR*>(EventData.GetAttachment()));
@@ -347,7 +347,7 @@ void FAsyncLoadingTraceAnalyzer::OnEvent(uint16 RouteId, const FOnEventContext&
 	}
 	case RouteId_DestroyAsyncPackage:
 	{
-		uint64 AsyncPackagePtr = EventData.GetValue("AsyncPackage").As<uint64>();
+		uint64 AsyncPackagePtr = EventData.GetValue<uint64>("AsyncPackage");
 		//check(ActiveAsyncPackagesMap.Contains(AsyncPackagePtr));
 		ActiveAsyncPackagesMap.Remove(AsyncPackagePtr);
 		break;
@@ -355,9 +355,9 @@ void FAsyncLoadingTraceAnalyzer::OnEvent(uint16 RouteId, const FOnEventContext&
 
 	case RouteId_StreamableHandleRequestAssociation:
 	{
-		uint64 StreamableHandlePtr = EventData.GetValue("StreamableHandle").As<uint64>();
+		uint64 StreamableHandlePtr = EventData.GetValue<uint64>("StreamableHandle");
 		TSharedRef<FStreamableHandleState>* StreamableHandleState = ActiveStreamableHandlesMap.Find(StreamableHandlePtr);
-		uint64 RequestId = EventData.GetValue("RequestId").As<uint64>();
+		uint64 RequestId = EventData.GetValue<uint64>("RequestId");
 		TSharedRef<FRequestState>* RequestState = ActiveRequestsMap.Find(RequestId);
 		if (StreamableHandleState && RequestState)
 		{
@@ -367,9 +367,9 @@ void FAsyncLoadingTraceAnalyzer::OnEvent(uint16 RouteId, const FOnEventContext&
 	}
 	case RouteId_AsyncPackageRequestAssociation:
 	{
-		uint64 AsyncPackagePtr = EventData.GetValue("AsyncPackage").As<uint64>();
+		uint64 AsyncPackagePtr = EventData.GetValue<uint64>("AsyncPackage");
 		TSharedRef<FAsyncPackageState>* AsyncPackageState = ActiveAsyncPackagesMap.Find(AsyncPackagePtr);
-		uint64 RequestId = EventData.GetValue("RequestId").As<uint64>();
+		uint64 RequestId = EventData.GetValue<uint64>("RequestId");
 		TSharedRef<FRequestState>* RequestState = ActiveRequestsMap.Find(RequestId);
 		if (AsyncPackageState && RequestState)
 		{
@@ -380,9 +380,9 @@ void FAsyncLoadingTraceAnalyzer::OnEvent(uint16 RouteId, const FOnEventContext&
 	}
 	case RouteId_AsyncPackageLinkerAssociation:
 	{
-		uint64 LinkerPtr = EventData.GetValue("Linker").As<uint64>();
+		uint64 LinkerPtr = EventData.GetValue<uint64>("Linker");
 		TSharedRef<FLinkerState>* LinkerState = ActiveLinkersMap.Find(LinkerPtr);
-		uint64 AsyncPackagePtr = EventData.GetValue("AsyncPackage").As<uint64>();
+		uint64 AsyncPackagePtr = EventData.GetValue<uint64>("AsyncPackage");
 		TSharedRef<FAsyncPackageState>* AsyncPackageState = ActiveAsyncPackagesMap.Find(AsyncPackagePtr);
 		if (LinkerState && AsyncPackageState)
 		{
@@ -394,38 +394,38 @@ void FAsyncLoadingTraceAnalyzer::OnEvent(uint16 RouteId, const FOnEventContext&
 	}
 	case RouteId_LinkerArchiveAssociation:
 	{
-		uint64 ArchivePtr = EventData.GetValue("Archive").As<uint64>();
-		uint64 LinkerPtr = EventData.GetValue("Linker").As<uint64>();
+		uint64 ArchivePtr = EventData.GetValue<uint64>("Archive");
+		uint64 LinkerPtr = EventData.GetValue<uint64>("Linker");
 		break;
 	}
 	case RouteId_BeginAsyncPackageScope:
 	{
 		Trace::FAnalysisSessionEditScope _(Session);
 
-		uint64 AsyncPackagePtr = EventData.GetValue("AsyncPackage").As<uint64>();
+		uint64 AsyncPackagePtr = EventData.GetValue<uint64>("AsyncPackage");
 		TSharedRef<FAsyncPackageState>* AsyncPackageState = ActiveAsyncPackagesMap.Find(AsyncPackagePtr);
-		uint32 ThreadId = EventData.GetValue("ThreadId").As<uint32>();
+		uint32 ThreadId = EventData.GetValue<uint32>("ThreadId");
 		TSharedRef<FThreadState> ThreadState = GetThreadState(ThreadId);
-		ELoadTimeProfilerPackageEventType EventType = static_cast<ELoadTimeProfilerPackageEventType>(EventData.GetValue("EventType").As<uint8>());
-		ThreadState->EnterPackageScope(Context.SessionContext.TimestampFromCycle(EventData.GetValue("Cycle").As<uint64>()), AsyncPackageState ? (*AsyncPackageState)->PackageInfo : nullptr, EventType);
+		ELoadTimeProfilerPackageEventType EventType = static_cast<ELoadTimeProfilerPackageEventType>(EventData.GetValue<uint8>("EventType"));
+		ThreadState->EnterPackageScope(Context.SessionContext.TimestampFromCycle(EventData.GetValue<uint64>("Cycle")), AsyncPackageState ? (*AsyncPackageState)->PackageInfo : nullptr, EventType);
 		break;
 	}
 	case RouteId_EndAsyncPackageScope:
 	{
 		Trace::FAnalysisSessionEditScope _(Session);
 
-		uint32 ThreadId = EventData.GetValue("ThreadId").As<uint32>();
+		uint32 ThreadId = EventData.GetValue<uint32>("ThreadId");
 		TSharedRef<FThreadState> ThreadState = GetThreadState(ThreadId);
-		uint64 Cycle = EventData.GetValue("Cycle").As<uint64>();
-		ThreadState->LeaveScope(Context.SessionContext.TimestampFromCycle(EventData.GetValue("Cycle").As<uint64>()));
+		uint64 Cycle = EventData.GetValue<uint64>("Cycle");
+		ThreadState->LeaveScope(Context.SessionContext.TimestampFromCycle(EventData.GetValue<uint64>("Cycle")));
 		break;
 	}
 	case RouteId_BeginFlushAsyncLoading:
 	{
 		Trace::FAnalysisSessionEditScope _(Session);
 
-		FlushAsyncLoadingRequestId = EventData.GetValue("RequestId").As<uint64>();
-		FlushAsyncLoadingStartCycle = EventData.GetValue("Cycle").As<uint64>();
+		FlushAsyncLoadingRequestId = EventData.GetValue<uint64>("RequestId");
+		FlushAsyncLoadingStartCycle = EventData.GetValue<uint64>("Cycle");
 		if (FlushAsyncLoadingRequestId != INDEX_NONE)
 		{
 			//check(ActiveRequestsMap.Contains(FlushAsyncLoadingRequestId));
@@ -440,7 +440,7 @@ void FAsyncLoadingTraceAnalyzer::OnEvent(uint16 RouteId, const FOnEventContext&
 	}
 	case RouteId_EndFlushAsyncLoading:
 	{
-		uint64 Cycle = EventData.GetValue("Cycle").As<uint64>();
+		uint64 Cycle = EventData.GetValue<uint64>("Cycle");
 		//BlockingRequestsTimeline->AppendEndEvent(Context.SessionContext.TimestampFromCycle(Cycle));
 		break;
 	}
@@ -448,7 +448,7 @@ void FAsyncLoadingTraceAnalyzer::OnEvent(uint16 RouteId, const FOnEventContext&
 	{
 		Trace::FAnalysisSessionEditScope _(Session);
 
-		uint64 ClassPtr = EventData.GetValue("Class").As<uint64>();
+		uint64 ClassPtr = EventData.GetValue<uint64>("Class");
 		const Trace::FClassInfo& ClassInfo = LoadTimeProfilerProvider.AddClassInfo(reinterpret_cast<const TCHAR*>(EventData.GetAttachment()));
 		ClassInfosMap.Add(ClassPtr, &ClassInfo);
 		break;
diff --git a/Engine/Source/Developer/TraceServices/Private/Analyzers/LogTraceAnalysis.cpp b/Engine/Source/Developer/TraceServices/Private/Analyzers/LogTraceAnalysis.cpp
--- a/Engine/Source/Developer/TraceServices/Private/Analyzers/LogTraceAnalysis.cpp
+++ b/Engine/Source/Developer/TraceServices/Private/Analyzers/LogTraceAnalysis.cpp
@@ -29,21 +29,21 @@ void FLogTraceAnalyzer::OnEvent(uint16 RouteId, const FOnEventContext& Context)
 	{
 	case RouteId_LogCategory:
 	{
-		uint64 CategoryPointer = EventData.GetValue("CategoryPointer").As<uint64>();
+		uint64 CategoryPointer = EventData.GetValue<uint64>("CategoryPointer");
 		Trace::FLogCategory& Category = LogProvider.GetCategory(CategoryPointer);
 		Category.Name = Session.StoreString(reinterpret_cast<const TCHAR*>(EventData.GetAttachment()));
-		Category.DefaultVerbosity = static_cast<ELogVerbosity::Type>(EventData.GetValue("DefaultVerbosity").As<uint8>());
+		Category.DefaultVerbosity = static_cast<ELogVerbosity::Type>(EventData.GetValue<uint8>("DefaultVerbosity"));
 		break;
 	}
 	case RouteId_LogMessageSpec:
 	{
-		uint64 LogPoint = EventData.GetValue("LogPoint").As<uint64>();
+		uint64 LogPoint = EventData.GetValue<uint64>("LogPoint");
 		Trace::FLogMessageSpec& Spec = LogProvider.GetMessageSpec(LogPoint);
-		uint64 CategoryPointer = EventData.GetValue("CategoryPointer").As<uint64>();
+		uint64 CategoryPointer = EventData.GetValue<uint64>("CategoryPointer");
 		Trace::FLogCategory& Category = LogProvider.GetCategory(CategoryPointer);
 		Spec.Category = &Category;
-		Spec.Line = EventData.GetValue("Line").As<int32>();
-		Spec.Verbosity = static_cast<ELogVerbosity::Type>(EventData.GetValue("Verbosity").As<uint8>());
+		Spec.Line = EventData.GetValue<int32>("Line");
+		Spec.Verbosity = static_cast<ELogVerbosity::Type>(EventData.GetValue<uint8>("Verbosity"));
 		const ANSICHAR* File = reinterpret_cast<const ANSICHAR*>(EventData.GetAttachment());
 		Spec.File = Session.StoreString(ANSI_TO_TCHAR(File));
 		Spec.FormatString = Session.StoreString(reinterpret_cast<const TCHAR*>(EventData.GetAttachment() + strlen(File) + 1));
@@ -51,8 +51,8 @@ void FLogTraceAnalyzer::OnEvent(uint16 RouteId, const FOnEventContext& Context)
 	}
 	case RouteId_LogMessage:
 	{
-		uint64 LogPoint = EventData.GetValue("LogPoint").As<uint64>();
-		uint64 Cycle = EventData.GetValue("Cycle").As<uint64>();
+		uint64 LogPoint = EventData.GetValue<uint64>("LogPoint");
+		uint64 Cycle = EventData.GetValue<uint64>("Cycle");
 		LogProvider.AppendMessage(LogPoint, Context.SessionContext.TimestampFromCycle(Cycle), EventData.GetAttachment());
 		break;
 	}
diff --git a/Engine/Source/Developer/TraceServices/Private/Analyzers/MiscTraceAnalysis.cpp b/Engine/Source/Developer/TraceServices/Private/Analyzers/MiscTraceAnalysis.cpp
--- a/Engine/Source/Developer/TraceServices/Private/Analyzers/MiscTraceAnalysis.cpp
+++ b/Engine/Source/Developer/TraceServices/Private/Analyzers/MiscTraceAnalysis.cpp
@@ -51,16 +51,16 @@ void FMiscTraceAnalyzer::OnEvent(uint16 RouteId, const FOnEventContext& Context)
 	{
 	case RouteId_RegisterGameThread:
 	{
-		uint32 ThreadId = EventData.GetValue("ThreadId").As<uint32>();
+		uint32 ThreadId = EventData.GetValue<uint32>("ThreadId");
 		ThreadProvider.AddGameThread(ThreadId);
 		break;
 	}
 	case RouteId_CreateThread:
 	{
-		uint32 CreatedThreadId = EventData.GetValue("CreatedThreadId").As<uint32>();
-		EThreadPriority Priority = static_cast<EThreadPriority>(EventData.GetValue("Priority").As<uint32>());
+		uint32 CreatedThreadId = EventData.GetValue<uint32>("CreatedThreadId");
+		EThreadPriority Priority = static_cast<EThreadPriority>(EventData.GetValue<uint32>("Priority"));
 		ThreadProvider.AddThread(CreatedThreadId, reinterpret_cast<const TCHAR*>(EventData.GetAttachment()), Priority);
-		uint32 CurrentThreadId = EventData.GetValue("CurrentThreadId").As<uint32>();
+		uint32 CurrentThreadId = EventData.GetValue<uint32>("CurrentThreadId");
 		FThreadState* ThreadState = GetThreadState(CurrentThreadId);
 		if (ThreadState->ThreadGroupStack.Num())
 		{
@@ -71,30 +71,30 @@ void FMiscTraceAnalyzer::OnEvent(uint16 RouteId, const FOnEventContext& Context)
 	case RouteId_SetThreadGroup:
 	{
 		const TCHAR* GroupName = Session.StoreString(ANSI_TO_TCHAR(reinterpret_cast<const char*>(EventData.GetAttachment())));
-		uint32 ThreadId = EventData.GetValue("ThreadId").As<uint32>();
+		uint32 ThreadId = EventData.GetValue<uint32>("ThreadId");
 		ThreadProvider.SetThreadGroup(ThreadId, GroupName);
 		break;
 	}
 	case RouteId_BeginThreadGroupScope:
 	{
 		const TCHAR* GroupName = Session.StoreString(ANSI_TO_TCHAR(reinterpret_cast<const char*>(EventData.GetAttachment())));
-		uint32 CurrentThreadId = EventData.GetValue("CurrentThreadId").As<uint32>();
+		uint32 CurrentThreadId = EventData.GetValue<uint32>("CurrentThreadId");
 		FThreadState* ThreadState = GetThreadState(CurrentThreadId);
 		ThreadState->ThreadGroupStack.Push(GroupName);
 		break;
 	}
 	case RouteId_EndThreadGroupScope:
 	{
-		uint32 CurrentThreadId = EventData.GetValue("CurrentThreadId").As<uint32>();
+		uint32 CurrentThreadId = EventData.GetValue<uint32>("CurrentThreadId");
 		FThreadState* ThreadState = GetThreadState(CurrentThreadId);
 		ThreadState->ThreadGroupStack.Pop();
 		break;
 	}
 	case RouteId_BookmarkSpec:
 	{
-		uint64 BookmarkPoint = EventData.GetValue("BookmarkPoint").As<uint64>();
+		uint64 BookmarkPoint = EventData.GetValue<uint64>("BookmarkPoint");
 		Trace::FBookmarkSpec& Spec = BookmarkProvider.GetSpec(BookmarkPoint);
-		Spec.Line = EventData.GetValue("Line").As<int32>();
+		Spec.Line = EventData.GetValue<int32>("Line");
 		const ANSICHAR* File = reinterpret_cast<const ANSICHAR*>(EventData.GetAttachment());
 		Spec.File = Session.StoreString(ANSI_TO_TCHAR(File));
 		Spec.FormatString = Session.StoreString(reinterpret_cast<const TCHAR*>(EventData.GetAttachment() + strlen(File) + 1));
@@ -109,8 +109,8 @@ void FMiscTraceAnalyzer::OnEvent(uint16 RouteId, const FOnEventContext& Context)
 	}
 	case RouteId_Bookmark:
 	{
-		uint64 BookmarkPoint = EventData.GetValue("BookmarkPoint").As<uint64>();
-		uint64 Cycle = EventData.GetValue("Cycle").As<uint64>();
+		uint64 BookmarkPoint = EventData.GetValue<uint64>("BookmarkPoint");
+		uint64 Cycle = EventData.GetValue<uint64>("Cycle");
 		double Timestamp = Context.SessionContext.TimestampFromCycle(Cycle);
 		BookmarkProvider.AppendBookmark(Timestamp, BookmarkPoint, EventData.GetAttachment());
 		LogProvider.AppendMessage(BookmarkPoint, Timestamp, EventData.GetAttachment());
@@ -118,16 +118,16 @@ void FMiscTraceAnalyzer::OnEvent(uint16 RouteId, const FOnEventContext& Context)
 	}
 	case RouteId_BeginFrame:
 	{
-		uint64 Cycle = EventData.GetValue("Cycle").As<uint64>();
-		uint8 FrameType = EventData.GetValue("FrameType").As<uint8>();
+		uint64 Cycle = EventData.GetValue<uint64>("Cycle");
+		uint8 FrameType = EventData.GetValue<uint8>("FrameType");
 		check(FrameType < TraceFrameType_Count);
 		FrameProvider.BeginFrame(ETraceFrameType(FrameType), Context.SessionContext.TimestampFromCycle(Cycle));
 		break;
 	}
 	case RouteId_EndFrame:
 	{
-		uint64 Cycle = EventData.GetValue("Cycle").As<uint64>();
-		uint8 FrameType = EventData.GetValue("FrameType").As<uint8>();
+		uint64 Cycle = EventData.GetValue<uint64>("Cycle");
+		uint8 FrameType = EventData.GetValue<uint8>("FrameType");
 		check(FrameType < TraceFrameType_Count);
 		FrameProvider.EndFrame(ETraceFrameType(FrameType), Context.SessionContext.TimestampFromCycle(Cycle));
 		break;
diff --git a/Engine/Source/Developer/TraceServices/Private/Analyzers/PlatformFileTraceAnalysis.cpp b/Engine/Source/Developer/TraceServices/Private/Analyzers/PlatformFileTraceAnalysis.cpp
--- a/Engine/Source/Developer/TraceServices/Private/Analyzers/PlatformFileTraceAnalysis.cpp
+++ b/Engine/Source/Developer/TraceServices/Private/Analyzers/PlatformFileTraceAnalysis.cpp
@@ -33,9 +33,9 @@ void FPlatformFileTraceAnalyzer::OnEvent(uint16 RouteId, const FOnEventContext&
 	{
 	case RouteId_BeginOpen:
 	{
-		double Time = Context.SessionContext.TimestampFromCycle(EventData.GetValue("Cycle").As<uint64>());
-		uint64 TempHandle = EventData.GetValue("TempHandle").As<uint64>();
-		uint32 ThreadId = EventData.GetValue("ThreadId").As<uint32>();
+		double Time = Context.SessionContext.TimestampFromCycle(EventData.GetValue<uint64>("Cycle"));
+		uint64 TempHandle = EventData.GetValue<uint64>("TempHandle");
+		uint32 ThreadId = EventData.GetValue<uint32>("ThreadId");
 		check(!PendingOpenMap.Contains(TempHandle));
 		uint32 FileIndex = FileActivityProvider.GetFileIndex(reinterpret_cast<const TCHAR*>(EventData.GetAttachment()));
 		FPendingActivity& Open = PendingOpenMap.Add(TempHandle);
@@ -45,9 +45,9 @@ void FPlatformFileTraceAnalyzer::OnEvent(uint16 RouteId, const FOnEventContext&
 	}
 	case RouteId_EndOpen:
 	{
-		double Time = Context.SessionContext.TimestampFromCycle(EventData.GetValue("Cycle").As<uint64>());
-		uint64 TempHandle = EventData.GetValue("TempHandle").As<uint64>();
-		uint64 FileHandle = EventData.GetValue("FileHandle").As<uint64>();
+		double Time = Context.SessionContext.TimestampFromCycle(EventData.GetValue<uint64>("Cycle"));
+		uint64 TempHandle = EventData.GetValue<uint64>("TempHandle");
+		uint64 FileHandle = EventData.GetValue<uint64>("FileHandle");
 		check(PendingOpenMap.Contains(TempHandle));
 		FPendingActivity& Open = PendingOpenMap[TempHandle];
 		if (FileHandle == uint64(-1)) // TODO: Portable invalid file handle?
@@ -65,10 +65,10 @@ void FPlatformFileTraceAnalyzer::OnEvent(uint16 RouteId, const FOnEventContext&
 	}
 	case RouteId_BeginClose:
 	{
-		double Time = Context.SessionContext.TimestampFromCycle(EventData.GetValue("Cycle").As<uint64>());
-		uint64 TempHandle = EventData.GetValue("TempHandle").As<uint64>();
-		uint64 FileHandle = EventData.GetValue("FileHandle").As<uint64>();
-		uint32 ThreadId = EventData.GetValue("ThreadId").As<uint32>();
+		double Time = Context.SessionContext.TimestampFromCycle(EventData.GetValue<uint64>("Cycle"));
+		uint64 TempHandle = EventData.GetValue<uint64>("TempHandle");
+		uint64 FileHandle = EventData.GetValue<uint64>("FileHandle");
+		uint32 ThreadId = EventData.GetValue<uint32>("ThreadId");
 		check(OpenFilesMap.Contains(FileHandle));
 		check(!PendingCloseMap.Contains(TempHandle));
 		uint64 FileIndex = OpenFilesMap[FileHandle];
@@ -80,8 +80,8 @@ void FPlatformFileTraceAnalyzer::OnEvent(uint16 RouteId, const FOnEventContext&
 	}
 	case RouteId_EndClose:
 	{
-		double Time = Context.SessionContext.TimestampFromCycle(EventData.GetValue("Cycle").As<uint64>());
-		uint64 TempHandle = EventData.GetValue("TempHandle").As<uint64>();
+		double Time = Context.SessionContext.TimestampFromCycle(EventData.GetValue<uint64>("Cycle"));
+		uint64 TempHandle = EventData.GetValue<uint64>("TempHandle");
 		check(PendingCloseMap.Contains(TempHandle));
 		FPendingActivity& Close = PendingCloseMap[TempHandle];
 		FileActivityProvider.EndActivity(Close.FileIndex, Close.ActivityIndex, Time, false);
@@ -90,12 +90,12 @@ void FPlatformFileTraceAnalyzer::OnEvent(uint16 RouteId, const FOnEventContext&
 	}
 	case RouteId_BeginRead:
 	{
-		double Time = Context.SessionContext.TimestampFromCycle(EventData.GetValue("Cycle").As<uint64>());
-		uint64 ReadHandle = EventData.GetValue("ReadHandle").As<uint64>();
-		uint64 FileHandle = EventData.GetValue("FileHandle").As<uint64>();
-		uint64 Offset = EventData.GetValue("Offset").As<uint64>();
-		uint64 Size = EventData.GetValue("Size").As<uint64>();
-		uint32 ThreadId = EventData.GetValue("ThreadId").As<uint32>();
+		double Time = Context.SessionContext.TimestampFromCycle(EventData.GetValue<uint64>("Cycle"));
+		uint64 ReadHandle = EventData.GetValue<uint64>("ReadHandle");
+		uint64 FileHandle = EventData.GetValue<uint64>("FileHandle");
+		uint64 Offset = EventData.GetValue<uint64>("Offset");
+		uint64 Size = EventData.GetValue<uint64>("Size");
+		uint32 ThreadId = EventData.GetValue<uint32>("ThreadId");
 		check(OpenFilesMap.Contains(FileHandle));
 		uint32 FileIndex = OpenFilesMap[FileHandle];
 		uint64 ReadIndex = FileActivityProvider.BeginActivity(FileIndex, Trace::FileActivityType_Read, Offset, Size, Time);
@@ -106,10 +106,10 @@ void FPlatformFileTraceAnalyzer::OnEvent(uint16 RouteId, const FOnEventContext&
 	}
 	case RouteId_EndRead:
 	{
-		double Time = Context.SessionContext.TimestampFromCycle(EventData.GetValue("Cycle").As<uint64>());
-		uint64 ReadHandle = EventData.GetValue("ReadHandle").As<uint64>();
-		uint64 SizeRead = EventData.GetValue("SizeRead").As<uint64>();
-		uint32 ThreadId = EventData.GetValue("ThreadId").As<uint32>();
+		double Time = Context.SessionContext.TimestampFromCycle(EventData.GetValue<uint64>("Cycle"));
+		uint64 ReadHandle = EventData.GetValue<uint64>("ReadHandle");
+		uint64 SizeRead = EventData.GetValue<uint64>("SizeRead");
+		uint32 ThreadId = EventData.GetValue<uint32>("ThreadId");
 		check(ActiveReadsMap.Contains(ReadHandle));
 		const FPendingActivity& Read = ActiveReadsMap[ReadHandle];
 		FileActivityProvider.EndActivity(Read.FileIndex, Read.ActivityIndex, Time, false);
@@ -118,12 +118,12 @@ void FPlatformFileTraceAnalyzer::OnEvent(uint16 RouteId, const FOnEventContext&
 	}
 	case RouteId_BeginWrite:
 	{
-		double Time = Context.SessionContext.TimestampFromCycle(EventData.GetValue("Cycle").As<uint64>());
-		uint64 WriteHandle = EventData.GetValue("WriteHandle").As<uint64>();
-		uint64 FileHandle = EventData.GetValue("FileHandle").As<uint64>();
-		uint64 Offset = EventData.GetValue("Offset").As<uint64>();
-		uint64 Size = EventData.GetValue("Size").As<uint64>();
-		uint32 ThreadId = EventData.GetValue("ThreadId").As<uint32>();
+		double Time = Context.SessionContext.TimestampFromCycle(EventData.GetValue<uint64>("Cycle"));
+		uint64 WriteHandle = EventData.GetValue<uint64>("WriteHandle");
+		uint64 FileHandle = EventData.GetValue<uint64>("FileHandle");
+		uint64 Offset = EventData.GetValue<uint64>("Offset");
+		uint64 Size = EventData.GetValue<uint64>("Size");
+		uint32 ThreadId = EventData.GetValue<uint32>("ThreadId");
 		check(OpenFilesMap.Contains(FileHandle));
 		uint32 FileIndex = OpenFilesMap[FileHandle];
 		uint64 WriteIndex = FileActivityProvider.BeginActivity(FileIndex, Trace::FileActivityType_Write, Offset, Size, Time);
@@ -134,10 +134,10 @@ void FPlatformFileTraceAnalyzer::OnEvent(uint16 RouteId, const FOnEventContext&
 	}
 	case RouteId_EndWrite:
 	{
-		double Time = Context.SessionContext.TimestampFromCycle(EventData.GetValue("Cycle").As<uint64>());
-		uint64 WriteHandle = EventData.GetValue("WriteHandle").As<uint64>();
-		uint64 SizeWritten = EventData.GetValue("SizeWritten").As<uint64>();
-		uint32 ThreadId = EventData.GetValue("ThreadId").As<uint32>();
+		double Time = Context.SessionContext.TimestampFromCycle(EventData.GetValue<uint64>("Cycle"));
+		uint64 WriteHandle = EventData.GetValue<uint64>("WriteHandle");
+		uint64 SizeWritten = EventData.GetValue<uint64>("SizeWritten");
+		uint32 ThreadId = EventData.GetValue<uint32>("ThreadId");
 		check(ActiveWritesMap.Contains(WriteHandle));
 		const FPendingActivity& Write = ActiveWritesMap[WriteHandle];
 		FileActivityProvider.EndActivity(Write.FileIndex, Write.ActivityIndex, Time, false);
diff --git a/Engine/Source/Developer/TraceServices/Private/Analyzers/StatsTraceAnalysis.cpp b/Engine/Source/Developer/TraceServices/Private/Analyzers/StatsTraceAnalysis.cpp
--- a/Engine/Source/Developer/TraceServices/Private/Analyzers/StatsTraceAnalysis.cpp
+++ b/Engine/Source/Developer/TraceServices/Private/Analyzers/StatsTraceAnalysis.cpp
@@ -31,16 +31,16 @@ void FStatsAnalyzer::OnEvent(uint16 RouteId, const FOnEventContext& Context)
 	case RouteId_Spec:
 	{
 		Trace::FAnalysisSessionEditScope _(Session);
-		uint32 StatId = EventData.GetValue("Id").As<uint32>();
+		uint32 StatId = EventData.GetValue<uint32>("Id");
 		check(!CountersMap.Contains(StatId));
 		const ANSICHAR* Name = reinterpret_cast<const ANSICHAR*>(EventData.GetAttachment());
 		const TCHAR* Description = reinterpret_cast<const TCHAR*>(EventData.GetAttachment() + strlen(Name) + 1);
 		Trace::ECounterDisplayHint DisplayHint = Trace::CounterDisplayHint_None;
-		if (EventData.GetValue("IsFloatingPoint").As<bool>())
+		if (EventData.GetValue<bool>("IsFloatingPoint"))
 		{
 			DisplayHint = Trace::CounterDisplayHint_FloatingPoint;
 		}
-		else if (EventData.GetValue("IsMemory").As<bool>())
+		else if (EventData.GetValue<bool>("IsMemory"))
 		{
 			DisplayHint = Trace::CounterDisplayHint_Memory;
 		}
@@ -51,7 +51,7 @@ void FStatsAnalyzer::OnEvent(uint16 RouteId, const FOnEventContext& Context)
 	case RouteId_EventBatch:
 	{
 		Trace::FAnalysisSessionEditScope _(Session);
-		uint32 ThreadId = EventData.GetValue("ThreadId").As<uint32>();
+		uint32 ThreadId = EventData.GetValue<uint32>("ThreadId");
 		TSharedRef<FThreadState> ThreadState = GetThreadState(ThreadId);
 		uint64 BufferSize = EventData.GetAttachmentSize();
 		const uint8* BufferPtr = EventData.GetAttachment();
