UnrealInsights:
Fixed issue with deleted .utrace files still appearing in the session list

diff --git a/Engine/Source/Developer/TraceServices/Private/SessionService.cpp b/Engine/Source/Developer/TraceServices/Private/SessionService.cpp
--- a/Engine/Source/Developer/TraceServices/Private/SessionService.cpp
+++ b/Engine/Source/Developer/TraceServices/Private/SessionService.cpp
@@ -233,16 +233,28 @@ void FSessionService::UpdateSessions()
 	TArray<FRecorderSessionInfo> RecorderSessions;
 	TraceRecorder->GetActiveSessions(RecorderSessions);
 
+	TSet<FSessionHandle> SessionsToRemove;
+	for (const auto& KV : Sessions)
+	{
+		SessionsToRemove.Add(KV.Key);
+	}
+
 	FScopeLock Lock(&SessionsCS);
 	for (const FStoreSessionInfo& StoreSession : StoreSessions)
 	{
 		FSessionInfoInternal& Session = Sessions.FindOrAdd(StoreSession.Handle);
+		SessionsToRemove.Remove(StoreSession.Handle);
 		Session.Uri = StoreSession.Uri;
 		Session.Name = StoreSession.Name;
 		Session.bIsLive = StoreSession.bIsLive;
 		Session.RecorderSessionHandle = 0;
 	}
 
+	for (FSessionHandle Session : SessionsToRemove)
+	{
+		Sessions.Remove(Session);
+	}
+
 	for (const FRecorderSessionInfo& RecorderSession : RecorderSessions)
 	{
 		FSessionInfoInternal* FindIt = Sessions.Find(RecorderSession.StoreSessionHandle);
