UnrealInsights:
Made counters a generic provider concept that can be used by multiple analyzers

diff --git a/Engine/Source/Developer/TraceInsights/Private/Insights/Widgets/SStatsView.cpp b/Engine/Source/Developer/TraceInsights/Private/Insights/Widgets/SStatsView.cpp
--- a/Engine/Source/Developer/TraceInsights/Private/Insights/Widgets/SStatsView.cpp
+++ b/Engine/Source/Developer/TraceInsights/Private/Insights/Widgets/SStatsView.cpp
@@ -1401,11 +1401,11 @@ void SStatsView::RebuildTree(bool bResync)
 		bListHasChanged = true;
 	}
 
-	if (Session.IsValid() && Trace::ReadCounterProvider(*Session.Get()))
+	if (Session.IsValid())
 	{
 		Trace::FAnalysisSessionReadScope SessionReadScope(*Session.Get());
 
-		const Trace::ICounterProvider& CountersProvider = *Trace::ReadCounterProvider(*Session.Get());
+		const Trace::ICounterProvider& CountersProvider = Trace::ReadCounterProvider(*Session.Get());
 
 		if (!bResync)
 		{
@@ -1423,8 +1423,8 @@ void SStatsView::RebuildTree(bool bResync)
 			{
 				FName Name(Counter.GetName());
 				FName Group(Counter.GetDisplayHint() == Trace::CounterDisplayHint_Memory ? TEXT("Memory") :
-							Counter.GetDisplayHint() == Trace::CounterDisplayHint_FloatingPoint ? TEXT("float") : TEXT("int64"));
-				EStatsNodeType Type = Counter.GetDisplayHint() == Trace::CounterDisplayHint_FloatingPoint ? EStatsNodeType::Float : EStatsNodeType::Int64;
+							Counter.IsFloatingPoint() ? TEXT("float") : TEXT("int64"));
+				EStatsNodeType Type = Counter.IsFloatingPoint() ? EStatsNodeType::Float : EStatsNodeType::Int64;
 				FStatsNodePtr StatsNodePtr = MakeShareable(new FStatsNode(Counter.GetId(), Name, Group, Type));
 				StatsNodes.Add(StatsNodePtr);
 				//StatsNodesMap.Add(Name, StatsNodePtr);
@@ -1715,7 +1715,7 @@ void SStatsView::UpdateStats(double StartTime, double EndTime)
 		StatsNodePtr->ResetAggregatedStats();
 	}
 
-	if (Session.IsValid() && StartTime < EndTime && Trace::ReadCounterProvider(*Session.Get()))
+	if (Session.IsValid() && StartTime < EndTime)
 	{
 		const bool bComputeMedian = true;
 
@@ -1725,13 +1725,13 @@ void SStatsView::UpdateStats(double StartTime, double EndTime)
 		{
 			Trace::FAnalysisSessionReadScope SessionReadScope(*Session.Get());
 
-			const Trace::ICounterProvider& CountersProvider = *Trace::ReadCounterProvider(*Session.Get());
+			const Trace::ICounterProvider& CountersProvider = Trace::ReadCounterProvider(*Session.Get());
 
 			// Compute instance count and total/min/max inclusive/exclusive times for each counter.
 			// Iterate through all counters.
 			CountersProvider.EnumerateCounters([&CalculationHelperDbl, &CalculationHelperInt](const Trace::ICounter& Counter)
 			{
-				if (Counter.GetDisplayHint() == Trace::CounterDisplayHint_FloatingPoint)
+				if (Counter.IsFloatingPoint())
 				{
 					CalculationHelperDbl.Update(Counter);
 				}
@@ -1752,7 +1752,7 @@ void SStatsView::UpdateStats(double StartTime, double EndTime)
 				// Iterate again through all counters.
 				CountersProvider.EnumerateCounters([&CalculationHelperDbl, &CalculationHelperInt](const Trace::ICounter& Counter)
 				{
-					if (Counter.GetDisplayHint() == Trace::CounterDisplayHint_FloatingPoint)
+					if (Counter.IsFloatingPoint())
 					{
 						CalculationHelperDbl.UpdateHistograms(Counter);
 					}
diff --git a/Engine/Source/Developer/TraceServices/Private/AnalysisService.cpp b/Engine/Source/Developer/TraceServices/Private/AnalysisService.cpp
--- a/Engine/Source/Developer/TraceServices/Private/AnalysisService.cpp
+++ b/Engine/Source/Developer/TraceServices/Private/AnalysisService.cpp
@@ -14,6 +14,7 @@
 #include "Model/FramesPrivate.h"
 #include "Model/BookmarksPrivate.h"
 #include "Model/ThreadsPrivate.h"
+#include "Model/CountersPrivate.h"
 
 namespace Trace
 {
@@ -71,21 +72,43 @@ FAnalysisSession::FAnalysisSession(const TCHAR* SessionName)
 
 FAnalysisSession::~FAnalysisSession()
 {
-	for (auto& KV : Providers)
+	for (int32 AnalyzerIndex = Analyzers.Num() - 1; AnalyzerIndex >= 0; --AnalyzerIndex)
 	{
-		IProvider* Provider = KV.Value;
-		delete Provider;
+		delete Analyzers[AnalyzerIndex];
 	}
+	for (int32 ProviderIndex = Providers.Num() - 1; ProviderIndex >= 0; --ProviderIndex)
+	{
+		delete Providers[ProviderIndex];
+	}
+}
+
+void FAnalysisSession::AddAnalyzer(IAnalyzer* Analyzer)
+{
+	Analyzers.Add(Analyzer);
 }
 
 void FAnalysisSession::AddProvider(const FName& InName, IProvider* Provider)
 {
-	Providers.Add(InName, Provider);
+	Providers.Add(Provider);
+	ProvidersMap.Add(InName, Provider);
 }
 
 const IProvider* FAnalysisSession::ReadProviderPrivate(const FName& InName) const
 {
-	IProvider* const* FindIt = Providers.Find(InName);
+	IProvider* const* FindIt = ProvidersMap.Find(InName);
+	if (FindIt)
+	{
+		return *FindIt;
+	}
+	else
+	{
+		return nullptr;
+	}
+}
+
+IProvider* FAnalysisSession::EditProviderPrivate(const FName& InName)
+{
+	IProvider** FindIt = ProvidersMap.Find(InName);
 	if (FindIt)
 	{
 		return *FindIt;
@@ -114,7 +137,6 @@ void FAnalysisService::AnalyzeInternal(TSharedRef<FAnalysisSession> AnalysisSess
 {
 	OnAnalysisStarted().Broadcast(AnalysisSession);
 
-	TArray<IAnalyzer*> Analyzers;
 	{
 		IAnalysisSession& Session = AnalysisSession.Get();
 		Trace::FAnalysisSessionEditScope _(Session);
@@ -126,25 +148,21 @@ void FAnalysisService::AnalyzeInternal(TSharedRef<FAnalysisSession> AnalysisSess
 		Session.AddProvider(FThreadProvider::ProviderName, ThreadProvider);
 		FFrameProvider* FrameProvider = new FFrameProvider(Session);
 		Session.AddProvider(FFrameProvider::ProviderName, FrameProvider);
-		Analyzers.Add(new FMiscTraceAnalyzer(Session, *ThreadProvider, *BookmarkProvider, *LogProvider, *FrameProvider));
-		Analyzers.Add(new FLogTraceAnalyzer(Session, *LogProvider));
-		ModuleService.OnAnalysisBegin(Session, Analyzers);
+		FCounterProvider* CounterProvider = new FCounterProvider(Session, *FrameProvider);
+		Session.AddProvider(FCounterProvider::ProviderName, CounterProvider);
+		Session.AddAnalyzer(new FMiscTraceAnalyzer(Session, *ThreadProvider, *BookmarkProvider, *LogProvider, *FrameProvider));
+		Session.AddAnalyzer(new FLogTraceAnalyzer(Session, *LogProvider));
+		ModuleService.OnAnalysisBegin(Session);
 	}
 	
 	FAnalysisContext Context;
-	for (Trace::IAnalyzer* Analyzer : Analyzers)
+	for (Trace::IAnalyzer* Analyzer : AnalysisSession->ReadAnalyzers())
 	{
 		Context.AddAnalyzer(*Analyzer);
 	}
 	Trace::FAnalysisProcessor Processor = Context.Process(*DataStream);
 	Processor.Wait();
 
-	for (Trace::IAnalyzer* Analyzer : Analyzers)
-	{
-		delete Analyzer;
-	}
-	Analyzers.Empty();
-
 	delete DataStream;
 
 	AnalysisSession->SetComplete();
diff --git a/Engine/Source/Developer/TraceServices/Private/AnalysisServicePrivate.h b/Engine/Source/Developer/TraceServices/Private/AnalysisServicePrivate.h
--- a/Engine/Source/Developer/TraceServices/Private/AnalysisServicePrivate.h
+++ b/Engine/Source/Developer/TraceServices/Private/AnalysisServicePrivate.h
@@ -56,10 +56,14 @@ public:
 	virtual void ReadAccessCheck() const override { return Lock.ReadAccessCheck(); }
 	virtual void WriteAccessCheck() override { return Lock.WriteAccessCheck(); }
 
+	virtual void AddAnalyzer(IAnalyzer* Analyzer) override;
 	virtual void AddProvider(const FName& Name, IProvider* Provider) override;
 
+	const TArray<IAnalyzer*> ReadAnalyzers() { return Analyzers; }
+
 private:
 	virtual const IProvider* ReadProviderPrivate(const FName& Name) const override;
+	virtual IProvider* EditProviderPrivate(const FName& Name) override;
 
 	mutable FAnalysisSessionLock Lock;
 
@@ -68,7 +72,9 @@ private:
 	double DurationSeconds = 0.0;
 	FSlabAllocator Allocator;
 	FStringStore StringStore;
-	TMap<FName, IProvider*> Providers;
+	TArray<IAnalyzer*> Analyzers;
+	TArray<IProvider*> Providers;
+	TMap<FName, IProvider*> ProvidersMap;
 };
 
 class FAnalysisService
diff --git a/Engine/Source/Developer/TraceServices/Private/Analyzers/StatsTraceAnalysis.cpp b/Engine/Source/Developer/TraceServices/Private/Analyzers/StatsTraceAnalysis.cpp
--- a/Engine/Source/Developer/TraceServices/Private/Analyzers/StatsTraceAnalysis.cpp
+++ b/Engine/Source/Developer/TraceServices/Private/Analyzers/StatsTraceAnalysis.cpp
@@ -2,9 +2,9 @@
 #include "StatsTraceAnalysis.h"
 #include "AnalysisServicePrivate.h"
 #include "Common/Utils.h"
-#include "Model/Counters.h"
+#include "TraceServices/Model/Counters.h"
 
-FStatsAnalyzer::FStatsAnalyzer(Trace::FAnalysisSession& InSession, Trace::FCounterProvider& InCounterProvider)
+FStatsAnalyzer::FStatsAnalyzer(Trace::IAnalysisSession& InSession, Trace::ICounterProvider& InCounterProvider)
 	: Session(InSession)
 	, CounterProvider(InCounterProvider)
 {
@@ -25,32 +25,40 @@ void FStatsAnalyzer::OnAnalysisEnd()
 
 void FStatsAnalyzer::OnEvent(uint16 RouteId, const FOnEventContext& Context)
 {
+	Trace::FAnalysisSessionEditScope _(Session);
+
 	const auto& EventData = Context.EventData;
 	switch (RouteId)
 	{
 	case RouteId_Spec:
 	{
-		Trace::FAnalysisSessionEditScope _(Session);
 		uint32 StatId = EventData.GetValue<uint32>("Id");
-		check(!CountersMap.Contains(StatId));
+		Trace::ICounter* Counter;
+		Trace::ICounter** FindIt = CountersMap.Find(StatId);
+		if (!FindIt)
+		{
+			Counter = CounterProvider.CreateCounter();
+			CountersMap.Add(StatId, Counter);
+		}
+		else
+		{
+			Counter = *FindIt;
+		}
 		const ANSICHAR* Name = reinterpret_cast<const ANSICHAR*>(EventData.GetAttachment());
 		const TCHAR* Description = reinterpret_cast<const TCHAR*>(EventData.GetAttachment() + strlen(Name) + 1);
 		Trace::ECounterDisplayHint DisplayHint = Trace::CounterDisplayHint_None;
-		if (EventData.GetValue<bool>("IsFloatingPoint"))
-		{
-			DisplayHint = Trace::CounterDisplayHint_FloatingPoint;
-		}
-		else if (EventData.GetValue<bool>("IsMemory"))
+		if (EventData.GetValue<bool>("IsMemory"))
 		{
 			DisplayHint = Trace::CounterDisplayHint_Memory;
 		}
-		Trace::FCounterInternal* Counter = CounterProvider.CreateCounter(ANSI_TO_TCHAR(Name), Description, DisplayHint);
-		CountersMap.Add(StatId, Counter);
+		Counter->SetName(Session.StoreString(ANSI_TO_TCHAR(Name)));
+		Counter->SetDescription(Session.StoreString(Description));
+		Counter->SetIsFloatingPoint(EventData.GetValue<bool>("IsFloatingPoint"));
+		Counter->SetDisplayHint(DisplayHint);
 		break;
 	}
 	case RouteId_EventBatch:
 	{
-		Trace::FAnalysisSessionEditScope _(Session);
 		uint32 ThreadId = EventData.GetValue<uint32>("ThreadId");
 		TSharedRef<FThreadState> ThreadState = GetThreadState(ThreadId);
 		uint64 BufferSize = EventData.GetAttachmentSize();
@@ -70,8 +78,17 @@ void FStatsAnalyzer::OnEvent(uint16 RouteId, const FOnEventContext& Context)
 
 			uint64 DecodedIdAndOp = FTraceAnalyzerUtils::Decode7bit(BufferPtr);
 			uint32 StatId = DecodedIdAndOp >> 3;
-			//check(CountersMap.Contains(StatId));
-			Trace::FCounterInternal& Counter = *CountersMap[StatId];
+			Trace::ICounter* Counter;
+			Trace::ICounter** FindIt = CountersMap.Find(StatId);
+			if (!FindIt)
+			{
+				Counter = CounterProvider.CreateCounter();
+				CountersMap.Add(StatId, Counter);
+			}
+			else
+			{
+				Counter = *FindIt;
+			}
 			uint8 Op = DecodedIdAndOp & 0x7;
 			uint64 CycleDiff = FTraceAnalyzerUtils::Decode7bit(BufferPtr);
 			uint64 Cycle = ThreadState->LastCycle + CycleDiff;
@@ -81,24 +98,24 @@ void FStatsAnalyzer::OnEvent(uint16 RouteId, const FOnEventContext& Context)
 			{
 			case Increment:
 			{
-				CounterProvider.Add(Counter, Time, int64(1));
+				Counter->AddValue(Time, int64(1));
 				break;
 			}
 			case Decrement:
 			{
-				CounterProvider.Add(Counter, Time, int64(-1));
+				Counter->AddValue(Time, int64(-1));
 				break;
 			}
 			case AddInteger:
 			{
 				int64 Amount = FTraceAnalyzerUtils::DecodeZigZag(BufferPtr);
-				CounterProvider.Add(Counter, Time, Amount);
+				Counter->AddValue(Time, Amount);
 				break;
 			}
 			case SetInteger:
 			{
 				int64 Value = FTraceAnalyzerUtils::DecodeZigZag(BufferPtr);
-				CounterProvider.Set(Counter, Time, Value);
+				Counter->SetValue(Time, Value);
 				break;
 			}
 			case AddFloat:
@@ -106,7 +123,7 @@ void FStatsAnalyzer::OnEvent(uint16 RouteId, const FOnEventContext& Context)
 				double Amount;
 				memcpy(&Amount, BufferPtr, sizeof(double));
 				BufferPtr += sizeof(double);
-				CounterProvider.Add(Counter, Time, Amount);
+				Counter->AddValue(Time, Amount);
 				break;
 			}
 			case SetFloat:
@@ -114,6 +131,7 @@ void FStatsAnalyzer::OnEvent(uint16 RouteId, const FOnEventContext& Context)
 				double Value;
 				memcpy(&Value, BufferPtr, sizeof(double));
 				BufferPtr += sizeof(double);
+				Counter->SetValue(Time, Value);
 				break;
 			}
 			}
diff --git a/Engine/Source/Developer/TraceServices/Private/Analyzers/StatsTraceAnalysis.h b/Engine/Source/Developer/TraceServices/Private/Analyzers/StatsTraceAnalysis.h
--- a/Engine/Source/Developer/TraceServices/Private/Analyzers/StatsTraceAnalysis.h
+++ b/Engine/Source/Developer/TraceServices/Private/Analyzers/StatsTraceAnalysis.h
@@ -8,16 +8,16 @@
 
 namespace Trace
 {
-	class FAnalysisSession;
-	class FCounterInternal;
-	class FCounterProvider;
+	class IAnalysisSession;
+	class ICounter;
+	class ICounterProvider;
 }
 
 class FStatsAnalyzer
 	: public Trace::IAnalyzer
 {
 public:
-	FStatsAnalyzer(Trace::FAnalysisSession& Session, Trace::FCounterProvider& CounterProvider);
+	FStatsAnalyzer(Trace::IAnalysisSession& Session, Trace::ICounterProvider& CounterProvider);
 	virtual void OnAnalysisBegin(const FOnAnalysisContext& Context) override;
 	virtual void OnEvent(uint16 RouteId, const FOnEventContext& Context) override;
 	virtual void OnAnalysisEnd() override;
@@ -36,8 +36,8 @@ private:
 
 	TSharedRef<FThreadState> GetThreadState(uint32 ThreadId);
 
-	Trace::FAnalysisSession& Session;
-	Trace::FCounterProvider& CounterProvider;
-	TMap<uint32, Trace::FCounterInternal*> CountersMap;
+	Trace::IAnalysisSession& Session;
+	Trace::ICounterProvider& CounterProvider;
+	TMap<uint32, Trace::ICounter*> CountersMap;
 	TMap<uint32, TSharedRef<FThreadState>> ThreadStatesMap;
 };
diff --git a/Engine/Source/Developer/TraceServices/Private/Common/PagedArray.h b/Engine/Source/Developer/TraceServices/Private/Common/PagedArray.h
--- a/Engine/Source/Developer/TraceServices/Private/Common/PagedArray.h
+++ b/Engine/Source/Developer/TraceServices/Private/Common/PagedArray.h
@@ -43,143 +43,198 @@ inline const ItemType* GetLastItem(const TPagedArrayPage<ItemType>& Page)
 	}
 }
 
-template<typename InItemType, typename InPageType = TPagedArrayPage<InItemType>>
-class TPagedArray
+template<typename ItemType, typename PageType>
+class TPagedArray;
+
+template<typename ItemType, typename PageType>
+class TPagedArrayIterator
 {
 public:
-	typedef InItemType ItemType;
-	typedef InPageType PageType;
-	
-	class FIterator
+	TPagedArrayIterator()
 	{
-	public:
-		uint64 GetCurrentPageIndex()
-		{
-			return CurrentPageIndex;
-		}
 
-		const PageType* GetCurrentPage()
+	}
+
+	TPagedArrayIterator(const TPagedArray<ItemType, PageType>& InOuter, uint64 InitialPageIndex, uint64 InitialItemIndex)
+		: Outer(&InOuter)
+	{
+		if (Outer->PagesArray.Num())
 		{
-			return Outer.FirstPage + CurrentPageIndex;
+			check(InitialPageIndex < Outer->PagesArray.Num());
+			CurrentPageIndex = InitialPageIndex;
+			OnCurrentPageChanged();
+			PageType* CurrentPage = Outer->FirstPage + CurrentPageIndex;
+			check(InitialItemIndex < CurrentPage->Count);
+			CurrentItem = CurrentPage->Items + InitialItemIndex;
 		}
+	}
+
+	uint64 GetCurrentPageIndex()
+	{
+		return CurrentPageIndex;
+	}
+
+	const PageType* GetCurrentPage()
+	{
+		return Outer->FirstPage + CurrentPageIndex;
+	}
+
+	const ItemType* GetCurrentItem()
+	{
+		return CurrentItem;
+	}
 
-		const ItemType* GetCurrentItem()
+	const ItemType& operator*() const
+	{
+		return *CurrentItem;
+	}
+
+	const ItemType* operator->() const
+	{
+		return CurrentItem;
+	}
+
+	explicit operator bool() const
+	{
+		return CurrentItem != nullptr;
+	}
+
+	const PageType* PrevPage()
+	{
+		if (CurrentPageIndex == 0)
 		{
-			return CurrentItem;
+			CurrentItem = nullptr;
+			CurrentPageFirstItem = nullptr;
+			CurrentPageLastItem = nullptr;
+			return nullptr;
 		}
+		--CurrentPageIndex;
+		OnCurrentPageChanged();
+		CurrentItem = CurrentPageLastItem;
+		return GetCurrentPage();
+	}
 
-		const PageType* PrevPage()
+	const PageType* NextPage()
+	{
+		if (CurrentPageIndex == Outer->PagesArray.Num() - 1)
 		{
-			if (CurrentPageIndex == 0)
-			{
-				return nullptr;
-			}
-			--CurrentPageIndex;
-			OnCurrentPageChanged();
-			CurrentItem = CurrentPageLastItem;
-			return GetCurrentPage();
+			CurrentItem = nullptr;
+			CurrentPageFirstItem = nullptr;
+			CurrentPageLastItem = nullptr;
+			return nullptr;
 		}
+		++CurrentPageIndex;
+		OnCurrentPageChanged();
+		CurrentItem = CurrentPageFirstItem;
+		return GetCurrentPage();
+	}
 
-		const PageType* NextPage()
+	const ItemType* NextItem()
+	{
+		if (CurrentItem == CurrentPageLastItem)
 		{
-			if (CurrentPageIndex == Outer.PagesArray.Num() - 1)
+			if (!NextPage())
 			{
 				return nullptr;
 			}
-			++CurrentPageIndex;
-			OnCurrentPageChanged();
-			CurrentItem = CurrentPageFirstItem;
-			return GetCurrentPage();
-		}
-
-		const ItemType* NextItem()
-		{
-			if (CurrentItem == CurrentPageLastItem)
+			else
 			{
-				if (!NextPage())
-				{
-					return nullptr;
-				}
-				else
-				{
-					return CurrentItem;
-				}
+				return CurrentItem;
 			}
-			++CurrentItem;
-			return CurrentItem;
 		}
+		++CurrentItem;
+		return CurrentItem;
+	}
+
+	TPagedArrayIterator& operator++()
+	{
+		NextItem();
+		return *this;
+	}
+
+	TPagedArrayIterator operator++(int)
+	{
+		TPagedArrayIterator Tmp(*this);
+		NextItem();
+		return Tmp;
+	}
 
-		const ItemType* PrevItem()
+	const ItemType* PrevItem()
+	{
+		if (CurrentItem == CurrentPageFirstItem)
 		{
-			if (CurrentItem == CurrentPageFirstItem)
+			if (!PrevPage())
 			{
-				if (!PrevPage())
-				{
-					return nullptr;
-				}
-				else
-				{
-					return CurrentItem;
-				}
+				return nullptr;
 			}
-			--CurrentItem;
-			return CurrentItem;
-		}
-
-		const ItemType* SetPosition(uint64 Index)
-		{
-			uint64 PageIndex = Index / Outer.PageSize;
-			uint64 ItemIndexInPage = Index % Outer.PageSize;
-			if (PageIndex != CurrentPageIndex)
+			else
 			{
-				check(PageIndex < Outer.PagesArray.Num());
-				CurrentPageIndex = PageIndex;
-				OnCurrentPageChanged();
+				return CurrentItem;
 			}
-			PageType* CurrentPage = Outer.FirstPage + CurrentPageIndex;
-			check(ItemIndexInPage < CurrentPage->Count);
-			CurrentItem = CurrentPage->Items + ItemIndexInPage;
-			return CurrentItem;
 		}
+		--CurrentItem;
+		return CurrentItem;
+	}
 
-	private:
-		friend class TPagedArray;
+	TPagedArrayIterator& operator--()
+	{
+		PrevItem();
+		return *this;
+	}
 
-		FIterator(const TPagedArray& InOuter, uint64 InitialPageIndex, uint64 InitialItemIndex)
-			: Outer(InOuter)
+	TPagedArrayIterator operator--(int)
+	{
+		TPagedArrayIterator Tmp(*this);
+		PrevItem();
+		return Tmp;
+	}
+
+	const ItemType* SetPosition(uint64 Index)
+	{
+		uint64 PageIndex = Index / Outer->PageSize;
+		uint64 ItemIndexInPage = Index % Outer->PageSize;
+		if (PageIndex != CurrentPageIndex)
 		{
-			if (Outer.PagesArray.Num())
-			{
-				check(InitialPageIndex < Outer.PagesArray.Num());
-				CurrentPageIndex = InitialPageIndex;
-				OnCurrentPageChanged();
-				PageType* CurrentPage = Outer.FirstPage + CurrentPageIndex;
-				check(InitialItemIndex < CurrentPage->Count);
-				CurrentItem = CurrentPage->Items + InitialItemIndex;
-			}
+			check(PageIndex < Outer->PagesArray.Num());
+			CurrentPageIndex = PageIndex;
+			OnCurrentPageChanged();
 		}
+		PageType* CurrentPage = Outer->FirstPage + CurrentPageIndex;
+		check(ItemIndexInPage < CurrentPage->Count);
+		CurrentItem = CurrentPage->Items + ItemIndexInPage;
+		return CurrentItem;
+	}
 
-		void OnCurrentPageChanged()
+private:
+	void OnCurrentPageChanged()
+	{
+		PageType* CurrentPage = Outer->FirstPage + CurrentPageIndex;
+		CurrentPageFirstItem = CurrentPage->Items;
+		if (CurrentPage->Items)
 		{
-			PageType* CurrentPage = Outer.FirstPage + CurrentPageIndex;
-			CurrentPageFirstItem = CurrentPage->Items;
-			if (CurrentPage->Items)
-			{
-				CurrentPageLastItem = CurrentPage->Items + CurrentPage->Count - 1;
-			}
-			else
-			{
-				CurrentPageLastItem = nullptr;
-			}
+			CurrentPageLastItem = CurrentPage->Items + CurrentPage->Count - 1;
+		}
+		else
+		{
+			CurrentPageLastItem = nullptr;
 		}
+	}
 
-		const TPagedArray& Outer;
-		uint64 CurrentPageIndex = 0;
-		const ItemType* CurrentItem = nullptr;
-		const ItemType* CurrentPageFirstItem = nullptr;
-		const ItemType* CurrentPageLastItem = nullptr;
-	};
+	const TPagedArray<ItemType, PageType>* Outer = nullptr;
+	const ItemType* CurrentItem = nullptr;
+	const ItemType* CurrentPageFirstItem = nullptr;
+	const ItemType* CurrentPageLastItem = nullptr;
+	uint64 CurrentPageIndex = 0;
+};
 
+template<typename InItemType, typename InPageType = TPagedArrayPage<InItemType>>
+class TPagedArray
+{
+public:
+	typedef InItemType ItemType;
+	typedef InPageType PageType;
+	typedef TPagedArrayIterator<InItemType, InPageType> TIterator;
+	
 	TPagedArray(Trace::ILinearAllocator& InAllocator, uint64 InPageSize)
 		: Allocator(InAllocator)
 		, PageSize(InPageSize)
@@ -266,16 +321,21 @@ public:
 		return FirstPage + PageIndex;
 	}
 
-	FIterator GetIteratorFromPage(uint64 PageIndex) const
+	TIterator GetIterator() const
 	{
-		return FIterator(*this, PageIndex, 0);
+		return TIterator(*this, 0, 0);
 	}
 
-	FIterator GetIteratorFromItem(uint64 ItemIndex) const
+	TIterator GetIteratorFromPage(uint64 PageIndex) const
+	{
+		return TIterator(*this, PageIndex, 0);
+	}
+
+	TIterator GetIteratorFromItem(uint64 ItemIndex) const
 	{
 		uint64 PageIndex = ItemIndex / PageSize;
 		uint64 IndexInPage = ItemIndex % PageSize;
-		return FIterator(*this, PageIndex, IndexInPage);
+		return TIterator(*this, PageIndex, IndexInPage);
 	}
 
 	const PageType* GetPages() const
@@ -298,6 +358,9 @@ public:
 	}
 
 private:
+	template<typename ItemType, typename PageType>
+	friend class TPagedArrayIterator;
+
 	Trace::ILinearAllocator& Allocator;
 	TArray<PageType> PagesArray;
 	PageType* FirstPage = nullptr;
diff --git a/Engine/Source/Developer/TraceServices/Private/Model/Counters.cpp b/Engine/Source/Developer/TraceServices/Private/Model/Counters.cpp
--- a/Engine/Source/Developer/TraceServices/Private/Model/Counters.cpp
+++ b/Engine/Source/Developer/TraceServices/Private/Model/Counters.cpp
@@ -1,188 +1,185 @@
 // Copyright 1998-2019 Epic Games, Inc. All Rights Reserved.
 
-#include "Model/Counters.h"
+#include "TraceServices/Model/Counters.h"
+#include "Model/CountersPrivate.h"
 #include "AnalysisServicePrivate.h"
 
 namespace Trace
 {
 
-FCounterInternal::FCounterInternal(ILinearAllocator& Allocator, uint32 InId, const TCHAR* InName, const TCHAR* InDescription, ECounterDisplayHint InDisplayHint)
-	: Name(InName)
-	, Description(InDescription)
+const FName FCounterProvider::ProviderName("CounterProvider");
+
+FCounter::FCounter(ILinearAllocator& Allocator, const TPagedArray<double>& InFrameStartTimes, uint32 InId)
+	: FrameStartTimes(InFrameStartTimes)
+	, IntCounterData(Allocator)
+	, DoubleCounterData(Allocator)
 	, Id(InId)
-	, DisplayHint(InDisplayHint)
-	, Timestamps(Allocator, 1024)
-	, Ops(Allocator, 1024)
-	, IntOpArguments(Allocator, 1024)
-	, DoubleOpArguments(Allocator, 1024)
 {
 
 }
 
 template<typename CounterType, typename EnumerationType>
-static void EnumerateCounterValuesInternal(const TPagedArray<double>& Timestamps, const TPagedArray<ECounterOpType>& Ops, const TPagedArray<CounterType>& OpArgs, double IntervalStart, double IntervalEnd, TFunctionRef<void(double, EnumerationType)> Callback)
-{
-	CounterType Value = 0;
-	auto TimeIterator = Timestamps.GetIteratorFromItem(0);
-	auto OpIterator = Ops.GetIteratorFromItem(0);
-	auto OpArgIterator = OpArgs.GetIteratorFromItem(0);
-	const double* Time = TimeIterator.GetCurrentItem();
-	const ECounterOpType* Op = OpIterator.GetCurrentItem();
-	const CounterType* OpArg = OpArgIterator.GetCurrentItem();
-	
-	double PrevTime = 0.0;
-	CounterType PrevValue = CounterType();
-	bool IsFirstValue = true;
-	bool IsFirstIncludedValue = true;
-	while (Time)
-	{
-		switch (*Op)
+static void EnumerateCounterValuesInternal(const TCounterData<CounterType>& CounterData, const TPagedArray<double>& FrameStartTimes, bool bResetEveryFrame, double IntervalStart, double IntervalEnd, TFunctionRef<void(double, EnumerationType)> Callback)
+{
+	auto CounterIterator = bResetEveryFrame ? CounterData.GetIterator(FrameStartTimes) : CounterData.GetIterator();
+	while (CounterIterator)
+	{
+		const TTuple<double, CounterType>& Current = *CounterIterator;
+		double Time = Current.template Get<0>();
+		CounterType CounterValue = Current.template Get<1>();
+		if (Time > IntervalEnd)
 		{
-		case CounterOpType_Add:
-			Value += *OpArg;
-			break;
-		case CounterOpType_Set:
-			Value = *OpArg;
 			break;
 		}
-
-		if (*Time >= IntervalStart)
+		if (Time >= IntervalStart)
 		{
-			if (IsFirstIncludedValue && !IsFirstValue)
-			{
-				Callback(PrevTime, PrevValue);
-			}
-			PrevTime = *Time;
-			PrevValue = EnumerationType(Value);
-			Callback(PrevTime, PrevValue);
-			IsFirstIncludedValue = false;
-		}
-		if (*Time > IntervalEnd)
-		{
-			break;
+			Callback(Time, CounterValue);
 		}
+		++CounterIterator;
+	}
+}
 
-		IsFirstValue = false;
+void FCounter::SetIsFloatingPoint(bool bInIsFloatingPoint)
+{
+	check(!ModCount);
+	bIsFloatingPoint = bInIsFloatingPoint;
+}
 
-		Time = TimeIterator.NextItem();
-		Op = OpIterator.NextItem();
-		OpArg = OpArgIterator.NextItem();
+void FCounter::EnumerateValues(double IntervalStart, double IntervalEnd, TFunctionRef<void(double, int64)> Callback) const
+{
+	if (bIsFloatingPoint)
+	{
+		EnumerateCounterValuesInternal<double, int64>(DoubleCounterData, FrameStartTimes, bIsResetEveryFrame, IntervalStart, IntervalEnd, Callback);
+	}
+	else
+	{
+		EnumerateCounterValuesInternal<int64, int64>(IntCounterData, FrameStartTimes, bIsResetEveryFrame, IntervalStart, IntervalEnd, Callback);
 	}
 }
 
-void FCounterInternal::EnumerateValues(double IntervalStart, double IntervalEnd, TFunctionRef<void(double, int64)> Callback) const
+void FCounter::EnumerateFloatValues(double IntervalStart, double IntervalEnd, TFunctionRef<void(double, double)> Callback) const
 {
-	if (DisplayHint == CounterDisplayHint_FloatingPoint)
+	if (bIsFloatingPoint)
 	{
-		EnumerateCounterValuesInternal<double, int64>(Timestamps, Ops, DoubleOpArguments, IntervalStart, IntervalEnd, Callback);
+		EnumerateCounterValuesInternal<double, double>(DoubleCounterData, FrameStartTimes, bIsResetEveryFrame, IntervalStart, IntervalEnd, Callback);
 	}
 	else
 	{
-		EnumerateCounterValuesInternal<int64, int64>(Timestamps, Ops, IntOpArguments, IntervalStart, IntervalEnd, Callback);
+		EnumerateCounterValuesInternal<int64, double>(IntCounterData, FrameStartTimes, bIsResetEveryFrame, IntervalStart, IntervalEnd, Callback);
 	}
 }
 
-void FCounterInternal::EnumerateFloatValues(double IntervalStart, double IntervalEnd, TFunctionRef<void(double, double)> Callback) const
+void FCounter::AddValue(double Time, int64 Value)
 {
-	if (DisplayHint == CounterDisplayHint_FloatingPoint)
+	if (bIsFloatingPoint)
 	{
-		EnumerateCounterValuesInternal<double, double>(Timestamps, Ops, DoubleOpArguments, IntervalStart, IntervalEnd, Callback);
+		DoubleCounterData.InsertOp(Time, CounterOpType_Add, double(Value));
 	}
 	else
 	{
-		EnumerateCounterValuesInternal<int64, double>(Timestamps, Ops, IntOpArguments, IntervalStart, IntervalEnd, Callback);
+		IntCounterData.InsertOp(Time, CounterOpType_Add, Value);
 	}
+	++ModCount;
 }
 
-FCounterProvider::FCounterProvider(IAnalysisSession& InSession)
-	: Session(InSession)
+void FCounter::AddValue(double Time, double Value)
 {
-
+	if (bIsFloatingPoint)
+	{
+		DoubleCounterData.InsertOp(Time, CounterOpType_Add, Value);
+	}
+	else
+	{
+		IntCounterData.InsertOp(Time, CounterOpType_Add, int64(Value));
+	}
+	++ModCount;
 }
 
-FCounterProvider::~FCounterProvider()
+void FCounter::SetValue(double Time, int64 Value)
 {
-	for (FCounterInternal* Counter : Counters)
+	if (bIsFloatingPoint)
 	{
-		delete Counter;
+		DoubleCounterData.InsertOp(Time, CounterOpType_Set, double(Value));
+	}
+	else
+	{
+		IntCounterData.InsertOp(Time, CounterOpType_Set, Value);
 	}
+	++ModCount;
 }
 
-void FCounterProvider::EnumerateCounters(TFunctionRef<void(const ICounter &)> Callback) const
+void FCounter::SetValue(double Time, double Value)
 {
-	uint32 Id = 0;
-	for (FCounterInternal* Counter : Counters)
+	if (bIsFloatingPoint)
 	{
-		Callback(*Counter);
+		DoubleCounterData.InsertOp(Time, CounterOpType_Set, Value);
 	}
+	else
+	{
+		IntCounterData.InsertOp(Time, CounterOpType_Set, int64(Value));
+	}
+	++ModCount;
 }
 
-FCounterInternal* FCounterProvider::CreateCounter(const TCHAR* Name, const TCHAR* Description, ECounterDisplayHint DisplayHint)
+FCounterProvider::FCounterProvider(IAnalysisSession& InSession, IFrameProvider& InFrameProvider)
+	: Session(InSession)
+	, FrameProvider(InFrameProvider)
+	, FrameStartTimes(InSession.GetLinearAllocator(), 1024)
 {
-	FCounterInternal* Counter = new FCounterInternal(Session.GetLinearAllocator(), Counters.Num(), Name, Description, DisplayHint);
-	Counters.Add(Counter);
-	return Counter;
+	FrameAddedDelegateHandle = FrameProvider.OnFrameAdded().AddRaw(this, &FCounterProvider::FrameAdded);
 }
 
-template<typename T>
-static void InsertOp(FCounterInternal& Counter, double Time, ECounterOpType Type, T Arg)
+FCounterProvider::~FCounterProvider()
 {
-	TPagedArray<double>& Timestamps = Counter.EditTimestamps();
-	uint64 InsertionIndex;
-	if (Timestamps.Num() == 0 || Timestamps[Timestamps.Num() - 1] <= Time)
-	{
-		InsertionIndex = Timestamps.Num();
-	}
-	else
+	FrameProvider.OnFrameAdded().Remove(FrameAddedDelegateHandle);
+	for (FCounter* Counter : Counters)
 	{
-		auto TimestampIterator = Timestamps.GetIteratorFromItem(Timestamps.Num() - 1);
-		auto CurrentPage = TimestampIterator.GetCurrentPage();
-		while (CurrentPage && *GetFirstItem(*CurrentPage) > Time)
-		{
-			CurrentPage = TimestampIterator.PrevPage();
-		}
-		if (!CurrentPage)
-		{
-			InsertionIndex = 0;
-		}
-		else
-		{
-			uint64 PageInsertionIndex = Algo::LowerBound(*CurrentPage, Time);
-			InsertionIndex = TimestampIterator.GetCurrentPageIndex() * Timestamps.GetPageSize() + PageInsertionIndex;
-		}
+		delete Counter;
 	}
-	Timestamps.Insert(InsertionIndex) = Time;
-	Counter.EditOps().Insert(InsertionIndex) = Type;
-	if (Counter.GetDisplayHint() == CounterDisplayHint_FloatingPoint)
+}
+
+void FCounterProvider::EnumerateCounters(TFunctionRef<void(const ICounter &)> Callback) const
+{
+	uint32 Id = 0;
+	for (FCounter* Counter : Counters)
 	{
-		Counter.EditDoubleOpArguments().Insert(InsertionIndex) = double(Arg);
+		Callback(*Counter);
 	}
-	else
+}
+
+bool FCounterProvider::ReadCounter(uint32 CounterId, TFunctionRef<void(const ICounter &)> Callback) const
+{
+	if (CounterId < uint32(Counters.Num()))
 	{
-		Counter.EditIntOpArguments().Insert(InsertionIndex) = int64(Arg);
+		return false;
 	}
+	Callback(*Counters[CounterId]);
+	return true;
 }
 
-void FCounterProvider::Add(FCounterInternal& Counter, double Time, int64 Value)
+ICounter* FCounterProvider::CreateCounter()
 {
-	InsertOp(Counter, Time, CounterOpType_Add, Value);
+	FCounter* Counter = new FCounter(Session.GetLinearAllocator(), FrameStartTimes, Counters.Num());
+	Counters.Add(Counter);
+	return Counter;
 }
 
-void FCounterProvider::Add(FCounterInternal& Counter, double Time, double Value)
+void FCounterProvider::FrameAdded(const FFrame& Frame)
 {
-	InsertOp(Counter, Time, CounterOpType_Add, Value);
+	if (Frame.FrameType == TraceFrameType_Game)
+	{
+		FrameStartTimes.PushBack() = Frame.StartTime;
+	}
 }
 
-void FCounterProvider::Set(FCounterInternal& Counter, double Time, int64 Value)
+const ICounterProvider& ReadCounterProvider(const IAnalysisSession& Session)
 {
-	InsertOp(Counter, Time, CounterOpType_Set, Value);
+	return *Session.ReadProvider<ICounterProvider>(FCounterProvider::ProviderName);
 }
 
-void FCounterProvider::Set(FCounterInternal& Counter, double Time, double Value)
+ICounterProvider& EditCounterProvider(IAnalysisSession& Session)
 {
-	InsertOp(Counter, Time, CounterOpType_Set, Value);
+	return *Session.EditProvider<ICounterProvider>(FCounterProvider::ProviderName);
 }
 
-
-
 }
\ No newline at end of file
diff --git a/Engine/Source/Developer/TraceServices/Private/Model/Counters.h b/Engine/Source/Developer/TraceServices/Private/Model/Counters.h
deleted file mode 100644
diff --git a/Engine/Source/Developer/TraceServices/Private/Model/CountersPrivate.h b/Engine/Source/Developer/TraceServices/Private/Model/CountersPrivate.h
new file mode 100644
--- /dev/null
+++ b/Engine/Source/Developer/TraceServices/Private/Model/CountersPrivate.h
@@ -0,0 +1,225 @@
+// Copyright 1998-2019 Epic Games, Inc. All Rights Reserved.
+
+#pragma once
+
+#include "TraceServices/Model/Counters.h"
+#include "TraceServices/Model/Frames.h"
+#include "TraceServices/Model/AnalysisSession.h"
+#include "Templates/SharedPointer.h"
+#include "Common/PagedArray.h"
+
+namespace Trace
+{
+
+enum ECounterOpType : uint8
+{
+	CounterOpType_Add,
+	CounterOpType_Set,
+};
+
+template<typename ValueType>
+class TCounterData;
+
+template<typename ValueType>
+class TCounterDataIterator
+{
+public:
+	TCounterDataIterator(const TCounterData<ValueType>& Outer)
+		: TimestampsIterator(Outer.Timestamps.GetIterator())
+		, OpTypesIterator(Outer.OpTypes.GetIterator())
+		, OpArgumentsIterator(Outer.OpArguments.GetIterator())
+	{
+		Current = MakeTuple(0.0, ValueType());
+		UpdateValue();
+	}
+
+	TCounterDataIterator(const TCounterData<ValueType>& Outer, const TPagedArray<double>& FrameStartTimes)
+		: TCounterDataIterator(Outer)
+	{
+		FrameStartTimesIterator = FrameStartTimes.GetIterator();
+	}
+
+	const TTuple<double, ValueType>& operator*() const
+	{
+		return Current;
+	}
+
+	explicit operator bool() const
+	{
+		return bool(TimestampsIterator);
+	}
+
+	TCounterDataIterator& operator++()
+	{
+		++TimestampsIterator;
+		++OpTypesIterator;
+		++OpArgumentsIterator;
+		UpdateValue();
+		return *this;
+	}
+
+	TCounterDataIterator operator++(int)
+	{
+		TCounterDataIterator Tmp(*this);
+		this->operator++();
+		return Tmp;
+	}
+
+private:
+	void UpdateValue()
+	{
+		const double* Time = TimestampsIterator.GetCurrentItem();
+		if (Time)
+		{
+			bool bIsNewFrame = false;
+			while (FrameStartTimesIterator && *FrameStartTimesIterator < *Time)
+			{
+				bIsNewFrame = true;
+				++FrameStartTimesIterator;
+			}
+			Current.template Get<0>() = *Time;
+			switch (*OpTypesIterator)
+			{
+			case CounterOpType_Add:
+				if (bIsNewFrame)
+				{
+					Current.template Get<1>() = *OpArgumentsIterator;
+				}
+				else
+				{
+					Current.template Get<1>() += *OpArgumentsIterator;
+				}
+				break;
+			case CounterOpType_Set:
+				Current.template Get<1>() = *OpArgumentsIterator;
+				break;
+			}
+		}
+	}
+
+	TPagedArray<double>::TIterator FrameStartTimesIterator;
+	TPagedArray<double>::TIterator TimestampsIterator;
+	TPagedArray<ECounterOpType>::TIterator OpTypesIterator;
+	typename TPagedArray<ValueType>::TIterator OpArgumentsIterator;
+	TTuple<double, ValueType> Current;
+};
+
+template<typename ValueType>
+class TCounterData
+{
+public:
+	typedef TCounterDataIterator<ValueType> TIterator;
+
+	TCounterData(ILinearAllocator& Allocator)
+		: Timestamps(Allocator, 1024)
+		, OpTypes(Allocator, 1024)
+		, OpArguments(Allocator, 1024)
+	{
+
+	}
+
+	void InsertOp(double Timestamp, ECounterOpType OpType, ValueType OpArgument)
+	{
+		uint64 InsertionIndex;
+		if (Timestamps.Num() == 0 || Timestamps[Timestamps.Num() - 1] <= Timestamp)
+		{
+			InsertionIndex = Timestamps.Num();
+		}
+		else if (Timestamps[0] >= Timestamp)
+		{
+			InsertionIndex = 0;
+		}
+		else
+		{
+			auto TimestampIterator = Timestamps.GetIteratorFromItem(Timestamps.Num() - 1);
+			auto CurrentPage = TimestampIterator.GetCurrentPage();
+			while (*GetFirstItem(*CurrentPage) > Timestamp)
+			{
+				CurrentPage = TimestampIterator.PrevPage();
+			}
+			uint64 PageInsertionIndex = Algo::LowerBound(*CurrentPage, Timestamp);
+			InsertionIndex = TimestampIterator.GetCurrentPageIndex() * Timestamps.GetPageSize() + PageInsertionIndex;
+		}
+		Timestamps.Insert(InsertionIndex) = Timestamp;
+		OpTypes.Insert(InsertionIndex) = OpType;
+		OpArguments.Insert(InsertionIndex) = OpArgument;
+	}
+
+	TIterator GetIterator() const
+	{
+		return TIterator(*this);
+	}
+
+	TIterator GetIterator(TPagedArray<double> FrameStartTimes) const
+	{
+		return TIterator(*this, FrameStartTimes);
+	}
+
+private:
+	template<typename IteratorValueType>
+	friend class TCounterDataIterator;
+
+	TPagedArray<double> Timestamps;
+	TPagedArray<ECounterOpType> OpTypes;
+	TPagedArray<ValueType> OpArguments;
+};
+
+class FCounter
+	: public ICounter
+{
+public:
+	FCounter(ILinearAllocator& Allocator, const TPagedArray<double>& FrameStartTimes, uint32 Id);
+	virtual uint32 GetId() const override { return Id; }
+	virtual const TCHAR* GetName() const override { return Name; }
+	virtual void SetName(const TCHAR* InName) override { Name = InName; }
+	virtual const TCHAR* GetDescription() const override { return Description; }
+	virtual void SetDescription(const TCHAR* InDescription) override { Description = InDescription; }
+	virtual bool IsFloatingPoint() const override { return bIsFloatingPoint; }
+	virtual void SetIsFloatingPoint(bool bInIsFloatingPoint) override;
+	virtual ECounterDisplayHint GetDisplayHint() const { return DisplayHint; }
+	virtual void SetDisplayHint(ECounterDisplayHint InDisplayHint) override { DisplayHint = InDisplayHint; }
+	virtual void SetIsResetEveryFrame(bool bInIsResetEveryFrame) override { bIsResetEveryFrame = bInIsResetEveryFrame; }
+	virtual void EnumerateValues(double IntervalStart, double IntervalEnd, TFunctionRef<void(double, int64)> Callback) const override;
+	virtual void EnumerateFloatValues(double IntervalStart, double IntervalEnd, TFunctionRef<void(double, double)> Callback) const override;
+	virtual void AddValue(double Time, int64 Value) override;
+	virtual void AddValue(double Time, double Value) override;
+	virtual void SetValue(double Time, int64 Value) override;
+	virtual void SetValue(double Time, double Value) override;
+
+private:
+	const TPagedArray<double>& FrameStartTimes;
+	TCounterData<int64> IntCounterData;
+	TCounterData<double> DoubleCounterData;
+	const TCHAR* Name = nullptr;
+	const TCHAR* Description = nullptr;
+	uint64 ModCount = 0;
+	uint32 Id;
+	ECounterDisplayHint DisplayHint = CounterDisplayHint_None;
+	bool bIsFloatingPoint = false;
+	bool bIsResetEveryFrame = false;
+};
+
+class FCounterProvider
+	: public ICounterProvider
+{
+public:
+	static const FName ProviderName;
+
+	FCounterProvider(IAnalysisSession& Session, IFrameProvider& FrameProvider);
+	virtual ~FCounterProvider();
+	virtual uint64 GetCounterCount() const override { return Counters.Num(); }
+	virtual void EnumerateCounters(TFunctionRef<void(const ICounter&)> Callback) const override;
+	virtual bool ReadCounter(uint32 CounterId, TFunctionRef<void(const ICounter&)> Callback) const override;
+	virtual ICounter* CreateCounter() override;
+
+private:
+	void FrameAdded(const FFrame& Frame);
+
+	IAnalysisSession& Session;
+	IFrameProvider& FrameProvider;
+	TArray<FCounter*> Counters;
+	FDelegateHandle FrameAddedDelegateHandle;
+	TPagedArray<double> FrameStartTimes;
+};
+
+}
diff --git a/Engine/Source/Developer/TraceServices/Private/Model/Frames.cpp b/Engine/Source/Developer/TraceServices/Private/Model/Frames.cpp
--- a/Engine/Source/Developer/TraceServices/Private/Model/Frames.cpp
+++ b/Engine/Source/Developer/TraceServices/Private/Model/Frames.cpp
@@ -34,12 +34,9 @@ void FFrameProvider::EnumerateFrames(ETraceFrameType FrameType, uint64 Start, ui
 	{
 		return;
 	}
-	auto Iterator = Frames[FrameType].GetIteratorFromItem(Start);
-	const FFrame* Frame = Iterator.GetCurrentItem();
-	for (uint64 Index = Start; Index < End; ++Index)
+	for (auto Iterator = Frames[FrameType].GetIteratorFromItem(Start); Iterator; ++Iterator)
 	{
-		Callback(*Frame);
-		Frame = Iterator.NextItem();
+		Callback(*Iterator);
 	}
 }
 
@@ -52,6 +49,7 @@ void FFrameProvider::BeginFrame(ETraceFrameType FrameType, double Time)
 	Frame.StartTime = Time;
 	Frame.EndTime = std::numeric_limits<double>::infinity();
 	Frame.Index = Index;
+	Frame.FrameType = FrameType;
 	Session.UpdateDurationSeconds(Time);
 }
 
@@ -60,6 +58,7 @@ void FFrameProvider::EndFrame(ETraceFrameType FrameType, double Time)
 	Session.WriteAccessCheck();
 	FFrame& Frame = Frames[FrameType][Frames[FrameType].Num() - 1];
 	Frame.EndTime = Time;
+	OnFrameAddedDelegate.Broadcast(Frame);
 	Session.UpdateDurationSeconds(Time);
 }
 
diff --git a/Engine/Source/Developer/TraceServices/Private/Model/FramesPrivate.h b/Engine/Source/Developer/TraceServices/Private/Model/FramesPrivate.h
--- a/Engine/Source/Developer/TraceServices/Private/Model/FramesPrivate.h
+++ b/Engine/Source/Developer/TraceServices/Private/Model/FramesPrivate.h
@@ -21,10 +21,12 @@ public:
 	virtual void EnumerateFrames(ETraceFrameType FrameType, uint64 Start, uint64 End, TFunctionRef<void(const FFrame&)> Callback) const override;
 	void BeginFrame(ETraceFrameType FrameType, double Time);
 	void EndFrame(ETraceFrameType FrameType, double Time);
+	virtual FOnFrameAdded& OnFrameAdded() override { return OnFrameAddedDelegate; }
 
 private:
 	IAnalysisSession& Session;
 	TArray<TPagedArray<FFrame>> Frames;
+	FOnFrameAdded OnFrameAddedDelegate;
 };
 
 }
\ No newline at end of file
diff --git a/Engine/Source/Developer/TraceServices/Private/Model/LoadTimeProfiler.cpp b/Engine/Source/Developer/TraceServices/Private/Model/LoadTimeProfiler.cpp
--- a/Engine/Source/Developer/TraceServices/Private/Model/LoadTimeProfiler.cpp
+++ b/Engine/Source/Developer/TraceServices/Private/Model/LoadTimeProfiler.cpp
@@ -23,12 +23,9 @@ void FLoadTimeProfilerProvider::EnumeratePackages(TFunctionRef<void(const FPacka
 {
 	Session.ReadAccessCheck();
 
-	auto Iterator = Packages.GetIteratorFromItem(0);
-	const FPackageInfo* Package = Iterator.GetCurrentItem();
-	while (Package)
+	for (auto Iterator = Packages.GetIteratorFromItem(0); Iterator; ++Iterator)
 	{
-		Callback(*Package);
-		Package = Iterator.NextItem();
+		Callback(*Iterator);
 	}
 }
 
diff --git a/Engine/Source/Developer/TraceServices/Private/Model/Log.cpp b/Engine/Source/Developer/TraceServices/Private/Model/Log.cpp
--- a/Engine/Source/Developer/TraceServices/Private/Model/Log.cpp
+++ b/Engine/Source/Developer/TraceServices/Private/Model/Log.cpp
@@ -136,12 +136,9 @@ void FLogProvider::ConstructMessage(uint64 Index, TFunctionRef<void(const FLogMe
 void FLogProvider::EnumerateCategories(TFunctionRef<void(const FLogCategory &)> Callback) const
 {
 	Session.ReadAccessCheck();
-	auto Iterator = Categories.GetIteratorFromItem(0);
-	const FLogCategory* Category = Iterator.GetCurrentItem();
-	while (Category)
+	for (auto Iterator = Categories.GetIteratorFromItem(0); Iterator; ++Iterator)
 	{
-		Callback(*Category);
-		Category = Iterator.NextItem();
+		Callback(*Iterator);
 	}
 }
 
diff --git a/Engine/Source/Developer/TraceServices/Private/Model/Tables.h b/Engine/Source/Developer/TraceServices/Private/Model/Tables.h
--- a/Engine/Source/Developer/TraceServices/Private/Model/Tables.h
+++ b/Engine/Source/Developer/TraceServices/Private/Model/Tables.h
@@ -386,7 +386,7 @@ public:
 
 private:
 	const LayoutType& Layout;
-	typename TPagedArray<typename LayoutType::RowType>::FIterator Iterator;
+	typename TPagedArray<typename LayoutType::RowType>::TIterator Iterator;
 	const typename LayoutType::RowType* CurrentRow;
 };
 
diff --git a/Engine/Source/Developer/TraceServices/Private/ModuleService.cpp b/Engine/Source/Developer/TraceServices/Private/ModuleService.cpp
--- a/Engine/Source/Developer/TraceServices/Private/ModuleService.cpp
+++ b/Engine/Source/Developer/TraceServices/Private/ModuleService.cpp
@@ -69,14 +69,14 @@ void FModuleService::SetModuleEnabled(const FName& ModuleName, bool bEnabled)
 	}
 }
 
-void FModuleService::OnAnalysisBegin(IAnalysisSession& Session, TArray<IAnalyzer*>& OutAnalyzers)
+void FModuleService::OnAnalysisBegin(IAnalysisSession& Session)
 {
 	FScopeLock Lock(&CriticalSection);
 	Initialize();
 	for (const auto& KV : ModulesMap)
 	{
 		IModule* Module = KV.Value;
-		Module->OnAnalysisBegin(Session, EnabledModules.Contains(Module), OutAnalyzers);
+		Module->OnAnalysisBegin(Session);
 	}
 }
 
diff --git a/Engine/Source/Developer/TraceServices/Private/ModuleServicePrivate.h b/Engine/Source/Developer/TraceServices/Private/ModuleServicePrivate.h
--- a/Engine/Source/Developer/TraceServices/Private/ModuleServicePrivate.h
+++ b/Engine/Source/Developer/TraceServices/Private/ModuleServicePrivate.h
@@ -19,7 +19,7 @@ public:
 	FModuleService();
 	virtual void GetAvailableModules(TArray<FModuleInfo>& OutModules) override;
 	virtual void SetModuleEnabled(const FName& ModuleName, bool bEnabled) override;
-	void OnAnalysisBegin(IAnalysisSession& Session, TArray<IAnalyzer*>& OutAnalyzers);
+	void OnAnalysisBegin(IAnalysisSession& Session);
 	bool GetModuleLoggers(const FName& ModuleName, TArray<const TCHAR*>& OutLoggers);
 
 private:
diff --git a/Engine/Source/Developer/TraceServices/Private/Modules/LoadTimeProfilerModule.cpp b/Engine/Source/Developer/TraceServices/Private/Modules/LoadTimeProfilerModule.cpp
--- a/Engine/Source/Developer/TraceServices/Private/Modules/LoadTimeProfilerModule.cpp
+++ b/Engine/Source/Developer/TraceServices/Private/Modules/LoadTimeProfilerModule.cpp
@@ -9,25 +9,24 @@
 namespace Trace
 {
 
-FName FLoadTimeProfilerModule::ModuleName("TraceModule_LoadTimeProfiler");
+static const FName LoadTimeProfilerModuleName("TraceModule_LoadTimeProfiler");
+static const FName LoadTimeProfilerProviderName("LoadTimeProfiler");
+static const FName FileActivityProviderName("FileActivity");
 
 void FLoadTimeProfilerModule::GetModuleInfo(FModuleInfo& OutModuleInfo)
 {
-	OutModuleInfo.Name = ModuleName;
+	OutModuleInfo.Name = LoadTimeProfilerModuleName;
 	OutModuleInfo.DisplayName = TEXT("Asset Loading");
 }
 
-static const FName LoadTimeProfilerProviderName("LoadTimeProfiler");
-static const FName FileActivityProviderName("FileActivity");
-
-void FLoadTimeProfilerModule::OnAnalysisBegin(IAnalysisSession& Session, bool bIsEnabled, TArray<IAnalyzer*>& OutAnalyzers)
+void FLoadTimeProfilerModule::OnAnalysisBegin(IAnalysisSession& Session)
 {
 	FLoadTimeProfilerProvider* LoadTimeProfilerProvider = new FLoadTimeProfilerProvider(Session);
 	Session.AddProvider(LoadTimeProfilerProviderName, LoadTimeProfilerProvider);
-	OutAnalyzers.Add(new FAsyncLoadingTraceAnalyzer(Session, *LoadTimeProfilerProvider));
+	Session.AddAnalyzer(new FAsyncLoadingTraceAnalyzer(Session, *LoadTimeProfilerProvider));
 	FFileActivityProvider* FileActivityProvider = new FFileActivityProvider(Session);
 	Session.AddProvider(FileActivityProviderName, FileActivityProvider);
-	OutAnalyzers.Add(new FPlatformFileTraceAnalyzer(Session, *FileActivityProvider));
+	Session.AddAnalyzer(new FPlatformFileTraceAnalyzer(Session, *FileActivityProvider));
 }
 
 void FLoadTimeProfilerModule::GetLoggers(TArray<const TCHAR *>& OutLoggers)
diff --git a/Engine/Source/Developer/TraceServices/Private/Modules/LoadTimeProfilerModule.h b/Engine/Source/Developer/TraceServices/Private/Modules/LoadTimeProfilerModule.h
--- a/Engine/Source/Developer/TraceServices/Private/Modules/LoadTimeProfilerModule.h
+++ b/Engine/Source/Developer/TraceServices/Private/Modules/LoadTimeProfilerModule.h
@@ -12,11 +12,8 @@ class FLoadTimeProfilerModule
 {
 public:
 	virtual void GetModuleInfo(FModuleInfo& OutModuleInfo) override;
-	virtual void OnAnalysisBegin(IAnalysisSession& Session, bool bIsEnabled, TArray<IAnalyzer*>& OutAnalyzers) override;
+	virtual void OnAnalysisBegin(IAnalysisSession& Session) override;
 	virtual void GetLoggers(TArray<const TCHAR *>& OutLoggers) override;
-
-private:
-	static FName ModuleName;
 };
 
 }
diff --git a/Engine/Source/Developer/TraceServices/Private/Modules/StatsModule.cpp b/Engine/Source/Developer/TraceServices/Private/Modules/StatsModule.cpp
--- a/Engine/Source/Developer/TraceServices/Private/Modules/StatsModule.cpp
+++ b/Engine/Source/Developer/TraceServices/Private/Modules/StatsModule.cpp
@@ -4,27 +4,22 @@
 #include "Analyzers/StatsTraceAnalysis.h"
 #include "AnalysisServicePrivate.h"
 #include "TraceServices/ModuleService.h"
-#include "Model/Counters.h"
 
 namespace Trace
 {
 
-FName FStatsModule::ModuleName("TraceModule_Stats");
+static const FName StatsModuleName("TraceModule_Stats");
 
 void FStatsModule::GetModuleInfo(FModuleInfo& OutModuleInfo)
 {
-	OutModuleInfo.Name = ModuleName;
+	OutModuleInfo.Name = StatsModuleName;
 	OutModuleInfo.DisplayName = TEXT("Stats");
 }
 	
-static const FName CounterProviderName("CounterProvider");
-
-void FStatsModule::OnAnalysisBegin(IAnalysisSession& InSession, bool bIsEnabled, TArray<IAnalyzer*>& OutAnalyzers)
+void FStatsModule::OnAnalysisBegin(IAnalysisSession& Session)
 {
-	FAnalysisSession& Session = static_cast<FAnalysisSession&>(InSession);
-	FCounterProvider* CounterProvider = new FCounterProvider(InSession);
-	InSession.AddProvider(CounterProviderName, CounterProvider);
-	OutAnalyzers.Add(new FStatsAnalyzer(Session, *CounterProvider));
+	ICounterProvider& CounterProvider = EditCounterProvider(Session);
+	Session.AddAnalyzer(new FStatsAnalyzer(Session, CounterProvider));
 }
 
 void FStatsModule::GetLoggers(TArray<const TCHAR *>& OutLoggers)
@@ -32,9 +27,4 @@ void FStatsModule::GetLoggers(TArray<const TCHAR *>& OutLoggers)
 	OutLoggers.Add(TEXT("Stats"));
 }
 
-const ICounterProvider* ReadCounterProvider(const IAnalysisSession& Session)
-{
-	return Session.ReadProvider<ICounterProvider>(CounterProviderName);
-}
-
 }
diff --git a/Engine/Source/Developer/TraceServices/Private/Modules/StatsModule.h b/Engine/Source/Developer/TraceServices/Private/Modules/StatsModule.h
--- a/Engine/Source/Developer/TraceServices/Private/Modules/StatsModule.h
+++ b/Engine/Source/Developer/TraceServices/Private/Modules/StatsModule.h
@@ -13,11 +13,8 @@ class FStatsModule
 {
 public:
 	virtual void GetModuleInfo(FModuleInfo& OutModuleInfo) override;
-	virtual void OnAnalysisBegin(IAnalysisSession& Session, bool bIsEnabled, TArray<IAnalyzer*>& OutAnalyzers) override;
+	virtual void OnAnalysisBegin(IAnalysisSession& Session) override;
 	virtual void GetLoggers(TArray<const TCHAR *>& OutLoggers) override;
-
-private:
-	static FName ModuleName;
 };
 
 }
diff --git a/Engine/Source/Developer/TraceServices/Private/Modules/TimingProfilerModule.cpp b/Engine/Source/Developer/TraceServices/Private/Modules/TimingProfilerModule.cpp
--- a/Engine/Source/Developer/TraceServices/Private/Modules/TimingProfilerModule.cpp
+++ b/Engine/Source/Developer/TraceServices/Private/Modules/TimingProfilerModule.cpp
@@ -9,24 +9,23 @@
 namespace Trace
 {
 
-FName FTimingProfilerModule::ModuleName("TraceModule_TimingProfiler");
+static const FName TimingProfilerModuleName("TraceModule_TimingProfiler");
+static const FName TimingProfilerProviderName("TimingProfilerProvider");
 
 void FTimingProfilerModule::GetModuleInfo(FModuleInfo& OutModuleInfo)
 {
-	OutModuleInfo.Name = ModuleName;
+	OutModuleInfo.Name = TimingProfilerModuleName;
 	OutModuleInfo.DisplayName = TEXT("Timing");
 }
 	
-static const FName TimingProfilerProviderName("TimingProfilerProvider");
-
-void FTimingProfilerModule::OnAnalysisBegin(IAnalysisSession& InSession, bool bIsEnabled, TArray<IAnalyzer*>& OutAnalyzers)
+void FTimingProfilerModule::OnAnalysisBegin(IAnalysisSession& InSession)
 {
 	FAnalysisSession& Session = static_cast<FAnalysisSession&>(InSession);
 	
 	FTimingProfilerProvider* TimingProfilerProvider = new FTimingProfilerProvider(Session);
 	Session.AddProvider(TimingProfilerProviderName, TimingProfilerProvider);
-	OutAnalyzers.Add(new FCpuProfilerAnalyzer(Session, *TimingProfilerProvider));
-	OutAnalyzers.Add(new FGpuProfilerAnalyzer(Session, *TimingProfilerProvider));
+	Session.AddAnalyzer(new FCpuProfilerAnalyzer(Session, *TimingProfilerProvider));
+	Session.AddAnalyzer(new FGpuProfilerAnalyzer(Session, *TimingProfilerProvider));
 }
 
 void FTimingProfilerModule::GetLoggers(TArray<const TCHAR *>& OutLoggers)
diff --git a/Engine/Source/Developer/TraceServices/Private/Modules/TimingProfilerModule.h b/Engine/Source/Developer/TraceServices/Private/Modules/TimingProfilerModule.h
--- a/Engine/Source/Developer/TraceServices/Private/Modules/TimingProfilerModule.h
+++ b/Engine/Source/Developer/TraceServices/Private/Modules/TimingProfilerModule.h
@@ -13,11 +13,8 @@ class FTimingProfilerModule
 {
 public:
 	virtual void GetModuleInfo(FModuleInfo& OutModuleInfo) override;
-	virtual void OnAnalysisBegin(IAnalysisSession& Session, bool bIsEnabled, TArray<IAnalyzer*>& OutAnalyzers) override;
+	virtual void OnAnalysisBegin(IAnalysisSession& Session) override;
 	virtual void GetLoggers(TArray<const TCHAR *>& OutLoggers) override;
-
-private:
-	static FName ModuleName;
 };
 
 }
diff --git a/Engine/Source/Developer/TraceServices/Public/TraceServices/AnalysisService.h b/Engine/Source/Developer/TraceServices/Public/TraceServices/AnalysisService.h
--- a/Engine/Source/Developer/TraceServices/Public/TraceServices/AnalysisService.h
+++ b/Engine/Source/Developer/TraceServices/Public/TraceServices/AnalysisService.h
@@ -14,7 +14,7 @@
 #include "TraceServices/Model/Threads.h"
 #include "TraceServices/Model/TimingProfiler.h"
 #include "TraceServices/Model/LoadTimeProfiler.h"
-#include "TraceServices/Model/Stats.h"
+#include "TraceServices/Model/Counters.h"
 
 namespace Trace
 {
diff --git a/Engine/Source/Developer/TraceServices/Public/TraceServices/Model/AnalysisSession.h b/Engine/Source/Developer/TraceServices/Public/TraceServices/Model/AnalysisSession.h
--- a/Engine/Source/Developer/TraceServices/Public/TraceServices/Model/AnalysisSession.h
+++ b/Engine/Source/Developer/TraceServices/Public/TraceServices/Model/AnalysisSession.h
@@ -10,6 +10,7 @@ namespace Trace
 {
 
 class ILinearAllocator;
+class IAnalyzer;
 
 class IProvider
 {
@@ -38,15 +39,23 @@ public:
 	virtual ILinearAllocator& GetLinearAllocator() = 0;
 	virtual const TCHAR* StoreString(const TCHAR* String) = 0;
 	
+	virtual void AddAnalyzer(IAnalyzer* Analyzer) = 0;
+
 	virtual void AddProvider(const FName& Name, IProvider* Provider) = 0;
 	template<typename ProviderType>
 	const ProviderType* ReadProvider(const FName& Name) const
 	{
 		return static_cast<const ProviderType*>(ReadProviderPrivate(Name));
 	}
+	template<typename ProviderType>
+	ProviderType* EditProvider(const FName& Name)
+	{
+		return static_cast<ProviderType*>(EditProviderPrivate(Name));
+	}
 
 private:
 	virtual const IProvider* ReadProviderPrivate(const FName& Name) const = 0;
+	virtual IProvider* EditProviderPrivate(const FName& Name) = 0;
 };
 
 struct FAnalysisSessionReadScope
diff --git a/Engine/Source/Developer/TraceServices/Public/TraceServices/Model/Counters.h b/Engine/Source/Developer/TraceServices/Public/TraceServices/Model/Counters.h
new file mode 100644
--- /dev/null
+++ b/Engine/Source/Developer/TraceServices/Public/TraceServices/Model/Counters.h
@@ -0,0 +1,54 @@
+// Copyright 1998-2019 Epic Games, Inc. All Rights Reserved.
+
+#pragma once
+
+#include "TraceServices/Model/AnalysisSession.h"
+#include "Templates/Function.h"
+
+namespace Trace
+{
+
+enum ECounterDisplayHint
+{
+	CounterDisplayHint_None,
+	CounterDisplayHint_Memory,
+};
+
+class ICounter
+{
+public:
+	virtual ~ICounter() = default;
+
+	virtual uint32 GetId() const = 0;
+	virtual const TCHAR* GetName() const = 0;
+	virtual void SetName(const TCHAR* Name) = 0;
+	virtual const TCHAR* GetDescription() const = 0;
+	virtual void SetDescription(const TCHAR* Description) = 0;
+	virtual bool IsFloatingPoint() const = 0;
+	virtual void SetIsFloatingPoint(bool bIsFloatingPoint) = 0;
+	virtual ECounterDisplayHint GetDisplayHint() const = 0;
+	virtual void SetDisplayHint(ECounterDisplayHint DisplayHint) = 0;
+	virtual void SetIsResetEveryFrame(bool bInIsResetEveryFrame) = 0;
+	virtual void EnumerateValues(double IntervalStart, double IntervalEnd, TFunctionRef<void(double, int64)> Callback) const = 0;
+	virtual void EnumerateFloatValues(double IntervalStart, double IntervalEnd, TFunctionRef<void(double, double)> Callback) const = 0;
+	virtual void AddValue(double Time, int64 Value) = 0;
+	virtual void AddValue(double Time, double Value) = 0;
+	virtual void SetValue(double Time, int64 Value) = 0;
+	virtual void SetValue(double Time, double Value) = 0;
+};
+
+class ICounterProvider
+	: public IProvider
+{
+public:
+	virtual ~ICounterProvider() = default;
+	virtual uint64 GetCounterCount() const = 0;
+	virtual bool ReadCounter(uint32 CounterId, TFunctionRef<void(const ICounter&)> Callback) const = 0;
+	virtual void EnumerateCounters(TFunctionRef<void(const ICounter&)> Callback) const = 0;
+	virtual ICounter* CreateCounter() = 0;
+};
+
+TRACESERVICES_API const ICounterProvider& ReadCounterProvider(const IAnalysisSession& Session);
+TRACESERVICES_API ICounterProvider& EditCounterProvider(IAnalysisSession& Session);
+
+}
diff --git a/Engine/Source/Developer/TraceServices/Public/TraceServices/Model/Frames.h b/Engine/Source/Developer/TraceServices/Public/TraceServices/Model/Frames.h
--- a/Engine/Source/Developer/TraceServices/Public/TraceServices/Model/Frames.h
+++ b/Engine/Source/Developer/TraceServices/Public/TraceServices/Model/Frames.h
@@ -5,6 +5,7 @@
 #include "TraceServices/Model/AnalysisSession.h"
 #include "ProfilingDebugging/MiscTrace.h"
 #include "Templates/Function.h"
+#include "Delegates/Delegate.h"
 
 namespace Trace
 {
@@ -14,6 +15,7 @@ struct FFrame
 	uint64 Index;
 	double StartTime;
 	double EndTime;
+	ETraceFrameType FrameType;
 };
 
 class IFrameProvider
@@ -23,6 +25,8 @@ public:
 	virtual ~IFrameProvider() = default;
 	virtual uint64 GetFrameCount(ETraceFrameType FrameType) const = 0;
 	virtual void EnumerateFrames(ETraceFrameType FrameType, uint64 Start, uint64 End, TFunctionRef<void(const FFrame&)> Callback) const = 0;
+	DECLARE_MULTICAST_DELEGATE_OneParam(FOnFrameAdded, const FFrame&);
+	virtual FOnFrameAdded& OnFrameAdded() = 0;
 };
 
 TRACESERVICES_API const IFrameProvider& ReadFrameProvider(const IAnalysisSession& Session);
diff --git a/Engine/Source/Developer/TraceServices/Public/TraceServices/Model/Stats.h b/Engine/Source/Developer/TraceServices/Public/TraceServices/Model/Stats.h
deleted file mode 100644
diff --git a/Engine/Source/Developer/TraceServices/Public/TraceServices/ModuleService.h b/Engine/Source/Developer/TraceServices/Public/TraceServices/ModuleService.h
--- a/Engine/Source/Developer/TraceServices/Public/TraceServices/ModuleService.h
+++ b/Engine/Source/Developer/TraceServices/Public/TraceServices/ModuleService.h
@@ -25,7 +25,7 @@ class IModule
 {
 public:
 	virtual void GetModuleInfo(FModuleInfo& OutModuleInfo) = 0;
-	virtual void OnAnalysisBegin(IAnalysisSession& Session, bool bIsEnabled, TArray<IAnalyzer*>& OutAnalyzers) = 0;
+	virtual void OnAnalysisBegin(IAnalysisSession& Session) = 0;
 	virtual void GetLoggers(TArray<const TCHAR*>& OutLoggers) = 0;
 };
 
