UnrealInsights:
Added request grouping concept to load time tracing
Workaround for issue where the wrong package was given the cpu cost for deferred post loading of exports
Added more tables to LoadTimeProfilerProvider
Added tracing of missing blueprint class names

diff --git a/Engine/Source/Developer/TraceServices/Private/Analyzers/LoadTimeTraceAnalysis.cpp b/Engine/Source/Developer/TraceServices/Private/Analyzers/LoadTimeTraceAnalysis.cpp
--- a/Engine/Source/Developer/TraceServices/Private/Analyzers/LoadTimeTraceAnalysis.cpp
+++ b/Engine/Source/Developer/TraceServices/Private/Analyzers/LoadTimeTraceAnalysis.cpp
@@ -6,6 +6,8 @@
 #include "Trace/Trace.h"
 #include "Model/LoadTimeProfilerPrivate.h"
 #include "Analyzers/MiscTraceAnalysis.h"
+#include "Common/FormatArgs.h"
+#include <limits>
 
 FAsyncLoadingTraceAnalyzer::FAsyncLoadingTraceAnalyzer(Trace::IAnalysisSession& InSession, Trace::FLoadTimeProfilerProvider& InLoadTimeProfilerProvider)
 	: Session(InSession)
@@ -73,22 +75,13 @@ void FAsyncLoadingTraceAnalyzer::OnAnalysisBegin(const FOnAnalysisContext& Conte
 	Builder.RouteEvent(RouteId_DestroyAsyncPackage, "LoadTime", "DestroyAsyncPackage");
 	Builder.RouteEvent(RouteId_BeginRequest, "LoadTime", "BeginRequest");
 	Builder.RouteEvent(RouteId_EndRequest, "LoadTime", "EndRequest");
-	Builder.RouteEvent(RouteId_BeginLoadMap, "LoadTime", "BeginLoadMap");
-	Builder.RouteEvent(RouteId_EndLoadMap, "LoadTime", "EndLoadMap");
-	Builder.RouteEvent(RouteId_NewStreamableHandle, "LoadTime", "NewStreamableHandle");
-	Builder.RouteEvent(RouteId_DestroyStreamableHandle, "LoadTime", "DestroyStreamableHandle");
-	Builder.RouteEvent(RouteId_BeginLoadStreamableHandle, "LoadTime", "BeginLoadStreamableHandle");
-	Builder.RouteEvent(RouteId_EndLoadStreamableHandle, "LoadTime", "EndLoadStreamableHandle");
-	Builder.RouteEvent(RouteId_BeginWaitForStreamableHandle, "LoadTime", "BeginWaitForStreamableHandle");
-	Builder.RouteEvent(RouteId_EndWaitForStreamableHandle, "LoadTime", "EndWaitForStreamableHandle");
-	Builder.RouteEvent(RouteId_StreamableHandleRequestAssociation, "LoadTime", "StreamableHandleRequestAssociation");
+	Builder.RouteEvent(RouteId_BeginRequestGroup, "LoadTime", "BeginRequestGroup");
+	Builder.RouteEvent(RouteId_EndRequestGroup, "LoadTime", "EndRequestGroup");
 	Builder.RouteEvent(RouteId_AsyncPackageRequestAssociation, "LoadTime", "AsyncPackageRequestAssociation");
 	Builder.RouteEvent(RouteId_AsyncPackageLinkerAssociation, "LoadTime", "AsyncPackageLinkerAssociation");
-	Builder.RouteEvent(RouteId_LinkerArchiveAssociation, "LoadTime", "LinkerArchiveAssociation");
+	Builder.RouteEvent(RouteId_AsyncPackageImportDependency, "LoadTime", "AsyncPackageImportDependency");
 	Builder.RouteEvent(RouteId_BeginAsyncPackageScope, "LoadTime", "BeginAsyncPackageScope");
 	Builder.RouteEvent(RouteId_EndAsyncPackageScope, "LoadTime", "EndAsyncPackageScope");
-	Builder.RouteEvent(RouteId_BeginFlushAsyncLoading, "LoadTime", "BeginFlushAsyncLoading");
-	Builder.RouteEvent(RouteId_EndFlushAsyncLoading, "LoadTime", "EndFlushAsyncLoading");
 	Builder.RouteEvent(RouteId_ClassInfo, "LoadTime", "ClassInfo");
 }
 
@@ -139,6 +132,28 @@ void FAsyncLoadingTraceAnalyzer::OnEvent(uint16 RouteId, const FOnEventContext&
 		}
 		break;
 	}
+	case RouteId_BeginAsyncPackageScope:
+	{
+		Trace::FAnalysisSessionEditScope _(Session);
+
+		uint64 AsyncPackagePtr = EventData.GetValue<uint64>("AsyncPackage");
+		TSharedRef<FAsyncPackageState>* AsyncPackageState = ActiveAsyncPackagesMap.Find(AsyncPackagePtr);
+		uint32 ThreadId = EventData.GetValue<uint32>("ThreadId");
+		TSharedRef<FThreadState> ThreadState = GetThreadState(ThreadId);
+		ELoadTimeProfilerPackageEventType EventType = static_cast<ELoadTimeProfilerPackageEventType>(EventData.GetValue<uint8>("EventType"));
+		ThreadState->EnterPackageScope(Context.SessionContext.TimestampFromCycle(EventData.GetValue<uint64>("Cycle")), AsyncPackageState ? (*AsyncPackageState)->PackageInfo : nullptr, EventType);
+		break;
+	}
+	case RouteId_EndAsyncPackageScope:
+	{
+		Trace::FAnalysisSessionEditScope _(Session);
+
+		uint32 ThreadId = EventData.GetValue<uint32>("ThreadId");
+		TSharedRef<FThreadState> ThreadState = GetThreadState(ThreadId);
+		uint64 Cycle = EventData.GetValue<uint64>("Cycle");
+		ThreadState->LeaveScope(Context.SessionContext.TimestampFromCycle(EventData.GetValue<uint64>("Cycle")));
+		break;
+	}
 	case RouteId_BeginCreateExport:
 	{
 		Trace::FAnalysisSessionEditScope _(Session);
@@ -155,6 +170,7 @@ void FAsyncLoadingTraceAnalyzer::OnEvent(uint16 RouteId, const FOnEventContext&
 		if (LinkerState && (*LinkerState)->PackageInfo)
 		{
 			(*LinkerState)->PackageInfo->Exports.Add(&Export);
+			Export.Package = (*LinkerState)->PackageInfo;
 		}
 		break;
 	}
@@ -201,139 +217,71 @@ void FAsyncLoadingTraceAnalyzer::OnEvent(uint16 RouteId, const FOnEventContext&
 	case RouteId_BeginRequest:
 	{
 		uint64 RequestId = EventData.GetValue<uint64>("RequestId");
+		uint32 ThreadId = EventData.GetValue<uint32>("ThreadId");
 		//check(!ActiveRequestsMap.Contains(RequestId));
 		TSharedRef<FRequestState> RequestState = MakeShared<FRequestState>();
-		RequestState->Id = RequestId;
 		RequestState->WallTimeStartCycle = EventData.GetValue<uint64>("Cycle");
 		RequestState->WallTimeEndCycle = 0;
-		Requests.Add(RequestState);
+		RequestState->ThreadId = ThreadId;
 		ActiveRequestsMap.Add(RequestId, RequestState);
-		uint32 ThreadId = EventData.GetValue<uint32>("ThreadId");
 		TSharedRef<FThreadState> ThreadState = GetThreadState(ThreadId);
-		if (ThreadState->ActiveLoadMap)
+		TSharedPtr<FRequestGroupState> RequestGroup = ThreadState->RequestGroupStack.Num() ? ThreadState->RequestGroupStack.Top() : nullptr;
+		if (!RequestGroup)
 		{
-			ThreadState->ActiveLoadMap->Requests.Add(RequestState);
+			RequestGroup = MakeShared<FRequestGroupState>();
+			RequestGroup->Name = TEXT("[ungrouped]");
+			RequestGroup->bIsClosed = true;
 		}
-		
-		{
-			/*Trace::FAnalysisSessionEditScope _(Session.Get());
-			RequestTimelineEventTypeMap.Add(RequestState->Id, TimelineProvider->AddEventType(*RequestState->PackageName, 0));*/
-		}
-		
+		RequestGroup->Requests.Add(RequestState);
+		++RequestGroup->ActiveRequestsCount;
+		RequestState->Group = RequestGroup;
 		break;
 	}
 	case RouteId_EndRequest:
 	{
 		uint64 RequestId = EventData.GetValue<uint64>("RequestId");
-		TSharedRef<FRequestState>* RequestState = ActiveRequestsMap.Find(RequestId);
-		if (RequestState)
+		TSharedRef<FRequestState>* FindRequestState = ActiveRequestsMap.Find(RequestId);
+		if (FindRequestState)
 		{
-			(*RequestState)->WallTimeEndCycle = EventData.GetValue<uint64>("Cycle");
+			TSharedRef<FRequestState> RequestState = *FindRequestState;
+			RequestState->WallTimeEndCycle = EventData.GetValue<uint64>("Cycle");
+			RequestState->Group->LatestEndCycle = FMath::Max(RequestState->Group->LatestEndCycle, RequestState->WallTimeEndCycle);
+			--RequestState->Group->ActiveRequestsCount;
+			if (RequestState->Group->LoadRequest && RequestState->Group->bIsClosed && RequestState->Group->ActiveRequestsCount == 0)
+			{
+				Trace::FAnalysisSessionEditScope _(Session);
+				RequestState->Group->LoadRequest->EndTime = Context.SessionContext.TimestampFromCycle(RequestState->Group->LatestEndCycle);
+			}
 		}
 		break;
 	}
-	case RouteId_BeginLoadMap:
+	case RouteId_BeginRequestGroup:
 	{
-		TSharedRef<FLoadMapState> LoadMapState = MakeShared<FLoadMapState>();
-		LoadMapState->Id = NextMapId++;
-		LoadMapState->WallTimeStartCycle = EventData.GetValue<uint64>("Cycle");
-		LoadMapState->WallTimeEndCycle = 0;
-		LoadMapState->Name = reinterpret_cast<const TCHAR*>(EventData.GetAttachment());
-		Maps.Add(LoadMapState);
+		TSharedRef<FRequestGroupState> GroupState = MakeShared<FRequestGroupState>();
+		const TCHAR* FormatString = reinterpret_cast<const TCHAR*>(EventData.GetAttachment());
+		const uint8* FormatArgs = EventData.GetAttachment() + (FCString::Strlen(FormatString) + 1) * sizeof(TCHAR);
+		Trace::FFormatArgsHelper::Format(FormatBuffer, FormatBufferSize - 1, FormatString, FormatArgs);
+		GroupState->Name = Session.StoreString(FormatBuffer);
 		uint32 ThreadId = EventData.GetValue<uint32>("ThreadId");
 		TSharedRef<FThreadState> ThreadState = GetThreadState(ThreadId);
-		ThreadState->ActiveLoadMap = LoadMapState;
+		ThreadState->RequestGroupStack.Push(GroupState);
 		break;
 	}
-	case RouteId_EndLoadMap:
+	case RouteId_EndRequestGroup:
 	{
 		uint32 ThreadId = EventData.GetValue<uint32>("ThreadId");
 		TSharedRef<FThreadState> ThreadState = GetThreadState(ThreadId);
-		if (ThreadState->ActiveLoadMap)
+		if (ThreadState->RequestGroupStack.Num())
 		{
-			ThreadState->ActiveLoadMap->WallTimeEndCycle = EventData.GetValue<uint64>("Cycle");
-			ThreadState->ActiveLoadMap = nullptr;
+			TSharedPtr<FRequestGroupState> GroupState = ThreadState->RequestGroupStack.Pop(false);
+			GroupState->bIsClosed = true;
+			if (GroupState->LoadRequest && GroupState->ActiveRequestsCount == 0)
+			{
+				GroupState->LoadRequest->EndTime = Context.SessionContext.TimestampFromCycle(GroupState->LatestEndCycle);
+			}
 		}
 		break;
 	}
-	case RouteId_NewStreamableHandle:
-	{
-		uint64 Ptr = EventData.GetValue<uint64>("StreamableHandle");
-		//check(!ActiveStreamableHandlesMap.Contains(Ptr));
-		TSharedRef<FStreamableHandleState> StreamableHandleState = MakeShared<FStreamableHandleState>();
-		StreamableHandleState->Id = NextStreamableHandleId++;
-		StreamableHandleState->DebugName = FString((const char*)EventData.GetAttachment(), EventData.GetValue<uint16>("DebugNameSize"));
-		StreamableHandleState->WallTimeStartCycle = 0;
-		StreamableHandleState->WallTimeEndCycle = 0;
-		StreamableHandles.Add(StreamableHandleState);
-		ActiveStreamableHandlesMap.Add(Ptr, StreamableHandleState);
-		
-		{
-			/*Trace::FAnalysisSessionEditScope _(Session.Get());
-			StreamableHandlesTimelineEventTypeMap.Add(StreamableHandleState->Id, TimelineProvider->AddEventType(*StreamableHandleState->DebugName, 0));*/
-		}
-		
-		break;
-	}
-	case RouteId_DestroyStreamableHandle:
-	{
-		uint64 Ptr = EventData.GetValue<uint64>("StreamableHandle");
-		ActiveStreamableHandlesMap.Remove(Ptr);
-		break;
-	}
-	case RouteId_BeginLoadStreamableHandle:
-	{
-		uint64 Ptr = EventData.GetValue<uint64>("StreamableHandle");
-		TSharedRef<FStreamableHandleState>* StreamableHandleState = ActiveStreamableHandlesMap.Find(Ptr);
-		if (StreamableHandleState)
-		{
-			(*StreamableHandleState)->WallTimeStartCycle = EventData.GetValue<uint64>("Cycle");
-		}
-		break;
-	}
-	case RouteId_EndLoadStreamableHandle:
-	{
-		uint64 Ptr = EventData.GetValue<uint64>("StreamableHandle");
-		TSharedRef<FStreamableHandleState>* StreamableHandleState = ActiveStreamableHandlesMap.Find(Ptr);
-		if (StreamableHandleState)
-		{
-			(*StreamableHandleState)->WallTimeEndCycle = EventData.GetValue<uint64>("Cycle");
-		}
-		break;
-	}
-	case RouteId_BeginWaitForStreamableHandle:
-	{
-		uint64 Ptr = EventData.GetValue<uint64>("StreamableHandle");
-		TSharedRef<FStreamableHandleState>* StreamableHandleState = ActiveStreamableHandlesMap.Find(Ptr);
-		if (StreamableHandleState)
-		{
-			uint32 ThreadId = EventData.GetValue<uint32>("ThreadId");
-			TSharedRef<FThreadState> ThreadState = GetThreadState(ThreadId);
-			ThreadState->WaitForStreamableHandleHandle = *StreamableHandleState;
-			ThreadState->WaitForStreamableHandleStartCycle = EventData.GetValue<uint64>("Cycle");
-		}
-		/*uint64 EventTypeId = StreamableHandlesTimelineEventTypeMap[ThreadState->WaitForStreamableHandleHandle->Id];
-
-		{
-			Trace::FAnalysisSessionEditScope _(Session.Get());
-			BlockingRequestsTimeline->AppendBeginEvent(Context.SessionContext.TimestampFromCycle(ThreadState->WaitForStreamableHandleStartCycle), EventTypeId);
-		}*/
-
-		break;
-	}
-	case RouteId_EndWaitForStreamableHandle:
-	{
-		uint32 ThreadId = EventData.GetValue<uint32>("ThreadId");
-		TSharedRef<FThreadState> ThreadState = GetThreadState(ThreadId);
-		ThreadState->WaitForStreamableHandleHandle = nullptr;
-
-		/*{
-			Trace::FAnalysisSessionEditScope _(Session.Get());
-			BlockingRequestsTimeline->AppendEndEvent(Context.SessionContext.TimestampFromCycle(EventData.GetValue<uint64>("Cycle")));
-		}*/
-
-		break;
-	}
 	case RouteId_NewAsyncPackage:
 	{
 		Trace::FAnalysisSessionEditScope _(Session);
@@ -352,29 +300,30 @@ void FAsyncLoadingTraceAnalyzer::OnEvent(uint16 RouteId, const FOnEventContext&
 		ActiveAsyncPackagesMap.Remove(AsyncPackagePtr);
 		break;
 	}
-
-	case RouteId_StreamableHandleRequestAssociation:
+	case RouteId_AsyncPackageImportDependency:
 	{
-		uint64 StreamableHandlePtr = EventData.GetValue<uint64>("StreamableHandle");
-		TSharedRef<FStreamableHandleState>* StreamableHandleState = ActiveStreamableHandlesMap.Find(StreamableHandlePtr);
-		uint64 RequestId = EventData.GetValue<uint64>("RequestId");
-		TSharedRef<FRequestState>* RequestState = ActiveRequestsMap.Find(RequestId);
-		if (StreamableHandleState && RequestState)
+		uint64 AsyncPackagePtr = EventData.GetValue<uint64>("AsyncPackage");
+		TSharedRef<FAsyncPackageState>* FindAsyncPackage = ActiveAsyncPackagesMap.Find(AsyncPackagePtr);
+		uint64 ImportedAsyncPackagePtr = EventData.GetValue<uint64>("ImportedAsyncPackage");
+		TSharedRef<FAsyncPackageState>* FindImportedAsyncPackage = ActiveAsyncPackagesMap.Find(ImportedAsyncPackagePtr);
+		if (FindAsyncPackage && FindImportedAsyncPackage)
 		{
-			(*StreamableHandleState)->Requests.Add(*RequestState);
+			TSharedRef<FAsyncPackageState> AsyncPackageState = *FindAsyncPackage;
+			TSharedRef<FAsyncPackageState> ImportedAsyncPackageState = *FindImportedAsyncPackage;
+			check(AsyncPackageState->Request);
+			PackageRequestAssociation(Context, ImportedAsyncPackageState, AsyncPackageState->Request.ToSharedRef());
 		}
 		break;
 	}
 	case RouteId_AsyncPackageRequestAssociation:
 	{
 		uint64 AsyncPackagePtr = EventData.GetValue<uint64>("AsyncPackage");
-		TSharedRef<FAsyncPackageState>* AsyncPackageState = ActiveAsyncPackagesMap.Find(AsyncPackagePtr);
+		TSharedRef<FAsyncPackageState>* FindAsyncPackageState = ActiveAsyncPackagesMap.Find(AsyncPackagePtr);
 		uint64 RequestId = EventData.GetValue<uint64>("RequestId");
-		TSharedRef<FRequestState>* RequestState = ActiveRequestsMap.Find(RequestId);
-		if (AsyncPackageState && RequestState)
+		TSharedRef<FRequestState>* FindRequestState = ActiveRequestsMap.Find(RequestId);
+		if (FindAsyncPackageState && FindRequestState)
 		{
-			(*RequestState)->AsyncPackages.Add(*AsyncPackageState);
-			(*AsyncPackageState)->Requests.Add(*RequestState);
+			PackageRequestAssociation(Context, *FindAsyncPackageState, *FindRequestState);
 		}
 		break;
 	}
@@ -392,58 +341,6 @@ void FAsyncLoadingTraceAnalyzer::OnEvent(uint16 RouteId, const FOnEventContext&
 		}
 		break;
 	}
-	case RouteId_LinkerArchiveAssociation:
-	{
-		uint64 ArchivePtr = EventData.GetValue<uint64>("Archive");
-		uint64 LinkerPtr = EventData.GetValue<uint64>("Linker");
-		break;
-	}
-	case RouteId_BeginAsyncPackageScope:
-	{
-		Trace::FAnalysisSessionEditScope _(Session);
-
-		uint64 AsyncPackagePtr = EventData.GetValue<uint64>("AsyncPackage");
-		TSharedRef<FAsyncPackageState>* AsyncPackageState = ActiveAsyncPackagesMap.Find(AsyncPackagePtr);
-		uint32 ThreadId = EventData.GetValue<uint32>("ThreadId");
-		TSharedRef<FThreadState> ThreadState = GetThreadState(ThreadId);
-		ELoadTimeProfilerPackageEventType EventType = static_cast<ELoadTimeProfilerPackageEventType>(EventData.GetValue<uint8>("EventType"));
-		ThreadState->EnterPackageScope(Context.SessionContext.TimestampFromCycle(EventData.GetValue<uint64>("Cycle")), AsyncPackageState ? (*AsyncPackageState)->PackageInfo : nullptr, EventType);
-		break;
-	}
-	case RouteId_EndAsyncPackageScope:
-	{
-		Trace::FAnalysisSessionEditScope _(Session);
-
-		uint32 ThreadId = EventData.GetValue<uint32>("ThreadId");
-		TSharedRef<FThreadState> ThreadState = GetThreadState(ThreadId);
-		uint64 Cycle = EventData.GetValue<uint64>("Cycle");
-		ThreadState->LeaveScope(Context.SessionContext.TimestampFromCycle(EventData.GetValue<uint64>("Cycle")));
-		break;
-	}
-	case RouteId_BeginFlushAsyncLoading:
-	{
-		Trace::FAnalysisSessionEditScope _(Session);
-
-		FlushAsyncLoadingRequestId = EventData.GetValue<uint64>("RequestId");
-		FlushAsyncLoadingStartCycle = EventData.GetValue<uint64>("Cycle");
-		if (FlushAsyncLoadingRequestId != INDEX_NONE)
-		{
-			//check(ActiveRequestsMap.Contains(FlushAsyncLoadingRequestId));
-			/*uint64 EventTypeId = RequestTimelineEventTypeMap[FlushAsyncLoadingRequestId];
-			BlockingRequestsTimeline->AppendBeginEvent(Context.SessionContext.TimestampFromCycle(FlushAsyncLoadingStartCycle), EventTypeId);*/
-		}
-		else
-		{
-			//BlockingRequestsTimeline->AppendBeginEvent(Context.SessionContext.TimestampFromCycle(FlushAsyncLoadingStartCycle), FlushAsyncLoadingEventId);
-		}
-		break;
-	}
-	case RouteId_EndFlushAsyncLoading:
-	{
-		uint64 Cycle = EventData.GetValue<uint64>("Cycle");
-		//BlockingRequestsTimeline->AppendEndEvent(Context.SessionContext.TimestampFromCycle(Cycle));
-		break;
-	}
 	case RouteId_ClassInfo:
 	{
 		Trace::FAnalysisSessionEditScope _(Session);
@@ -455,3 +352,24 @@ void FAsyncLoadingTraceAnalyzer::OnEvent(uint16 RouteId, const FOnEventContext&
 	}
 	}
 }
+
+void FAsyncLoadingTraceAnalyzer::PackageRequestAssociation(const FOnEventContext& Context, TSharedRef<FAsyncPackageState> AsyncPackageState, TSharedRef<FRequestState> RequestState)
+{
+	if (!AsyncPackageState->Request)
+	{
+		RequestState->AsyncPackages.Add(AsyncPackageState);
+		AsyncPackageState->Request = RequestState;
+		Trace::FLoadRequest* LoadRequest = RequestState->Group->LoadRequest;
+		Trace::FAnalysisSessionEditScope _(Session);
+		if (!LoadRequest)
+		{
+			LoadRequest = &LoadTimeProfilerProvider.CreateRequest();
+			LoadRequest->StartTime = Context.SessionContext.TimestampFromCycle(RequestState->WallTimeStartCycle);
+			LoadRequest->EndTime = std::numeric_limits<double>::infinity();
+			LoadRequest->Name = Session.StoreString(*RequestState->Group->Name);
+			LoadRequest->ThreadId = RequestState->ThreadId;
+			RequestState->Group->LoadRequest = LoadRequest;
+		}
+		LoadRequest->Packages.Add(AsyncPackageState->PackageInfo);
+	}
+}
\ No newline at end of file
diff --git a/Engine/Source/Developer/TraceServices/Private/Analyzers/LoadTimeTraceAnalysis.h b/Engine/Source/Developer/TraceServices/Private/Analyzers/LoadTimeTraceAnalysis.h
--- a/Engine/Source/Developer/TraceServices/Private/Analyzers/LoadTimeTraceAnalysis.h
+++ b/Engine/Source/Developer/TraceServices/Private/Analyzers/LoadTimeTraceAnalysis.h
@@ -34,29 +34,22 @@ private:
 	struct FRequestState;
 	struct FAsyncPackageState;
 
-	struct FLoadMapState
+	struct FRequestGroupState
 	{
-		uint64 Id;
 		FString Name;
-		uint64 WallTimeStartCycle;
-		uint64 WallTimeEndCycle;
-		TArray<TSharedRef<FRequestState>> Requests;
-	};
-
-	struct FStreamableHandleState
-	{
-		uint64 Id;
-		FString DebugName;
-		uint64 WallTimeStartCycle;
-		uint64 WallTimeEndCycle;
 		TArray<TSharedRef<FRequestState>> Requests;
+		Trace::FLoadRequest* LoadRequest = nullptr;
+		uint64 LatestEndCycle = 0;
+		uint64 ActiveRequestsCount = 0;
+		bool bIsClosed = false;
 	};
 
 	struct FRequestState
 	{
-		uint64 Id;
 		uint64 WallTimeStartCycle;
 		uint64 WallTimeEndCycle;
+		uint32 ThreadId;
+		TSharedPtr<FRequestGroupState> Group;
 		TArray<TSharedRef<FAsyncPackageState>> AsyncPackages;
 	};
 
@@ -70,7 +63,7 @@ private:
 	{
 		Trace::FPackageInfo* PackageInfo = nullptr;
 		TSharedPtr<FLinkerState> LinkerState;
-		TArray<TSharedRef<FRequestState>> Requests;
+		TSharedPtr<FRequestState> Request;
 	};
 
 	struct FScopeStackEntry
@@ -84,29 +77,19 @@ private:
 		FScopeStackEntry CpuScopeStack[256];
 		uint64 CpuScopeStackDepth = 0;
 		Trace::FLoadTimeProfilerCpuEvent CurrentEvent;
-		uint64 WaitForStreamableHandleStartCycle = 0;
-		TSharedPtr<FStreamableHandleState> WaitForStreamableHandleHandle;
-		TSharedPtr<FLoadMapState> ActiveLoadMap;
+		TArray<TSharedPtr<FRequestGroupState>> RequestGroupStack;
 		
 		TSharedPtr<Trace::FLoadTimeProfilerProvider::CpuTimelineInternal> CpuTimeline;
 
 		void EnterPackageScope(double Time, const Trace::FPackageInfo* PackageInfo, ELoadTimeProfilerPackageEventType EventType)
 		{
 			FScopeStackEntry& StackEntry = CpuScopeStack[CpuScopeStackDepth++];
-			StackEntry.Event.Export = CurrentEvent.Export;
-			StackEntry.Event.ExportEventType = CurrentEvent.ExportEventType;
+			StackEntry.Event.Export = nullptr;
+			StackEntry.Event.ExportEventType = LoadTimeProfilerObjectEventType_None;
 			StackEntry.Event.Package = PackageInfo;
 			StackEntry.Event.PackageEventType = EventType;
-			if (PackageInfo && CurrentEvent != StackEntry.Event)
-			{
-				CurrentEvent = StackEntry.Event;
-				CpuTimeline->AppendBeginEvent(Time, StackEntry.Event);
-				StackEntry.EnteredEvent = true;
-			}
-			else
-			{
-				StackEntry.EnteredEvent = false;
-			}
+			CurrentEvent = StackEntry.Event;
+			CpuTimeline->AppendBeginEvent(Time, StackEntry.Event);
 		}
 
 		void EnterExportScope(double Time, const Trace::FPackageExportInfo* ExportInfo, ELoadTimeProfilerObjectEventType EventType)
@@ -114,28 +97,20 @@ private:
 			FScopeStackEntry& StackEntry = CpuScopeStack[CpuScopeStackDepth++];
 			StackEntry.Event.Export = ExportInfo;
 			StackEntry.Event.ExportEventType = EventType;
-			StackEntry.Event.Package = CurrentEvent.Package;
+			StackEntry.Event.Package = ExportInfo ? ExportInfo->Package : nullptr;
 			StackEntry.Event.PackageEventType = CurrentEvent.PackageEventType;
-
-			if (ExportInfo && CurrentEvent != StackEntry.Event)
-			{
-				CurrentEvent = StackEntry.Event;
-				CpuTimeline->AppendBeginEvent(Time, StackEntry.Event);
-				StackEntry.EnteredEvent = true;
-			}
-			else
+			if (EventType == LoadTimeProfilerObjectEventType_PostLoad && StackEntry.Event.PackageEventType == LoadTimeProfilerPackageEventType_None)
 			{
-				StackEntry.EnteredEvent = false;
+				StackEntry.Event.PackageEventType = LoadTimeProfilerPackageEventType_DeferredPostLoad;
 			}
+			CurrentEvent = StackEntry.Event;
+			CpuTimeline->AppendBeginEvent(Time, StackEntry.Event);
 		}
 
 		void LeaveScope(double Time)
 		{
 			FScopeStackEntry& StackEntry = CpuScopeStack[--CpuScopeStackDepth];
-			if (StackEntry.EnteredEvent)
-			{
-				CpuTimeline->AppendEndEvent(Time);
-			}
+			CpuTimeline->AppendEndEvent(Time);
 			if (CpuScopeStackDepth > 0)
 			{
 				CurrentEvent = CpuScopeStack[CpuScopeStackDepth - 1].Event;
@@ -160,6 +135,7 @@ private:
 		}
 	};
 
+	void PackageRequestAssociation(const FOnEventContext& Context, TSharedRef<FAsyncPackageState> AsyncPackageState, TSharedRef<FRequestState> RequestState);
 	TSharedRef<FThreadState> GetThreadState(uint32 ThreadId);
 	const Trace::FClassInfo* GetClassInfo(uint64 ClassPtr) const;
 
@@ -174,51 +150,39 @@ private:
 		RouteId_DestroyAsyncPackage,
 		RouteId_BeginRequest,
 		RouteId_EndRequest,
-		RouteId_BeginLoadMap,
-		RouteId_EndLoadMap,
-		RouteId_NewStreamableHandle,
-		RouteId_DestroyStreamableHandle,
-		RouteId_BeginLoadStreamableHandle,
-		RouteId_EndLoadStreamableHandle,
-		RouteId_BeginWaitForStreamableHandle,
-		RouteId_EndWaitForStreamableHandle,
+		RouteId_BeginRequestGroup,
+		RouteId_EndRequestGroup,
 		RouteId_PackageSummary,
-		RouteId_StreamableHandleRequestAssociation,
 		RouteId_AsyncPackageRequestAssociation,
 		RouteId_AsyncPackageLinkerAssociation,
-		RouteId_LinkerArchiveAssociation,
+		RouteId_AsyncPackageImportDependency,
 		RouteId_BeginAsyncPackageScope,
 		RouteId_EndAsyncPackageScope,
 		RouteId_BeginCreateExport,
 		RouteId_EndCreateExport,
 		RouteId_BeginObjectScope,
 		RouteId_EndObjectScope,
-		RouteId_BeginFlushAsyncLoading,
-		RouteId_EndFlushAsyncLoading,
 		RouteId_ClassInfo,
 	};
 
+	enum
+	{
+		FormatBufferSize = 65536
+	};
+	TCHAR FormatBuffer[FormatBufferSize];
+
 	Trace::IAnalysisSession& Session;
 	Trace::FLoadTimeProfilerProvider& LoadTimeProfilerProvider;
-	//TSharedPtr<Trace::FMonotonicTimeline<uint64>> BlockingRequestsTimeline;
 
-	TArray<TSharedRef<FLoadMapState>> Maps;
-	TArray<TSharedRef<FStreamableHandleState>> StreamableHandles;
-	TArray<TSharedRef<FRequestState>> Requests;
+	//TArray<TSharedRef<FRequestState>> Requests;
 
 	TMap<uint64, TSharedRef<FLinkerState>> ActiveLinkersMap;
 	TMap<uint64, TSharedRef<FAsyncPackageState>> ActiveAsyncPackagesMap;
 	TMap<uint64, Trace::FPackageExportInfo*> ExportsMap;
-	TMap<uint64, TSharedRef<FStreamableHandleState>> ActiveStreamableHandlesMap;
 	TMap<uint64, TSharedRef<FRequestState>> ActiveRequestsMap;
 	TMap<uint32, TSharedRef<FThreadState>> ThreadStatesMap;
 	TMap<uint64, const Trace::FClassInfo*> ClassInfosMap;
 
-	uint64 FlushAsyncLoadingRequestId;
-	uint64 FlushAsyncLoadingStartCycle;
-
-	uint64 NextMapId = 0;
-	uint64 NextStreamableHandleId = 0;
 	int64 MainThreadId = -1;
 	int64 AsyncLoadingThreadId = -1;
 };
diff --git a/Engine/Source/Developer/TraceServices/Private/Model/LoadTimeProfiler.cpp b/Engine/Source/Developer/TraceServices/Private/Model/LoadTimeProfiler.cpp
--- a/Engine/Source/Developer/TraceServices/Private/Model/LoadTimeProfiler.cpp
+++ b/Engine/Source/Developer/TraceServices/Private/Model/LoadTimeProfiler.cpp
@@ -11,24 +11,16 @@ namespace Trace
 FLoadTimeProfilerProvider::FLoadTimeProfilerProvider(IAnalysisSession& InSession)
 	: Session(InSession)
 	, ClassInfos(Session.GetLinearAllocator(), 4096)
+	, Requests(Session.GetLinearAllocator(), 4096)
 	, Packages(Session.GetLinearAllocator(), 4096)
 	, Exports(Session.GetLinearAllocator(), 4090)
 	, MainThreadCpuTimeline(MakeShared<CpuTimelineInternal>(Session.GetLinearAllocator()))
 	, AsyncLoadingThreadCpuTimeline(MakeShared<CpuTimelineInternal>(Session.GetLinearAllocator()))
+	, RequestsTable(Requests)
 {
 	
 }
 
-void FLoadTimeProfilerProvider::EnumeratePackages(TFunctionRef<void(const FPackageInfo&)> Callback) const
-{
-	Session.ReadAccessCheck();
-
-	for (auto Iterator = Packages.GetIteratorFromItem(0); Iterator; ++Iterator)
-	{
-		Callback(*Iterator);
-	}
-}
-
 void FLoadTimeProfilerProvider::ReadMainThreadCpuTimeline(TFunctionRef<void(const CpuTimeline &)> Callback) const
 {
 	Session.ReadAccessCheck();
@@ -57,45 +49,7 @@ ITable<FLoadTimeProfilerAggregatedStats>* FLoadTimeProfilerProvider::CreateEvent
 	for (const auto& KV : Aggregation)
 	{
 		FLoadTimeProfilerAggregatedStats& Row = Table->AddRow();
-		switch (KV.Key)
-		{
-		case LoadTimeProfilerPackageEventType_CreateLinker:
-			Row.Name = TEXT("CreateLinker");
-			break;
-		case LoadTimeProfilerPackageEventType_FinishLinker:
-			Row.Name = TEXT("FinishLinker");
-			break;
-		case LoadTimeProfilerPackageEventType_StartImportPackages:
-			Row.Name = TEXT("StartImportPackages");
-			break;
-		case LoadTimeProfilerPackageEventType_SetupImports:
-			Row.Name = TEXT("SetupImports");
-			break;
-		case LoadTimeProfilerPackageEventType_SetupExports:
-			Row.Name = TEXT("SetupExports");
-			break;
-		case LoadTimeProfilerPackageEventType_ProcessImportsAndExports:
-			Row.Name = TEXT("ProcessImportsAndExports");
-			break;
-		case LoadTimeProfilerPackageEventType_ExportsDone:
-			Row.Name = TEXT("ExportsDone");
-			break;
-		case LoadTimeProfilerPackageEventType_PostLoadWait:
-			Row.Name = TEXT("PostLoadWait");
-			break;
-		case LoadTimeProfilerPackageEventType_StartPostLoad:
-			Row.Name = TEXT("StartPostLoad");
-			break;
-		case LoadTimeProfilerPackageEventType_Tick:
-			Row.Name = TEXT("Tick");
-			break;
-		case LoadTimeProfilerPackageEventType_DeferredPostLoad:
-			Row.Name = TEXT("DeferredPostLoad");
-			break;
-		case LoadTimeProfilerPackageEventType_Finish:
-			Row.Name = TEXT("Finish");
-			break;
-		}
+		Row.Name = GetLoadTimeProfilerPackageEventTypeString(KV.Key);
 		const FAggregatedTimingStats& Stats = KV.Value;
 		Row.Count = Stats.InstanceCount;
 		Row.Total = Stats.TotalExclusiveTime;
@@ -139,6 +93,128 @@ ITable<FLoadTimeProfilerAggregatedStats>* FLoadTimeProfilerProvider::CreateObjec
 	return Table;
 }
 
+ITable<FPackagesTableRow>* FLoadTimeProfilerProvider::CreatePackageDetailsTable(double IntervalStart, double IntervalEnd) const
+{
+	TTable<FPackagesTableLayout>* Table = new TTable<FPackagesTableLayout>();
+
+	TMap<TTuple<const FPackageInfo*, ELoadTimeProfilerPackageEventType>, FPackagesTableRow*> PackagesMap;
+
+	auto FindRow = [Table, &PackagesMap](const FLoadTimeProfilerCpuEvent& Event) -> FPackagesTableRow*
+	{
+		if (Event.Package)
+		{
+			auto Key = MakeTuple(Event.Package, Event.PackageEventType);
+			FPackagesTableRow** FindIt = PackagesMap.Find(Key);
+			FPackagesTableRow* Row = nullptr;
+			if (!FindIt)
+			{
+				Row = &Table->AddRow();
+				Row->PackageInfo = Event.Package;
+				Row->EventType = Event.PackageEventType;
+				PackagesMap.Add(Key, Row);
+			}
+			else
+			{
+				Row = *FindIt;
+			}
+
+			if (Event.Export && Event.ExportEventType == LoadTimeProfilerObjectEventType_Serialize)
+			{
+				++Row->SerializedExportsCount;
+				Row->SerializedExportsSize += Event.Export->SerialSize;
+			}
+
+			if (Event.PackageEventType == LoadTimeProfilerPackageEventType_CreateLinker)
+			{
+				Row->SerializedHeaderSize += Event.Package->Summary.TotalHeaderSize;
+			}
+
+			return Row;
+		}
+		else
+		{
+			return nullptr;
+		}
+	};
+
+	AsyncLoadingThreadCpuTimeline->EnumerateEvents(IntervalStart, IntervalEnd, [Table, &PackagesMap, FindRow](double StartTime, double EndTime, uint32, const FLoadTimeProfilerCpuEvent& Event)
+	{
+		FPackagesTableRow* Row = FindRow(Event);
+		if (Row)
+		{
+			Row->AsyncLoadingThreadTime += EndTime - StartTime; // TODO: Should be exclusive time
+		}
+	});
+	MainThreadCpuTimeline->EnumerateEvents(IntervalStart, IntervalEnd, [Table, &PackagesMap, FindRow](double StartTime, double EndTime, uint32, const FLoadTimeProfilerCpuEvent& Event)
+	{
+		FPackagesTableRow* Row = FindRow(Event);
+		if (Row)
+		{
+			Row->MainThreadTime += EndTime - StartTime; // TODO: Should be exclusive time
+		}
+	});
+
+	return Table;
+}
+
+ITable<FExportsTableRow>* FLoadTimeProfilerProvider::CreateExportDetailsTable(double IntervalStart, double IntervalEnd) const
+{
+	TTable<FExportsTableLayout>* Table = new TTable<FExportsTableLayout>();
+
+	TMap<TTuple<const FPackageExportInfo*, ELoadTimeProfilerObjectEventType>, FExportsTableRow*> ExportsMap;
+
+	auto FindRow = [Table, &ExportsMap](const FLoadTimeProfilerCpuEvent& Event) -> FExportsTableRow*
+	{
+		if (Event.Export)
+		{
+			auto Key = MakeTuple(Event.Export, Event.ExportEventType);
+			FExportsTableRow** FindIt = ExportsMap.Find(Key);
+			FExportsTableRow* Row = nullptr;
+			if (!FindIt)
+			{
+				Row = &Table->AddRow();
+				Row->ExportInfo = Event.Export;
+				Row->EventType = Event.ExportEventType;
+				ExportsMap.Add(Key, Row);
+			}
+			else
+			{
+				Row = *FindIt;
+			}
+
+			if (Event.ExportEventType == LoadTimeProfilerObjectEventType_Serialize)
+			{
+				Row->SerializedSize += Event.Export->SerialSize;
+			}
+
+			return Row;
+		}
+		else
+		{
+			return nullptr;
+		}
+	};
+
+	AsyncLoadingThreadCpuTimeline->EnumerateEvents(IntervalStart, IntervalEnd, [Table, &ExportsMap, FindRow](double StartTime, double EndTime, uint32, const FLoadTimeProfilerCpuEvent& Event)
+	{
+		FExportsTableRow* Row = FindRow(Event);
+		if (Row)
+		{
+			Row->AsyncLoadingThreadTime += EndTime - StartTime; // TODO: Should be exclusive time
+		}
+	});
+	MainThreadCpuTimeline->EnumerateEvents(IntervalStart, IntervalEnd, [Table, &ExportsMap, FindRow](double StartTime, double EndTime, uint32, const FLoadTimeProfilerCpuEvent& Event)
+	{
+		FExportsTableRow* Row = FindRow(Event);
+		if (Row)
+		{
+			Row->MainThreadTime += EndTime - StartTime; // TODO: Should be exclusive time
+		}
+	});
+
+	return Table;
+}
+
 const Trace::FClassInfo& FLoadTimeProfilerProvider::AddClassInfo(const TCHAR* ClassName)
 {
 	Session.WriteAccessCheck();
@@ -148,6 +224,14 @@ const Trace::FClassInfo& FLoadTimeProfilerProvider::AddClassInfo(const TCHAR* Cl
 	return ClassInfo;
 }
 
+Trace::FLoadRequest& FLoadTimeProfilerProvider::CreateRequest()
+{
+	Session.WriteAccessCheck();
+
+	FLoadRequest& RequestInfo = Requests.PushBack();
+	return RequestInfo;
+}
+
 Trace::FPackageInfo& FLoadTimeProfilerProvider::CreatePackage(const TCHAR* PackageName)
 {
 	Session.WriteAccessCheck();
@@ -169,4 +253,68 @@ Trace::FPackageExportInfo& FLoadTimeProfilerProvider::CreateExport()
 	return Export;
 }
 
+uint64 FLoadTimeProfilerProvider::PackageSizeSum(const FLoadRequest& Row)
+{
+	uint64 Sum = 0;
+	for (const FPackageInfo* Package : Row.Packages)
+	{
+		Sum += Package->Summary.TotalHeaderSize;
+		for (const FPackageExportInfo* Export : Package->Exports)
+		{
+			Sum += Export->SerialSize;
+		}
+	}
+	return Sum;
+}
+
+const TCHAR* GetLoadTimeProfilerPackageEventTypeString(ELoadTimeProfilerPackageEventType EventType)
+{
+	switch (EventType)
+	{
+	case LoadTimeProfilerPackageEventType_CreateLinker:
+		return TEXT("CreateLinker");
+	case LoadTimeProfilerPackageEventType_FinishLinker:
+		return TEXT("FinishLinker");
+	case LoadTimeProfilerPackageEventType_StartImportPackages:
+		return TEXT("StartImportPackages");
+	case LoadTimeProfilerPackageEventType_SetupImports:
+		return TEXT("SetupImports");
+	case LoadTimeProfilerPackageEventType_SetupExports:
+		return TEXT("SetupExports");
+	case LoadTimeProfilerPackageEventType_ProcessImportsAndExports:
+		return TEXT("ProcessImportsAndExports");
+	case LoadTimeProfilerPackageEventType_ExportsDone:
+		return TEXT("ExportsDone");
+	case LoadTimeProfilerPackageEventType_PostLoadWait:
+		return TEXT("PostLoadWait");
+	case LoadTimeProfilerPackageEventType_StartPostLoad:
+		return TEXT("StartPostLoad");
+	case LoadTimeProfilerPackageEventType_Tick:
+		return TEXT("Tick");
+	case LoadTimeProfilerPackageEventType_DeferredPostLoad:
+		return TEXT("DeferredPostLoad");
+	case LoadTimeProfilerPackageEventType_Finish:
+		return TEXT("Finish");
+	case LoadTimeProfilerPackageEventType_None:
+		return TEXT("None");
+	}
+	return TEXT("[invalid]");
+}
+
+const TCHAR* GetLoadTimeProfilerObjectEventTypeString(ELoadTimeProfilerObjectEventType EventType)
+{
+	switch (EventType)
+	{
+	case LoadTimeProfilerObjectEventType_Create:
+		return TEXT("Create");
+	case LoadTimeProfilerObjectEventType_Serialize:
+		return TEXT("Serialize");
+	case LoadTimeProfilerObjectEventType_PostLoad:
+		return TEXT("PostLoad");
+	case LoadTimeProfilerObjectEventType_None:
+		return TEXT("None");
+	}
+	return TEXT("[invalid]");
+}
+
 }
\ No newline at end of file
diff --git a/Engine/Source/Developer/TraceServices/Private/Model/LoadTimeProfilerPrivate.h b/Engine/Source/Developer/TraceServices/Private/Model/LoadTimeProfilerPrivate.h
--- a/Engine/Source/Developer/TraceServices/Private/Model/LoadTimeProfilerPrivate.h
+++ b/Engine/Source/Developer/TraceServices/Private/Model/LoadTimeProfilerPrivate.h
@@ -22,13 +22,15 @@ public:
 	typedef TMonotonicTimeline<FLoadTimeProfilerCpuEvent> CpuTimelineInternal;
 
 	FLoadTimeProfilerProvider(IAnalysisSession& Session);
-	virtual uint64 GetPackageCount() const override { return Packages.Num(); }
-	virtual void EnumeratePackages(TFunctionRef<void(const FPackageInfo&)> Callback) const override;
 	virtual void ReadMainThreadCpuTimeline(TFunctionRef<void(const CpuTimeline&)> Callback) const override;
 	virtual void ReadAsyncLoadingThreadCpuTimeline(TFunctionRef<void(const CpuTimeline&)> Callback) const override;
 	virtual ITable<FLoadTimeProfilerAggregatedStats>* CreateEventAggregation(double IntervalStart, double IntervalEnd) const override;
 	virtual ITable<FLoadTimeProfilerAggregatedStats>* CreateObjectTypeAggregation(double IntervalStart, double IntervalEnd) const override;
+	virtual ITable<FPackagesTableRow>* CreatePackageDetailsTable(double IntervalStart, double IntervalEnd) const override;
+	virtual ITable<FExportsTableRow>* CreateExportDetailsTable(double IntervalStart, double IntervalEnd) const override;
+	virtual const ITable<FLoadRequest>& GetRequestsTable() const override { return RequestsTable; }
 	const FClassInfo& AddClassInfo(const TCHAR* ClassName);
+	FLoadRequest& CreateRequest();
 	FPackageInfo& CreatePackage(const TCHAR* PackageName);
 	FPackageExportInfo& CreateExport();
 	TSharedRef<CpuTimelineInternal> EditMainThreadCpuTimeline() { return MainThreadCpuTimeline; }
@@ -39,14 +41,7 @@ public:
 	void SetAsyncLoadingThreadId(uint32 ThreadId) { AsyncLoadingThreadId = ThreadId; }
 
 private:
-	IAnalysisSession& Session;
-	TPagedArray<FClassInfo> ClassInfos;
-	TPagedArray<FPackageInfo> Packages;
-	TPagedArray<FPackageExportInfo> Exports;
-	TSharedRef<CpuTimelineInternal> MainThreadCpuTimeline;
-	TSharedRef<CpuTimelineInternal> AsyncLoadingThreadCpuTimeline;
-	uint32 MainThreadId = uint32(-1);
-	uint32 AsyncLoadingThreadId = uint32(-1);
+	static uint64 PackageSizeSum(const FLoadRequest& Row);
 
 	UE_TRACE_TABLE_LAYOUT_BEGIN(FAggregatedStatsTableLayout, FLoadTimeProfilerAggregatedStats)
 		UE_TRACE_TABLE_COLUMN(Name, TEXT("Name"))
@@ -57,6 +52,46 @@ private:
 		UE_TRACE_TABLE_COLUMN(Average, TEXT("Avg"))
 		UE_TRACE_TABLE_COLUMN(Median, TEXT("Med"))
 	UE_TRACE_TABLE_LAYOUT_END()
+
+	UE_TRACE_TABLE_LAYOUT_BEGIN(FPackagesTableLayout, FPackagesTableRow)
+		UE_TRACE_TABLE_PROJECTED_COLUMN(TableColumnType_CString, TEXT("Package"), [](const FPackagesTableRow& Row) { return Row.PackageInfo->Name; })
+		UE_TRACE_TABLE_PROJECTED_COLUMN(TableColumnType_CString, TEXT("EventType"), [](const FPackagesTableRow& Row) { return GetLoadTimeProfilerPackageEventTypeString(Row.EventType); })
+		UE_TRACE_TABLE_COLUMN(SerializedHeaderSize, TEXT("SerializedHeaderSize"))
+		UE_TRACE_TABLE_COLUMN(SerializedExportsCount, TEXT("SerializedExportsCount"))
+		UE_TRACE_TABLE_COLUMN(SerializedExportsSize, TEXT("SerializedExportsSize"))
+		UE_TRACE_TABLE_COLUMN(MainThreadTime, TEXT("MainThreadTime"))
+		UE_TRACE_TABLE_COLUMN(AsyncLoadingThreadTime, TEXT("AsyncLoadingThreadTime"))
+	UE_TRACE_TABLE_LAYOUT_END()
+
+	UE_TRACE_TABLE_LAYOUT_BEGIN(FExportsTableLayout, FExportsTableRow)
+		UE_TRACE_TABLE_PROJECTED_COLUMN(TableColumnType_CString, TEXT("Package"), [](const FExportsTableRow& Row) { return Row.ExportInfo->Package->Name; })
+		UE_TRACE_TABLE_PROJECTED_COLUMN(TableColumnType_CString, TEXT("Class"), [](const FExportsTableRow& Row) { return Row.ExportInfo->Class ? Row.ExportInfo->Class->Name : TEXT("[unknown]"); })
+		UE_TRACE_TABLE_PROJECTED_COLUMN(TableColumnType_CString, TEXT("EventType"), [](const FExportsTableRow& Row) { return GetLoadTimeProfilerObjectEventTypeString(Row.EventType); })
+		UE_TRACE_TABLE_COLUMN(SerializedSize, TEXT("SerializedSize"))
+		UE_TRACE_TABLE_COLUMN(MainThreadTime, TEXT("MainThreadTime"))
+		UE_TRACE_TABLE_COLUMN(AsyncLoadingThreadTime, TEXT("AsyncLoadingThreadTime"))
+	UE_TRACE_TABLE_LAYOUT_END()
+
+	UE_TRACE_TABLE_LAYOUT_BEGIN(FRequestsTableLayout, FLoadRequest)
+		UE_TRACE_TABLE_COLUMN(Name, TEXT("Name"))
+		UE_TRACE_TABLE_COLUMN(ThreadId, TEXT("ThreadId"))
+		UE_TRACE_TABLE_COLUMN(StartTime, TEXT("StartTime"))
+		UE_TRACE_TABLE_COLUMN(EndTime, TEXT("EndTime"))
+		UE_TRACE_TABLE_PROJECTED_COLUMN(TableColumnType_Int, TEXT("PackageCount"), [](const FLoadRequest& Row) { return Row.Packages.Num(); })
+		UE_TRACE_TABLE_PROJECTED_COLUMN(TableColumnType_Int, TEXT("Size"), PackageSizeSum)
+		UE_TRACE_TABLE_PROJECTED_COLUMN(TableColumnType_CString, TEXT("FirstPackage"), [](const FLoadRequest& Row) { return Row.Packages.Num() ? Row.Packages[0]->Name : TEXT("N/A"); })
+	UE_TRACE_TABLE_LAYOUT_END()
+
+	IAnalysisSession& Session;
+	TPagedArray<FClassInfo> ClassInfos;
+	TPagedArray<FLoadRequest> Requests;
+	TPagedArray<FPackageInfo> Packages;
+	TPagedArray<FPackageExportInfo> Exports;
+	TSharedRef<CpuTimelineInternal> MainThreadCpuTimeline;
+	TSharedRef<CpuTimelineInternal> AsyncLoadingThreadCpuTimeline;
+	uint32 MainThreadId = uint32(-1);
+	uint32 AsyncLoadingThreadId = uint32(-1);
+	TTableView<FRequestsTableLayout> RequestsTable;
 };
 
 }
diff --git a/Engine/Source/Developer/TraceServices/Private/Model/Tables.h b/Engine/Source/Developer/TraceServices/Private/Model/Tables.h
--- a/Engine/Source/Developer/TraceServices/Private/Model/Tables.h
+++ b/Engine/Source/Developer/TraceServices/Private/Model/Tables.h
@@ -374,14 +374,15 @@ public:
 	{
 		if (!CurrentRow)
 		{
-			return nullptr;
+			return TEXT("");
 		}
 		ETableColumnType ColumnType = Layout.GetColumnTypeConstExpr(ColumnIndex);
 		if (ColumnType == TableColumnType_CString)
 		{
-			return Layout.GetColumnValue(*CurrentRow, ColumnIndex).CStringValue;
+			const TCHAR* CStringValue = Layout.GetColumnValue(*CurrentRow, ColumnIndex).CStringValue;
+			return CStringValue ? CStringValue : TEXT("");
 		}
-		return nullptr;
+		return TEXT("");
 	}
 
 private:
@@ -482,15 +483,9 @@ public:
 };
 
 template<typename LayoutType>
-class TTableView
+class TTableBase
 	: public ITable<typename LayoutType::RowType>
 {
-public:
-	TTableView(const TPagedArray<typename LayoutType::RowType>& InRows)
-		: Rows(InRows)
-	{
-	}
-
 	virtual const ITableLayout& GetLayout() const override
 	{
 		return Layout;
@@ -498,44 +493,49 @@ public:
 
 	virtual uint64 GetRowCount() const override
 	{
-		return Rows.Num();
+		return GetRows().Num();
 	}
 
 	virtual ITableReader<typename LayoutType::RowType>* CreateReader() const override
 	{
-		return new TTableReader<LayoutType>(Layout, Rows);
+		return new TTableReader<LayoutType>(Layout, GetRows());
 	}
 
 private:
+	virtual const TPagedArray<typename LayoutType::RowType>& GetRows() const = 0;
+
 	LayoutType Layout;
-	const TPagedArray<typename LayoutType::RowType>& Rows;
 };
 
 template<typename LayoutType>
-class TTable
-	: public ITable<typename LayoutType::RowType>
+class TTableView
+	: public TTableBase<LayoutType>
 {
 public:
-	TTable()
-		: Allocator(2 << 20)
-		, Rows(Allocator, 1024)
+	TTableView(const TPagedArray<typename LayoutType::RowType>& InRows)
+		: Rows(InRows)
 	{
-
 	}
 
-	virtual const ITableLayout& GetLayout() const override
+private:
+	virtual const TPagedArray<typename LayoutType::RowType>& GetRows() const override
 	{
-		return Layout;
+		return Rows;
 	}
 
-	virtual uint64 GetRowCount() const override
-	{
-		return Rows.Num();
-	}
+	const TPagedArray<typename LayoutType::RowType>& Rows;
+};
 
-	virtual ITableReader<typename LayoutType::RowType>* CreateReader() const override
+template<typename LayoutType>
+class TTable
+	: public TTableBase<LayoutType>
+{
+public:
+	TTable()
+		: Allocator(2 << 20)
+		, Rows(Allocator, 1024)
 	{
-		return new TTableReader<LayoutType>(Layout, Rows);
+
 	}
 
 	typename LayoutType::RowType& AddRow()
@@ -544,7 +544,11 @@ public:
 	}
 
 private:
-	LayoutType Layout;
+	virtual const TPagedArray<typename LayoutType::RowType>& GetRows() const override
+	{
+		return Rows;
+	}
+
 	FSlabAllocator Allocator;
 	TPagedArray<typename LayoutType::RowType> Rows;
 };
diff --git a/Engine/Source/Developer/TraceServices/Public/TraceServices/Model/LoadTimeProfiler.h b/Engine/Source/Developer/TraceServices/Public/TraceServices/Model/LoadTimeProfiler.h
--- a/Engine/Source/Developer/TraceServices/Public/TraceServices/Model/LoadTimeProfiler.h
+++ b/Engine/Source/Developer/TraceServices/Public/TraceServices/Model/LoadTimeProfiler.h
@@ -58,12 +58,15 @@ struct FClassInfo
 	const TCHAR* Name;
 };
 
+struct FPackageInfo;
+
 struct FPackageExportInfo
 {
 	uint32 Id;
 	const FClassInfo* Class = nullptr;
 	uint64 SerialOffset = 0;
 	uint64 SerialSize = 0;
+	const FPackageInfo* Package = nullptr;
 	bool IsAsset = false;
 };
 
@@ -72,7 +75,36 @@ struct FPackageInfo
 	uint32 Id;
 	const TCHAR* Name = nullptr;
 	FPackageSummaryInfo Summary;
-	TArray<FPackageExportInfo*> Exports;
+	TArray<const FPackageExportInfo*> Exports;
+};
+
+struct FLoadRequest
+{
+	const TCHAR* Name = nullptr;
+	uint32 ThreadId = uint32(-1);
+	double StartTime = 0.0;
+	double EndTime = 0.0;
+	TArray<const FPackageInfo*> Packages;
+};
+
+struct FExportsTableRow
+{
+	const FPackageExportInfo* ExportInfo = nullptr;
+	uint64 SerializedSize = 0;
+	double MainThreadTime = 0.0;
+	double AsyncLoadingThreadTime = 0.0;
+	ELoadTimeProfilerObjectEventType EventType = LoadTimeProfilerObjectEventType_None;
+};
+
+struct FPackagesTableRow
+{
+	const FPackageInfo* PackageInfo = nullptr;
+	uint64 SerializedHeaderSize = 0;
+	uint64 SerializedExportsCount = 0;
+	uint64 SerializedExportsSize = 0;
+	double MainThreadTime = 0.0;
+	double AsyncLoadingThreadTime = 0.0;
+	ELoadTimeProfilerPackageEventType EventType = LoadTimeProfilerPackageEventType_None;
 };
 
 struct FLoadTimeProfilerCpuEvent
@@ -94,6 +126,9 @@ struct FLoadTimeProfilerAggregatedStats
 	double Median;
 };
 
+TRACESERVICES_API const TCHAR* GetLoadTimeProfilerPackageEventTypeString(ELoadTimeProfilerPackageEventType EventType);
+TRACESERVICES_API const TCHAR* GetLoadTimeProfilerObjectEventTypeString(ELoadTimeProfilerObjectEventType EventType);
+
 class ILoadTimeProfilerProvider
 	: public IProvider
 {
@@ -101,14 +136,15 @@ public:
 	typedef ITimeline<FLoadTimeProfilerCpuEvent> CpuTimeline;
 
 	virtual ~ILoadTimeProfilerProvider() = default;
-	virtual uint64 GetPackageCount() const = 0;
-	virtual void EnumeratePackages(TFunctionRef<void(const FPackageInfo&)> Callback) const = 0;
 	virtual uint32 GetMainThreadId() const = 0;
 	virtual void ReadMainThreadCpuTimeline(TFunctionRef<void(const CpuTimeline&)> Callback) const = 0;
 	virtual uint32 GetAsyncLoadingThreadId() const = 0;
 	virtual void ReadAsyncLoadingThreadCpuTimeline(TFunctionRef<void(const CpuTimeline&)> Callback) const = 0;
 	virtual ITable<FLoadTimeProfilerAggregatedStats>* CreateEventAggregation(double IntervalStart, double IntervalEnd) const = 0;
 	virtual ITable<FLoadTimeProfilerAggregatedStats>* CreateObjectTypeAggregation(double IntervalStart, double IntervalEnd) const = 0;
+	virtual ITable<FPackagesTableRow>* CreatePackageDetailsTable(double IntervalStart, double IntervalEnd) const = 0;
+	virtual ITable<FExportsTableRow>* CreateExportDetailsTable(double IntervalStart, double IntervalEnd) const = 0;
+	virtual const ITable<FLoadRequest>& GetRequestsTable() const = 0;
 };
 
 TRACESERVICES_API const ILoadTimeProfilerProvider* ReadLoadTimeProfilerProvider(const IAnalysisSession& Session);
diff --git a/Engine/Source/Runtime/Core/Private/Modules/ModuleManager.cpp b/Engine/Source/Runtime/Core/Private/Modules/ModuleManager.cpp
--- a/Engine/Source/Runtime/Core/Private/Modules/ModuleManager.cpp
+++ b/Engine/Source/Runtime/Core/Private/Modules/ModuleManager.cpp
@@ -11,6 +11,7 @@
 #include "Modules/ModuleManifest.h"
 #include "Misc/ScopeLock.h"
 #include "Misc/DataDrivenPlatformInfoRegistry.h"
+#include "Serialization/LoadTimeTrace.h"
 
 DEFINE_LOG_CATEGORY_STATIC(LogModuleManager, Log, All);
 
@@ -445,6 +446,7 @@ IModuleInterface* FModuleManager::LoadModuleWithFailureReason(const FName InModu
 		if (ModuleInfo->Module.IsValid())
 		{
 			FScopedBootTiming BootScope("LoadModule  - ", InModuleName);
+			TRACE_LOADTIME_REQUEST_GROUP_SCOPE(TEXT("LoadModule - %s"), *InModuleName.ToString());
 #if USE_PER_MODULE_UOBJECT_BOOTSTRAP
 			{
 				ProcessLoadedObjectsCallback.Broadcast(InModuleName, bCanProcessNewlyLoadedObjects);
diff --git a/Engine/Source/Runtime/Core/Private/Serialization/LoadTimeTrace.cpp b/Engine/Source/Runtime/Core/Private/Serialization/LoadTimeTrace.cpp
new file mode 100644
--- /dev/null
+++ b/Engine/Source/Runtime/Core/Private/Serialization/LoadTimeTrace.cpp
@@ -0,0 +1,44 @@
+// Copyright 1998-2019 Epic Games, Inc. All Rights Reserved.
+
+#include "Serialization/LoadTimeTrace.h"
+
+#if LOADTIMEPROFILERTRACE_ENABLED
+
+#include "Trace/Trace.h"
+#include "Misc/CString.h"
+
+UE_TRACE_EVENT_BEGIN(LoadTime, BeginRequestGroup)
+	UE_TRACE_EVENT_FIELD(uint32, ThreadId)
+UE_TRACE_EVENT_END()
+
+UE_TRACE_EVENT_BEGIN(LoadTime, EndRequestGroup)
+	UE_TRACE_EVENT_FIELD(uint32, ThreadId)
+UE_TRACE_EVENT_END()
+
+FLoadTimeProfilerTrace::FRequestGroupScope::~FRequestGroupScope()
+{
+	UE_TRACE_LOG(LoadTime, EndRequestGroup)
+		<< EndRequestGroup.ThreadId(FPlatformTLS::GetCurrentThreadId());
+}
+
+void FLoadTimeProfilerTrace::FRequestGroupScope::OutputBegin()
+{
+	uint16 FormatStringSize = (FCString::Strlen(FormatString) + 1) * sizeof(TCHAR);
+	auto Attachment = [this, FormatStringSize](uint8* Out)
+	{
+		memcpy(Out, FormatString, FormatStringSize);
+		memcpy(Out + FormatStringSize, FormatArgsBuffer, FormatArgsSize);
+	};
+	UE_TRACE_LOG(LoadTime, BeginRequestGroup, FormatStringSize + FormatArgsSize)
+		<< BeginRequestGroup.ThreadId(FPlatformTLS::GetCurrentThreadId())
+		<< BeginRequestGroup.Attachment(Attachment);
+}
+
+void FLoadTimeProfilerTrace::InitInternal()
+{
+	UE_TRACE_EVENT_IS_ENABLED(LoadTime, BeginRequestGroup);
+	UE_TRACE_EVENT_IS_ENABLED(LoadTime, EndRequestGroup);
+}
+
+#endif
+
diff --git a/Engine/Source/Runtime/Core/Public/Serialization/LoadTimeTrace.h b/Engine/Source/Runtime/Core/Public/Serialization/LoadTimeTrace.h
new file mode 100644
--- /dev/null
+++ b/Engine/Source/Runtime/Core/Public/Serialization/LoadTimeTrace.h
@@ -0,0 +1,72 @@
+// Copyright 1998-2019 Epic Games, Inc. All Rights Reserved.
+
+#pragma once
+
+#include "CoreTypes.h"
+#include "Trace/Trace.h"
+#include "ProfilingDebugging/FormatArgsTrace.h"
+
+#if UE_TRACE_ENABLED && !UE_BUILD_SHIPPING
+#define LOADTIMEPROFILERTRACE_ENABLED 1
+#else
+#define LOADTIMEPROFILERTRACE_ENABLED 0
+#endif
+
+enum ELoadTimeProfilerPackageEventType
+{
+	LoadTimeProfilerPackageEventType_CreateLinker,
+	LoadTimeProfilerPackageEventType_FinishLinker,
+	LoadTimeProfilerPackageEventType_StartImportPackages,
+	LoadTimeProfilerPackageEventType_SetupImports,
+	LoadTimeProfilerPackageEventType_SetupExports,
+	LoadTimeProfilerPackageEventType_ProcessImportsAndExports,
+	LoadTimeProfilerPackageEventType_ExportsDone,
+	LoadTimeProfilerPackageEventType_PostLoadWait,
+	LoadTimeProfilerPackageEventType_StartPostLoad,
+	LoadTimeProfilerPackageEventType_Tick,
+	LoadTimeProfilerPackageEventType_Finish,
+	LoadTimeProfilerPackageEventType_DeferredPostLoad,
+	LoadTimeProfilerPackageEventType_None,
+};
+
+enum ELoadTimeProfilerObjectEventType
+{
+	LoadTimeProfilerObjectEventType_Create,
+	LoadTimeProfilerObjectEventType_Serialize,
+	LoadTimeProfilerObjectEventType_PostLoad,
+	LoadTimeProfilerObjectEventType_None
+};
+
+#if LOADTIMEPROFILERTRACE_ENABLED
+
+struct FLoadTimeProfilerTrace
+{
+	struct FRequestGroupScope
+	{
+		template <typename... Types>
+		FRequestGroupScope(const TCHAR* InFormatString, Types... FormatArgs)
+		{
+			FormatString = InFormatString;
+			FormatArgsSize = FFormatArgsTrace::EncodeArguments(FormatArgsBuffer, FormatArgs...);
+			OutputBegin();
+		}
+
+		CORE_API ~FRequestGroupScope();
+
+	private:
+		CORE_API void OutputBegin();
+
+		const TCHAR* FormatString = nullptr;
+		uint16 FormatArgsSize = 0;
+		uint8 FormatArgsBuffer[1024];
+	};
+
+	CORE_API static void InitInternal();
+};
+
+#define TRACE_LOADTIME_REQUEST_GROUP_SCOPE(Format, ...) \
+	FLoadTimeProfilerTrace::FRequestGroupScope __LoadTimeTraceRequestGroupScope(Format, ##__VA_ARGS__);
+
+#else
+#define TRACE_LOADTIME_REQUEST_GROUP_SCOPE(...)
+#endif
diff --git a/Engine/Source/Runtime/CoreUObject/Private/Serialization/AsyncLoading.cpp b/Engine/Source/Runtime/CoreUObject/Private/Serialization/AsyncLoading.cpp
--- a/Engine/Source/Runtime/CoreUObject/Private/Serialization/AsyncLoading.cpp
+++ b/Engine/Source/Runtime/CoreUObject/Private/Serialization/AsyncLoading.cpp
@@ -1793,7 +1793,6 @@ void FAsyncLoadingThread::QueueEvent_CreateLinker(FAsyncPackage* Package, int32
 
 	TAsyncLoadPriority UserPriority = Package->GetPriority();
 	int32 PackageSerialNumber = GRandomizeLoadOrder ? GetRandomSerialNumber() : Package->SerialNumber;
-	TRACE_LOADTIME_QUEUE_EVENT(Package, LoadTimeProfilerPackageEventType_CreateLinker, UserPriority, PackageSerialNumber, EventSystemPriority);
 	EventQueue.AddAsyncEvent(UserPriority, PackageSerialNumber, EventSystemPriority,
 		TFunction<void(FAsyncLoadEventArgs& Args)>(
 			[WeakPtr, this](FAsyncLoadEventArgs& Args)
@@ -1859,7 +1858,6 @@ void FAsyncLoadingThread::QueueEvent_FinishLinker(FWeakAsyncPackagePtr WeakPtr,
 	{
 		TAsyncLoadPriority UserPriority = Pkg->GetPriority();
 		int32 PackageSerialNumber = GRandomizeLoadOrder ? GetRandomSerialNumber() : Pkg->SerialNumber;
-		TRACE_LOADTIME_QUEUE_EVENT(Pkg, LoadTimeProfilerPackageEventType_FinishLinker, UserPriority, PackageSerialNumber, EventSystemPriority)
 		EventQueue.AddAsyncEvent(UserPriority, PackageSerialNumber, EventSystemPriority,
 			TFunction<void(FAsyncLoadEventArgs& Args)>(
 				[WeakPtr, this](FAsyncLoadEventArgs& Args)
@@ -2008,7 +2006,6 @@ void FAsyncLoadingThread::QueueEvent_StartImportPackages(FAsyncPackage* Package,
 
 	TAsyncLoadPriority UserPriority = Package->GetPriority();
 	int32 PackageSerialNumber = GRandomizeLoadOrder ? GetRandomSerialNumber() : Package->SerialNumber;
-	TRACE_LOADTIME_QUEUE_EVENT(Package, LoadTimeProfilerPackageEventType_StartImportPackages, UserPriority, PackageSerialNumber, EventSystemPriority);
 	EventQueue.AddAsyncEvent(UserPriority, PackageSerialNumber, EventSystemPriority,
 		TFunction<void(FAsyncLoadEventArgs& Args)>(
 			[WeakPtr, this](FAsyncLoadEventArgs& Args)
@@ -2291,6 +2288,7 @@ EAsyncPackageState::Type FAsyncPackage::LoadImports_Event()
 				{
 					UE_LOG(LogStreaming, Display, TEXT("%s is prestreaming %s"), *Desc.NameToLoad.ToString(), *Import->ObjectName.ToString());
 				}
+				TRACE_LOADTIME_ASYNC_PACKAGE_IMPORT_DEPENDENCY(this, PendingPackage);
 				AsyncLoadingThread.InsertPackage(PendingPackage);
 				bDidSomething = true;
 			}
@@ -2351,7 +2349,6 @@ void FAsyncLoadingThread::QueueEvent_SetupImports(FAsyncPackage* Package, int32
 	FWeakAsyncPackagePtr WeakPtr(Package);
 	TAsyncLoadPriority UserPriority = Package->GetPriority();
 	int32 PackageSerialNumber = GRandomizeLoadOrder ? GetRandomSerialNumber() : Package->SerialNumber;
-	TRACE_LOADTIME_QUEUE_EVENT(Package, LoadTimeProfilerPackageEventType_SetupImports, UserPriority, PackageSerialNumber, EventSystemPriority);
 	EventQueue.AddAsyncEvent(UserPriority, PackageSerialNumber, EventSystemPriority,
 		TFunction<void(FAsyncLoadEventArgs& Args)>(
 			[WeakPtr, this](FAsyncLoadEventArgs& Args)
@@ -2714,7 +2711,6 @@ void FAsyncLoadingThread::QueueEvent_SetupExports(FAsyncPackage* Package, int32
 	FWeakAsyncPackagePtr WeakPtr(Package);
 	TAsyncLoadPriority UserPriority = Package->GetPriority();
 	int32 PackageSerialNumber = GRandomizeLoadOrder ? GetRandomSerialNumber() : Package->SerialNumber;
-	TRACE_LOADTIME_QUEUE_EVENT(Package, LoadTimeProfilerPackageEventType_SetupExports, UserPriority, PackageSerialNumber, EventSystemPriority);
 	EventQueue.AddAsyncEvent(UserPriority, PackageSerialNumber, EventSystemPriority,
 		TFunction<void(FAsyncLoadEventArgs& Args)>(
 			[WeakPtr, this](FAsyncLoadEventArgs& Args)
@@ -2752,7 +2748,6 @@ void FAsyncLoadingThread::QueueEvent_ProcessImportsAndExports(FAsyncPackage* Pac
 	FWeakAsyncPackagePtr WeakPtr(Package);
 	TAsyncLoadPriority UserPriority = Package->GetPriority();
 	int32 PackageSerialNumber = GRandomizeLoadOrder ? GetRandomSerialNumber() : Package->SerialNumber;
-	TRACE_LOADTIME_QUEUE_EVENT(Package, LoadTimeProfilerPackageEventType_ProcessImportsAndExports, UserPriority, PackageSerialNumber, EventSystemPriority);
 	EventQueue.AddAsyncEvent(Package->GetPriority(), GRandomizeLoadOrder ? GetRandomSerialNumber() : Package->SerialNumber, EventSystemPriority,
 		TFunction<void(FAsyncLoadEventArgs& Args)>(
 			[WeakPtr, this](FAsyncLoadEventArgs& Args)
@@ -2775,7 +2770,6 @@ void FAsyncLoadingThread::QueueEvent_ProcessPostloadWait(FAsyncPackage* Package,
 	FWeakAsyncPackagePtr WeakPtr(Package);
 	TAsyncLoadPriority UserPriority = Package->GetPriority();
 	int32 PackageSerialNumber = GRandomizeLoadOrder ? GetRandomSerialNumber() : Package->SerialNumber;
-	TRACE_LOADTIME_QUEUE_EVENT(Package, LoadTimeProfilerPackageEventType_PostLoadWait, UserPriority, PackageSerialNumber, EventSystemPriority);
 	EventQueue.AddAsyncEvent(UserPriority, PackageSerialNumber, EventSystemPriority,
 		TFunction<void(FAsyncLoadEventArgs& Args)>(
 			[WeakPtr, this](FAsyncLoadEventArgs& Args)
@@ -2798,7 +2792,6 @@ void FAsyncLoadingThread::QueueEvent_ExportsDone(FAsyncPackage* Package, int32 E
 	FWeakAsyncPackagePtr WeakPtr(Package);
 	TAsyncLoadPriority UserPriority = Package->GetPriority();
 	int32 PackageSerialNumber = GRandomizeLoadOrder ? GetRandomSerialNumber() : Package->SerialNumber;
-	TRACE_LOADTIME_QUEUE_EVENT(Package, LoadTimeProfilerPackageEventType_ExportsDone, UserPriority, PackageSerialNumber, EventSystemPriority);
 	EventQueue.AddAsyncEvent(UserPriority, PackageSerialNumber, EventSystemPriority,
 		TFunction<void(FAsyncLoadEventArgs& Args)>(
 			[WeakPtr, this](FAsyncLoadEventArgs& Args)
@@ -2821,7 +2814,6 @@ void FAsyncLoadingThread::QueueEvent_StartPostLoad(FAsyncPackage* Package, int32
 	FWeakAsyncPackagePtr WeakPtr(Package);
 	TAsyncLoadPriority UserPriority = Package->GetPriority();
 	int32 PackageSerialNumber = GRandomizeLoadOrder ? GetRandomSerialNumber() : Package->SerialNumber;
-	TRACE_LOADTIME_QUEUE_EVENT(Package, LoadTimeProfilerPackageEventType_StartPostLoad, UserPriority, PackageSerialNumber, EventSystemPriority);
 	EventQueue.AddAsyncEvent(Package->GetPriority(), GRandomizeLoadOrder ? GetRandomSerialNumber() : Package->SerialNumber, EventSystemPriority,
 		TFunction<void(FAsyncLoadEventArgs& Args)>(
 			[WeakPtr, this](FAsyncLoadEventArgs& Args)
@@ -4832,7 +4824,6 @@ EAsyncPackageState::Type FAsyncLoadingThread::ProcessLoadedPackages(bool bUseTim
 		FAsyncPackage* Package = LoadedPackagesToProcess[PackageIndex];
 		if (Package->GetDependencyRefCount() == 0)
 		{
-			TRACE_LOADTIME_ASYNC_PACKAGE_SCOPE(Package, LoadTimeProfilerPackageEventType_DeferredPostLoad);
 			SCOPED_LOADTIMER(ProcessLoadedPackagesTime);
 
 			Result = Package->PostLoadDeferredObjects(TickStartTime, bUseTimeLimit, TimeLimit);
@@ -5130,7 +5121,7 @@ FAsyncLoadingThread::FAsyncLoadingThread()
 	}
 
 #if LOADTIMEPROFILERTRACE_ENABLED
-	FLoadTimeProfilerTracePrivate::Init(FParse::Param(FCommandLine::Get(), TEXT("loadtimetrace")));
+	FLoadTimeProfilerTracePrivate::Init();
 #endif
 
 	QueuedRequestsEvent = FPlatformProcess::GetSynchEventFromPool();
@@ -6256,6 +6247,7 @@ void FAsyncPackage::AddImportDependency(const FName& PendingImport, FFlushTree*
 
 	if (!bReinsert)
 	{
+		TRACE_LOADTIME_ASYNC_PACKAGE_IMPORT_DEPENDENCY(this, PackageToStream);
 		FAsyncLoadingThread::Get().InsertPackage(PackageToStream, bReinsert);
 	}
 
@@ -7349,8 +7341,6 @@ void FlushAsyncLoading(int32 PackageID /* = INDEX_NONE */)
 			return;
 		}
 
-		TRACE_LOADTIME_FLUSH_ASYNCLOADING_SCOPE(PackageID);
-
 		FCoreDelegates::OnAsyncLoadingFlush.Broadcast();
 
 #if NO_LOGGING == 0
diff --git a/Engine/Source/Runtime/CoreUObject/Private/Serialization/LoadTimeTrace.cpp b/Engine/Source/Runtime/CoreUObject/Private/Serialization/LoadTimeTrace.cpp
--- a/Engine/Source/Runtime/CoreUObject/Private/Serialization/LoadTimeTrace.cpp
+++ b/Engine/Source/Runtime/CoreUObject/Private/Serialization/LoadTimeTrace.cpp
@@ -25,10 +25,6 @@ UE_TRACE_EVENT_BEGIN(LoadTime, ResumeAsyncLoading)
 	UE_TRACE_EVENT_FIELD(uint64, Cycle)
 UE_TRACE_EVENT_END()
 
-UE_TRACE_EVENT_BEGIN(LoadTime, NewPackage)
-	UE_TRACE_EVENT_FIELD(const UPackage*, Package)
-UE_TRACE_EVENT_END()
-
 UE_TRACE_EVENT_BEGIN(LoadTime, NewLinker)
 	UE_TRACE_EVENT_FIELD(const FLinkerLoad*, Linker)
 UE_TRACE_EVENT_END()
@@ -45,17 +41,6 @@ UE_TRACE_EVENT_BEGIN(LoadTime, PackageSummary)
 	UE_TRACE_EVENT_FIELD(uint32, ExportCount)
 UE_TRACE_EVENT_END()
 
-UE_TRACE_EVENT_BEGIN(LoadTime, BeginLinkerScope)
-	UE_TRACE_EVENT_FIELD(const FLinkerLoad*, Linker)
-	UE_TRACE_EVENT_FIELD(uint64, Cycle)
-	UE_TRACE_EVENT_FIELD(uint32, ThreadId)
-UE_TRACE_EVENT_END()
-
-UE_TRACE_EVENT_BEGIN(LoadTime, EndLinkerScope)
-	UE_TRACE_EVENT_FIELD(uint64, Cycle)
-	UE_TRACE_EVENT_FIELD(uint32, ThreadId)
-UE_TRACE_EVENT_END()
-
 UE_TRACE_EVENT_BEGIN(LoadTime, BeginCreateExport)
 	UE_TRACE_EVENT_FIELD(uint64, Cycle)
 	UE_TRACE_EVENT_FIELD(const FLinkerLoad*, Linker)
@@ -84,11 +69,6 @@ UE_TRACE_EVENT_BEGIN(LoadTime, EndObjectScope)
 	UE_TRACE_EVENT_FIELD(uint32, ThreadId)
 UE_TRACE_EVENT_END()
 
-UE_TRACE_EVENT_BEGIN(LoadTime, LinkerArchiveAssociation)
-	UE_TRACE_EVENT_FIELD(const FLinkerLoad*, Linker)
-	UE_TRACE_EVENT_FIELD(const FArchive*, Archive)
-UE_TRACE_EVENT_END()
-
 UE_TRACE_EVENT_BEGIN(LoadTime, BeginRequest)
 	UE_TRACE_EVENT_FIELD(uint64, Cycle)
 	UE_TRACE_EVENT_FIELD(uint64, RequestId)
@@ -100,46 +80,6 @@ UE_TRACE_EVENT_BEGIN(LoadTime, EndRequest)
 	UE_TRACE_EVENT_FIELD(uint64, RequestId)
 UE_TRACE_EVENT_END()
 
-UE_TRACE_EVENT_BEGIN(LoadTime, BeginLoadMap)
-	UE_TRACE_EVENT_FIELD(uint64, Cycle)
-	UE_TRACE_EVENT_FIELD(uint32, ThreadId)
-UE_TRACE_EVENT_END()
-
-UE_TRACE_EVENT_BEGIN(LoadTime, EndLoadMap)
-	UE_TRACE_EVENT_FIELD(uint64, Cycle)
-	UE_TRACE_EVENT_FIELD(uint32, ThreadId)
-UE_TRACE_EVENT_END()
-
-UE_TRACE_EVENT_BEGIN(LoadTime, NewStreamableHandle)
-	UE_TRACE_EVENT_FIELD(const FStreamableHandle*, StreamableHandle)
-	UE_TRACE_EVENT_FIELD(bool, IsCombined)
-UE_TRACE_EVENT_END()
-
-UE_TRACE_EVENT_BEGIN(LoadTime, DestroyStreamableHandle)
-	UE_TRACE_EVENT_FIELD(const FStreamableHandle*, StreamableHandle)
-UE_TRACE_EVENT_END()
-
-UE_TRACE_EVENT_BEGIN(LoadTime, BeginLoadStreamableHandle)
-	UE_TRACE_EVENT_FIELD(uint64, Cycle)
-	UE_TRACE_EVENT_FIELD(const FStreamableHandle*, StreamableHandle)
-	UE_TRACE_EVENT_END()
-
-UE_TRACE_EVENT_BEGIN(LoadTime, EndLoadStreamableHandle)
-	UE_TRACE_EVENT_FIELD(uint64, Cycle)
-	UE_TRACE_EVENT_FIELD(const FStreamableHandle*, StreamableHandle)
-	UE_TRACE_EVENT_END()
-
-UE_TRACE_EVENT_BEGIN(LoadTime, BeginWaitForStreamableHandle)
-	UE_TRACE_EVENT_FIELD(uint64, Cycle)
-	UE_TRACE_EVENT_FIELD(const FStreamableHandle*, StreamableHandle)
-	UE_TRACE_EVENT_FIELD(uint32, ThreadId)
-UE_TRACE_EVENT_END()
-
-UE_TRACE_EVENT_BEGIN(LoadTime, EndWaitForStreamableHandle)
-	UE_TRACE_EVENT_FIELD(uint64, Cycle)
-	UE_TRACE_EVENT_FIELD(uint32, ThreadId)
-UE_TRACE_EVENT_END()
-
 UE_TRACE_EVENT_BEGIN(LoadTime, NewAsyncPackage)
 	UE_TRACE_EVENT_FIELD(const FAsyncPackage*, AsyncPackage)
 UE_TRACE_EVENT_END()
@@ -148,19 +88,14 @@ UE_TRACE_EVENT_BEGIN(LoadTime, DestroyAsyncPackage)
 	UE_TRACE_EVENT_FIELD(const FAsyncPackage*, AsyncPackage)
 UE_TRACE_EVENT_END()
 
-UE_TRACE_EVENT_BEGIN(LoadTime, StreamableHandleRequestAssociation)
-	UE_TRACE_EVENT_FIELD(const FStreamableHandle*, StreamableHandle)
-	UE_TRACE_EVENT_FIELD(uint64, RequestId)
-UE_TRACE_EVENT_END()
-
 UE_TRACE_EVENT_BEGIN(LoadTime, AsyncPackageRequestAssociation)
 	UE_TRACE_EVENT_FIELD(const FAsyncPackage*, AsyncPackage)
 	UE_TRACE_EVENT_FIELD(uint64, RequestId)
 UE_TRACE_EVENT_END()
 
-UE_TRACE_EVENT_BEGIN(LoadTime, LoadImportDependency)
-	UE_TRACE_EVENT_FIELD(const FAsyncPackage*, Package)
-	UE_TRACE_EVENT_FIELD(const FAsyncPackage*, DependentPackage)
+UE_TRACE_EVENT_BEGIN(LoadTime, AsyncPackageImportDependency)
+	UE_TRACE_EVENT_FIELD(const FAsyncPackage*, AsyncPackage)
+	UE_TRACE_EVENT_FIELD(const FAsyncPackage*, ImportedAsyncPackage)
 UE_TRACE_EVENT_END()
 
 UE_TRACE_EVENT_BEGIN(LoadTime, AsyncPackageLinkerAssociation)
@@ -180,70 +115,34 @@ UE_TRACE_EVENT_BEGIN(LoadTime, EndAsyncPackageScope)
 	UE_TRACE_EVENT_FIELD(uint32, ThreadId)
 UE_TRACE_EVENT_END()
 
-UE_TRACE_EVENT_BEGIN(LoadTime, BeginFlushAsyncLoading)
-	UE_TRACE_EVENT_FIELD(uint64, Cycle)
-	UE_TRACE_EVENT_FIELD(uint64, RequestId)
-	UE_TRACE_EVENT_FIELD(uint32, ThreadId)
-UE_TRACE_EVENT_END()
-
-UE_TRACE_EVENT_BEGIN(LoadTime, EndFlushAsyncLoading)
-	UE_TRACE_EVENT_FIELD(uint64, Cycle)
-	UE_TRACE_EVENT_FIELD(uint32, ThreadId)
-UE_TRACE_EVENT_END()
-
-UE_TRACE_EVENT_BEGIN(LoadTime, QueueEvent)
-	UE_TRACE_EVENT_FIELD(uint64, Cycle)
-	UE_TRACE_EVENT_FIELD(const FAsyncPackage*, AsyncPackage)
-	UE_TRACE_EVENT_FIELD(uint32, ThreadId)
-	UE_TRACE_EVENT_FIELD(uint32, UserPriority)
-	UE_TRACE_EVENT_FIELD(uint32, PackageSerialNumber)
-	UE_TRACE_EVENT_FIELD(uint32, EventSystemPriority)
-	UE_TRACE_EVENT_FIELD(uint8, EventType)
-UE_TRACE_EVENT_END();
-
 UE_TRACE_EVENT_BEGIN(LoadTime, ClassInfo, Always)
 	UE_TRACE_EVENT_FIELD(const UClass*, Class)
 UE_TRACE_EVENT_END()
 
-void FLoadTimeProfilerTracePrivate::Init(bool bStartEnabled)
+void FLoadTimeProfilerTracePrivate::Init()
 {
+	FLoadTimeProfilerTrace::InitInternal();
 	UE_TRACE_EVENT_IS_ENABLED(LoadTime, StartAsyncLoading);
 	UE_TRACE_EVENT_IS_ENABLED(LoadTime, SuspendAsyncLoading);
 	UE_TRACE_EVENT_IS_ENABLED(LoadTime, ResumeAsyncLoading);
-	UE_TRACE_EVENT_IS_ENABLED(LoadTime, NewPackage);
 	UE_TRACE_EVENT_IS_ENABLED(LoadTime, NewLinker);
 	UE_TRACE_EVENT_IS_ENABLED(LoadTime, DestroyLinker);
 	UE_TRACE_EVENT_IS_ENABLED(LoadTime, PackageSummary);
-	UE_TRACE_EVENT_IS_ENABLED(LoadTime, BeginLinkerScope);
-	UE_TRACE_EVENT_IS_ENABLED(LoadTime, EndLinkerScope);
 	UE_TRACE_EVENT_IS_ENABLED(LoadTime, BeginCreateExport);
 	UE_TRACE_EVENT_IS_ENABLED(LoadTime, EndCreateExport);
 	UE_TRACE_EVENT_IS_ENABLED(LoadTime, BeginObjectScope);
 	UE_TRACE_EVENT_IS_ENABLED(LoadTime, EndObjectScope);
-	UE_TRACE_EVENT_IS_ENABLED(LoadTime, LinkerArchiveAssociation);
 	UE_TRACE_EVENT_IS_ENABLED(LoadTime, BeginRequest);
 	UE_TRACE_EVENT_IS_ENABLED(LoadTime, EndRequest);
-	UE_TRACE_EVENT_IS_ENABLED(LoadTime, BeginLoadMap);
-	UE_TRACE_EVENT_IS_ENABLED(LoadTime, EndLoadMap);
-	UE_TRACE_EVENT_IS_ENABLED(LoadTime, NewStreamableHandle);
-	UE_TRACE_EVENT_IS_ENABLED(LoadTime, DestroyStreamableHandle);
-	UE_TRACE_EVENT_IS_ENABLED(LoadTime, BeginLoadStreamableHandle);
-	UE_TRACE_EVENT_IS_ENABLED(LoadTime, EndLoadStreamableHandle);
-	UE_TRACE_EVENT_IS_ENABLED(LoadTime, BeginWaitForStreamableHandle);
-	UE_TRACE_EVENT_IS_ENABLED(LoadTime, EndWaitForStreamableHandle);
 	UE_TRACE_EVENT_IS_ENABLED(LoadTime, NewAsyncPackage);
 	UE_TRACE_EVENT_IS_ENABLED(LoadTime, DestroyAsyncPackage);
-	UE_TRACE_EVENT_IS_ENABLED(LoadTime, StreamableHandleRequestAssociation);
 	UE_TRACE_EVENT_IS_ENABLED(LoadTime, AsyncPackageRequestAssociation);
-	UE_TRACE_EVENT_IS_ENABLED(LoadTime, LoadImportDependency);
+	UE_TRACE_EVENT_IS_ENABLED(LoadTime, AsyncPackageImportDependency);
 	UE_TRACE_EVENT_IS_ENABLED(LoadTime, AsyncPackageLinkerAssociation);
 	UE_TRACE_EVENT_IS_ENABLED(LoadTime, BeginAsyncPackageScope);
 	UE_TRACE_EVENT_IS_ENABLED(LoadTime, EndAsyncPackageScope);
-	UE_TRACE_EVENT_IS_ENABLED(LoadTime, BeginFlushAsyncLoading);
-	UE_TRACE_EVENT_IS_ENABLED(LoadTime, EndFlushAsyncLoading);
-	UE_TRACE_EVENT_IS_ENABLED(LoadTime, QueueEvent);
 	UE_TRACE_EVENT_IS_ENABLED(LoadTime, ClassInfo);
-	if (bStartEnabled)
+	if (FParse::Param(FCommandLine::Get(), TEXT("loadtimetrace")))
 	{
 		Trace::ToggleEvent(TEXT("LoadTime"), true);
 	}
@@ -297,35 +196,6 @@ void FLoadTimeProfilerTracePrivate::OutputEndRequest(uint64 RequestId)
 		<< EndRequest.RequestId(RequestId);
 }
 
-void FLoadTimeProfilerTrace::OutputNewStreamableHandle(const FStreamableHandle* StreamableHandle, const TCHAR* DebugName, bool IsCombined)
-{
-	uint16 NameSize = (FCString::Strlen(DebugName) + 1) * sizeof(TCHAR);
-	UE_TRACE_LOG(LoadTime, NewStreamableHandle, NameSize)
-		<< NewStreamableHandle.StreamableHandle(StreamableHandle)
-		<< NewStreamableHandle.IsCombined(IsCombined)
-		<< NewStreamableHandle.Attachment(DebugName, NameSize);
-}
-
-void FLoadTimeProfilerTrace::OutputDestroyStreamableHandle(const FStreamableHandle* StreamableHandle)
-{
-	UE_TRACE_LOG(LoadTime, DestroyStreamableHandle)
-		<< DestroyStreamableHandle.StreamableHandle(StreamableHandle);
-}
-
-void FLoadTimeProfilerTrace::OutputBeginLoadStreamableHandle(const FStreamableHandle* StreamableHandle)
-{
-	UE_TRACE_LOG(LoadTime, BeginLoadStreamableHandle)
-		<< BeginLoadStreamableHandle.StreamableHandle(StreamableHandle)
-		<< BeginLoadStreamableHandle.Cycle(FPlatformTime::Cycles64());
-}
-
-void FLoadTimeProfilerTrace::OutputEndLoadStreamableHandle(const FStreamableHandle* StreamableHandle)
-{
-	UE_TRACE_LOG(LoadTime, EndLoadStreamableHandle)
-		<< EndLoadStreamableHandle.StreamableHandle(StreamableHandle)
-		<< EndLoadStreamableHandle.Cycle(FPlatformTime::Cycles64());
-}
-
 void FLoadTimeProfilerTracePrivate::OutputNewAsyncPackage(const FAsyncPackage* AsyncPackage, const TCHAR* PackageName)
 {
 	uint16 NameSize = (FCString::Strlen(PackageName) + 1) * sizeof(TCHAR);
@@ -350,13 +220,6 @@ void FLoadTimeProfilerTracePrivate::OutputPackageSummary(const FLinkerLoad* Link
 		<< PackageSummary.ExportCount(ExportCount);
 }
 
-void FLoadTimeProfilerTrace::OutputStreamableHandleRequestAssociation(const FStreamableHandle* StreamableHandle, uint64 RequestId)
-{
-	UE_TRACE_LOG(LoadTime, StreamableHandleRequestAssociation)
-		<< StreamableHandleRequestAssociation.StreamableHandle(StreamableHandle)
-		<< StreamableHandleRequestAssociation.RequestId(RequestId);
-}
-
 void FLoadTimeProfilerTracePrivate::OutputAsyncPackageRequestAssociation(const FAsyncPackage* AsyncPackage, uint64 RequestId)
 {
 	UE_TRACE_LOG(LoadTime, AsyncPackageRequestAssociation)
@@ -364,11 +227,11 @@ void FLoadTimeProfilerTracePrivate::OutputAsyncPackageRequestAssociation(const F
 		<< AsyncPackageRequestAssociation.RequestId(RequestId);
 }
 
-void FLoadTimeProfilerTracePrivate::OutputLoadImportDependency(const FAsyncPackage* Package, const FAsyncPackage* DependentPackage)
+void FLoadTimeProfilerTracePrivate::OutputAsyncPackageImportDependency(const FAsyncPackage* Package, const FAsyncPackage* ImportedPackage)
 {
-	UE_TRACE_LOG(LoadTime, LoadImportDependency)
-		<< LoadImportDependency.Package(Package)
-		<< LoadImportDependency.DependentPackage(DependentPackage);
+	UE_TRACE_LOG(LoadTime, AsyncPackageImportDependency)
+		<< AsyncPackageImportDependency.AsyncPackage(Package)
+		<< AsyncPackageImportDependency.ImportedAsyncPackage(ImportedPackage);
 }
 
 void FLoadTimeProfilerTracePrivate::OutputAsyncPackageLinkerAssociation(const FAsyncPackage* AsyncPackage, const FLinkerLoad* Linker)
@@ -378,25 +241,6 @@ void FLoadTimeProfilerTracePrivate::OutputAsyncPackageLinkerAssociation(const FA
 		<< AsyncPackageLinkerAssociation.Linker(Linker);
 }
 
-void FLoadTimeProfilerTracePrivate::OutputLinkerArchiveAssociation(const FLinkerLoad* Linker, const FArchive* Archive)
-{
-	UE_TRACE_LOG(LoadTime, LinkerArchiveAssociation)
-		<< LinkerArchiveAssociation.Linker(Linker)
-		<< LinkerArchiveAssociation.Archive(Archive);
-}
-
-void FLoadTimeProfilerTracePrivate::OutputQueueEvent(const FAsyncPackage* AsyncPackage, ELoadTimeProfilerPackageEventType EventType, TAsyncLoadPriority UserPriority, int32 PackageSerialNumber, int32 EventSystemPriority)
-{
-	UE_TRACE_LOG(LoadTime, QueueEvent)
-		<< QueueEvent.Cycle(FPlatformTime::Cycles64())
-		<< QueueEvent.AsyncPackage(AsyncPackage)
-		<< QueueEvent.ThreadId(FPlatformTLS::GetCurrentThreadId())
-		<< QueueEvent.UserPriority(UserPriority)
-		<< QueueEvent.PackageSerialNumber(PackageSerialNumber)
-		<< QueueEvent.EventSystemPriority(EventSystemPriority)
-		<< QueueEvent.EventType(EventType);
-}
-
 void FLoadTimeProfilerTracePrivate::OutputClassInfo(const UClass* Class, const TCHAR* Name)
 {
 	uint16 NameSize = (FCString::Strlen(Name) + 1) * sizeof(TCHAR);
@@ -458,65 +302,4 @@ FLoadTimeProfilerTracePrivate::FObjectScope::~FObjectScope()
 		<< EndObjectScope.ThreadId(FPlatformTLS::GetCurrentThreadId());
 }
 
-FLoadTimeProfilerTracePrivate::FLinkerScope::FLinkerScope(const FLinkerLoad* Linker)
-{
-	UE_TRACE_LOG(LoadTime, BeginLinkerScope)
-		<< BeginLinkerScope.Linker(Linker)
-		<< BeginLinkerScope.Cycle(FPlatformTime::Cycles64())
-		<< BeginLinkerScope.ThreadId(FPlatformTLS::GetCurrentThreadId());
-}
-
-FLoadTimeProfilerTracePrivate::FLinkerScope::~FLinkerScope()
-{
-	UE_TRACE_LOG(LoadTime, EndLinkerScope)
-		<< EndLinkerScope.Cycle(FPlatformTime::Cycles64())
-		<< EndLinkerScope.ThreadId(FPlatformTLS::GetCurrentThreadId());
-}
-
-FLoadTimeProfilerTracePrivate::FFlushAsyncLoadingScope::FFlushAsyncLoadingScope(uint64 RequestId)
-{
-	UE_TRACE_LOG(LoadTime, BeginFlushAsyncLoading)
-		<< BeginFlushAsyncLoading.Cycle(FPlatformTime::Cycles64())
-		<< BeginFlushAsyncLoading.RequestId(RequestId)
-		<< BeginFlushAsyncLoading.ThreadId(FPlatformTLS::GetCurrentThreadId());
-}
-
-FLoadTimeProfilerTracePrivate::FFlushAsyncLoadingScope::~FFlushAsyncLoadingScope()
-{
-	UE_TRACE_LOG(LoadTime, EndFlushAsyncLoading)
-		<< EndFlushAsyncLoading.Cycle(FPlatformTime::Cycles64())
-		<< EndFlushAsyncLoading.ThreadId(FPlatformTLS::GetCurrentThreadId());
-}
-
-FLoadTimeProfilerTrace::FWaitForStreamableHandleScope::FWaitForStreamableHandleScope(const FStreamableHandle* StreamableHandle)
-{
-	UE_TRACE_LOG(LoadTime, BeginWaitForStreamableHandle)
-		<< BeginWaitForStreamableHandle.Cycle(FPlatformTime::Cycles64())
-		<< BeginWaitForStreamableHandle.StreamableHandle(StreamableHandle)
-		<< BeginWaitForStreamableHandle.ThreadId(FPlatformTLS::GetCurrentThreadId());
-}
-
-FLoadTimeProfilerTrace::FWaitForStreamableHandleScope::~FWaitForStreamableHandleScope()
-{
-	UE_TRACE_LOG(LoadTime, EndWaitForStreamableHandle)
-		<< EndWaitForStreamableHandle.Cycle(FPlatformTime::Cycles64())
-		<< EndWaitForStreamableHandle.ThreadId(FPlatformTLS::GetCurrentThreadId());
-}
-
-FLoadTimeProfilerTrace::FLoadMapScope::FLoadMapScope(const TCHAR* Name)
-{
-	uint16 NameSize = (FCString::Strlen(Name) + 1) * sizeof(TCHAR);
-	UE_TRACE_LOG(LoadTime, BeginLoadMap, NameSize)
-		<< BeginLoadMap.Cycle(FPlatformTime::Cycles64())
-		<< BeginLoadMap.ThreadId(FPlatformTLS::GetCurrentThreadId())
-		<< BeginLoadMap.Attachment(Name, NameSize);
-}
-
-FLoadTimeProfilerTrace::FLoadMapScope::~FLoadMapScope()
-{
-	UE_TRACE_LOG(LoadTime, EndLoadMap)
-		<< EndLoadMap.Cycle(FPlatformTime::Cycles64())
-		<< EndLoadMap.ThreadId(FPlatformTLS::GetCurrentThreadId());
-}
-
 #endif
\ No newline at end of file
diff --git a/Engine/Source/Runtime/CoreUObject/Private/Serialization/LoadTimeTracePrivate.h b/Engine/Source/Runtime/CoreUObject/Private/Serialization/LoadTimeTracePrivate.h
--- a/Engine/Source/Runtime/CoreUObject/Private/Serialization/LoadTimeTracePrivate.h
+++ b/Engine/Source/Runtime/CoreUObject/Private/Serialization/LoadTimeTracePrivate.h
@@ -16,7 +16,7 @@ class UPackage;
 
 struct FLoadTimeProfilerTracePrivate
 {
-	static void Init(bool bStartEnabled);
+	static void Init();
 	static void OutputStartAsyncLoading();
 	static void OutputSuspendAsyncLoading();
 	static void OutputResumeAsyncLoading();
@@ -27,11 +27,9 @@ struct FLoadTimeProfilerTracePrivate
 	static void OutputBeginRequest(uint64 RequestId);
 	static void OutputEndRequest(uint64 RequestId);
 	static void OutputPackageSummary(const FLinkerLoad* Linker, uint32 TotalHeaderSize, uint32 NameCount, uint32 ImportCount, uint32 ExportCount);
-	static void OutputLoadImportDependency(const FAsyncPackage* Package, const FAsyncPackage* DependentPackage);
+	static void OutputAsyncPackageImportDependency(const FAsyncPackage* Package, const FAsyncPackage* ImportedPackage);
 	static void OutputAsyncPackageRequestAssociation(const FAsyncPackage* AsyncPackage, uint64 RequestId);
 	static void OutputAsyncPackageLinkerAssociation(const FAsyncPackage* AsyncPackage, const FLinkerLoad* Linker);
-	static void OutputLinkerArchiveAssociation(const FLinkerLoad* Linker, const FArchive* Archive);
-	static void OutputQueueEvent(const FAsyncPackage* AsyncPackage, ELoadTimeProfilerPackageEventType EventType, TAsyncLoadPriority UserPriority, int32 PackageSerialNumber, int32 EventSystemPriority);
 	static void OutputClassInfo(const UClass* Class, const TCHAR* Name);
 
 	struct FAsyncPackageScope
@@ -60,12 +58,6 @@ struct FLoadTimeProfilerTracePrivate
 		FLinkerScope(const FLinkerLoad* Linker);
 		~FLinkerScope();
 	};
-
-	struct FFlushAsyncLoadingScope
-	{
-		FFlushAsyncLoadingScope(uint64 RequestId);
-		~FFlushAsyncLoadingScope();
-	};
 };
 
 #define TRACE_LOADTIME_START_ASYNC_LOADING() \
@@ -86,9 +78,6 @@ struct FLoadTimeProfilerTracePrivate
 #define TRACE_LOADTIME_PACKAGE_SUMMARY(Linker, TotalHeaderSize, NameCount, ImportCount, ExportCount) \
 	FLoadTimeProfilerTracePrivate::OutputPackageSummary(Linker, TotalHeaderSize, NameCount, ImportCount, ExportCount);
 
-#define TRACE_LOADTIME_LINKER_ARCHIVE_ASSOCIATION(Linker, Archive) \
-	FLoadTimeProfilerTracePrivate::OutputLinkerArchiveAssociation(Linker, Archive);
-
 #define TRACE_LOADTIME_BEGIN_REQUEST(RequestId) \
 	FLoadTimeProfilerTracePrivate::OutputBeginRequest(RequestId);
 
@@ -107,8 +96,8 @@ struct FLoadTimeProfilerTracePrivate
 #define TRACE_LOADTIME_ASYNC_PACKAGE_LINKER_ASSOCIATION(AsyncPackage, Linker) \
 	FLoadTimeProfilerTracePrivate::OutputAsyncPackageLinkerAssociation(AsyncPackage, Linker);
 
-#define TRACE_LOADTIME_QUEUE_EVENT(AsyncPackage, EventType, UserPriority, PackageSerialNumber, EventSystemPriority)
-//	FLoadTimeProfilerTracePrivate::OutputQueueEvent(AsyncPackage, EventType, UserPriority, PackageSerialNumber, EventSystemPriority);
+#define TRACE_LOADTIME_ASYNC_PACKAGE_IMPORT_DEPENDENCY(AsyncPackage, ImportedAsyncPackage) \
+	FLoadTimeProfilerTracePrivate::OutputAsyncPackageImportDependency(AsyncPackage, ImportedAsyncPackage);
 
 #define TRACE_LOADTIME_ASYNC_PACKAGE_SCOPE(AsyncPackage, EventType) \
 	FLoadTimeProfilerTracePrivate::FAsyncPackageScope __LoadTimeTraceAsyncPackageScope(AsyncPackage, EventType);
@@ -119,9 +108,6 @@ struct FLoadTimeProfilerTracePrivate
 #define TRACE_LOADTIME_OBJECT_SCOPE(Object, EventType) \
 	FLoadTimeProfilerTracePrivate::FObjectScope __LoadTimeTraceObjectScope(Object, EventType);
 
-#define TRACE_LOADTIME_FLUSH_ASYNCLOADING_SCOPE(RequestId) \
-	FLoadTimeProfilerTracePrivate::FFlushAsyncLoadingScope __LoadTimeTraceFlushAsyncLoadingScope(RequestId);
-
 #define TRACE_LOADTIME_CLASS_INFO(Class, Name) \
 	FLoadTimeProfilerTracePrivate::OutputClassInfo(Class, Name);
 
@@ -134,15 +120,13 @@ struct FLoadTimeProfilerTracePrivate
 #define TRACE_LOADTIME_NEW_LINKER(...)
 #define TRACE_LOADTIME_DESTROY_LINKER(...)
 #define TRACE_LOADTIME_PACKAGE_SUMMARY(...)
-#define TRACE_LOADTIME_LINKER_ARCHIVE_ASSOCIATION(...)
 #define TRACE_LOADTIME_BEGIN_REQUEST(...)
 #define TRACE_LOADTIME_END_REQUEST(...)
 #define TRACE_LOADTIME_NEW_ASYNC_PACKAGE(...)
 #define TRACE_LOADTIME_DESTROY_ASYNC_PACKAGE(...)
 #define TRACE_LOADTIME_ASYNC_PACKAGE_REQUEST_ASSOCIATION(...)
 #define TRACE_LOADTIME_ASYNC_PACKAGE_LINKER_ASSOCIATION(...)
-#define TRACE_LOADTIME_FLUSH_ASYNCLOADING_SCOPE(...)
-#define TRACE_LOADTIME_QUEUE_EVENT(...)
+#define TRACE_LOADTIME_ASYNC_PACKAGE_IMPORT_DEPENDENCY(...)
 #define TRACE_LOADTIME_ASYNC_PACKAGE_SCOPE(...)
 #define TRACE_LOADTIME_TIMER_SCOPE(...)
 #define TRACE_LOADTIME_CREATE_EXPORT_SCOPE(...)
diff --git a/Engine/Source/Runtime/CoreUObject/Private/UObject/Class.cpp b/Engine/Source/Runtime/CoreUObject/Private/UObject/Class.cpp
--- a/Engine/Source/Runtime/CoreUObject/Private/UObject/Class.cpp
+++ b/Engine/Source/Runtime/CoreUObject/Private/UObject/Class.cpp
@@ -4035,6 +4035,11 @@ UClass::UClass(const FObjectInitializer& ObjectInitializer)
 	// If you add properties here, please update the other constructors and PurgeClass()
 
 	SetCppTypeInfoStatic(&DefaultCppClassTypeInfoStatic);
+#if LOADTIMEPROFILERTRACE_ENABLED
+	TCHAR Buffer[NAME_SIZE];
+	GetFName().GetPlainNameString(Buffer);
+	TRACE_LOADTIME_CLASS_INFO(this, Buffer);
+#endif
 }
 
 /**
diff --git a/Engine/Source/Runtime/CoreUObject/Private/UObject/LinkerLoad.cpp b/Engine/Source/Runtime/CoreUObject/Private/UObject/LinkerLoad.cpp
--- a/Engine/Source/Runtime/CoreUObject/Private/UObject/LinkerLoad.cpp
+++ b/Engine/Source/Runtime/CoreUObject/Private/UObject/LinkerLoad.cpp
@@ -1032,8 +1032,6 @@ FLinkerLoad::ELinkerStatus FLinkerLoad::CreateLoader(
 					Loader = IFileManager::Get().CreateFileReader(*Filename);
 				}
 
-				TRACE_LOADTIME_LINKER_ARCHIVE_ASSOCIATION(this, Loader);
-
 				if (!Loader)
 				{
 					UE_LOG(LogLinker, Warning, TEXT("Error opening file '%s'."), *Filename);
diff --git a/Engine/Source/Runtime/CoreUObject/Private/UObject/UObjectBase.cpp b/Engine/Source/Runtime/CoreUObject/Private/UObject/UObjectBase.cpp
--- a/Engine/Source/Runtime/CoreUObject/Private/UObject/UObjectBase.cpp
+++ b/Engine/Source/Runtime/CoreUObject/Private/UObject/UObjectBase.cpp
@@ -21,6 +21,7 @@
 #include "UObject/LinkerLoad.h"
 #include "Misc/CommandLine.h"
 #include "Interfaces/IPluginManager.h"
+#include "Serialization/LoadTimeTrace.h"
 
 DEFINE_LOG_CATEGORY_STATIC(LogUObjectBase, Log, All);
 DEFINE_STAT(STAT_UObjectsStatGroupTester);
@@ -867,6 +868,8 @@ void UClassReplaceHotReloadClasses()
  */
 static void UObjectLoadAllCompiledInDefaultProperties()
 {
+	TRACE_LOADTIME_REQUEST_GROUP_SCOPE(TEXT("UObjectLoadAllCompiledInDefaultProperties"));
+
 	static FName LongEnginePackageName(TEXT("/Script/Engine"));
 
 	TArray<UClass *(*)()>& DeferredCompiledInRegistration = GetDeferredCompiledInRegistration();
diff --git a/Engine/Source/Runtime/CoreUObject/Public/Serialization/LoadTimeTrace.h b/Engine/Source/Runtime/CoreUObject/Public/Serialization/LoadTimeTrace.h
deleted file mode 100644
diff --git a/Engine/Source/Runtime/Engine/Private/LevelStreaming.cpp b/Engine/Source/Runtime/Engine/Private/LevelStreaming.cpp
--- a/Engine/Source/Runtime/Engine/Private/LevelStreaming.cpp
+++ b/Engine/Source/Runtime/Engine/Private/LevelStreaming.cpp
@@ -30,6 +30,7 @@
 #include "SceneInterface.h"
 #include "Engine/NetDriver.h"
 #include "Engine/PackageMapClient.h"
+#include "Serialization/LoadTimeTrace.h"
 
 DEFINE_LOG_CATEGORY_STATIC(LogLevelStreaming, Log, All);
 
@@ -882,6 +883,8 @@ bool ULevelStreaming::RequestLevel(UWorld* PersistentWorld, bool bAllowLevelLoad
 		return false;
 	}
 
+	TRACE_LOADTIME_REQUEST_GROUP_SCOPE(TEXT("LevelStreaming - %s"), *GetPathName());
+
 	EPackageFlags PackageFlags = PKG_ContainsMap;
 	int32 PIEInstanceID = INDEX_NONE;
 
diff --git a/Engine/Source/Runtime/Engine/Private/StreamableManager.cpp b/Engine/Source/Runtime/Engine/Private/StreamableManager.cpp
--- a/Engine/Source/Runtime/Engine/Private/StreamableManager.cpp
+++ b/Engine/Source/Runtime/Engine/Private/StreamableManager.cpp
@@ -200,8 +200,6 @@ bool FStreamableHandle::BindUpdateDelegate(FStreamableUpdateDelegate NewDelegate
 
 EAsyncPackageState::Type FStreamableHandle::WaitUntilComplete(float Timeout, bool bStartStalledHandles)
 {
-	TRACE_LOADTIME_WAIT_FOR_STREAMABLE_HANDLE_SCOPE(this);
-
 	if (HasLoadCompleted())
 	{
 		return EAsyncPackageState::Complete;
@@ -506,7 +504,6 @@ void FStreamableHandle::StartStalledHandle()
 
 FStreamableHandle::~FStreamableHandle()
 {
-	TRACE_LOADTIME_DESTROY_STREAMABLE_HANDLE(this);
 	check(IsInGameThread() || IsInGarbageCollectorThread());
 
 	if (IsActive())
@@ -520,7 +517,6 @@ FStreamableHandle::~FStreamableHandle()
 
 void FStreamableHandle::CompleteLoad()
 {
-	TRACE_LOADTIME_END_LOAD_STREAMABLE_HANDLE(this);
 	// Only complete if it's still active
 	if (IsActive())
 	{
@@ -940,7 +936,6 @@ FStreamable* FStreamableManager::StreamInternal(const FSoftObjectPath& InTargetN
 			Existing->bAsyncLoadRequestOutstanding = true;
 			Existing->bLoadFailed = false;
 			int32 RequestId = LoadPackageAsync(Package, FLoadPackageAsyncDelegate::CreateSP(Handle, &FStreamableHandle::AsyncLoadCallbackWrapper, TargetName), Priority);
-			TRACE_LOADTIME_STREAMABLE_HANDLE_REQUEST_ASSOCIATION(&Handle.Get(), RequestId);
 		}
 	}
 	return Existing;
@@ -958,8 +953,6 @@ TSharedPtr<FStreamableHandle> FStreamableManager::RequestAsyncLoad(const TArray<
 	NewRequest->DebugName = DebugName;
 	NewRequest->Priority = Priority;
 
-	TRACE_LOADTIME_NEW_STREAMABLE_HANDLE(&NewRequest.Get(), *DebugName, false);
-
 	// Remove null requests
 
 	for (int32 i = NewRequest->RequestedAssets.Num() - 1; i >= 0 ; i--)
@@ -1089,7 +1082,7 @@ TSharedPtr<FStreamableHandle> FStreamableManager::RequestSyncLoad(const FSoftObj
 
 void FStreamableManager::StartHandleRequests(TSharedRef<FStreamableHandle> Handle)
 {
-	TRACE_LOADTIME_BEGIN_LOAD_STREAMABLE_HANDLE(&Handle.Get());
+	TRACE_LOADTIME_REQUEST_GROUP_SCOPE(TEXT("StreamableManager - %s"), *Handle->GetDebugName());
 
 	TArray<FStreamable *> ExistingStreamables;
 	ExistingStreamables.Reserve(Handle->RequestedAssets.Num());
@@ -1334,7 +1327,6 @@ TSharedPtr<FStreamableHandle> FStreamableManager::CreateCombinedHandle(const TAr
 	NewRequest->OwningManager = this;
 	NewRequest->bIsCombinedHandle = true;
 	NewRequest->DebugName = DebugName;
-	TRACE_LOADTIME_NEW_STREAMABLE_HANDLE(&NewRequest.Get(), *DebugName, true);
 
 	for (TSharedPtr<FStreamableHandle> ChildHandle : ChildHandles)
 	{
diff --git a/Engine/Source/Runtime/Engine/Private/UnrealEngine.cpp b/Engine/Source/Runtime/Engine/Private/UnrealEngine.cpp
--- a/Engine/Source/Runtime/Engine/Private/UnrealEngine.cpp
+++ b/Engine/Source/Runtime/Engine/Private/UnrealEngine.cpp
@@ -2319,6 +2319,7 @@ UEngineCustomTimeStep* InitializeCustomTimeStep(UEngine* InEngine, FSoftClassPat
 */
 void UEngine::InitializeObjectReferences()
 {
+	TRACE_LOADTIME_REQUEST_GROUP_SCOPE(TEXT("UEngine::InitializeObjectReferences"));
 	SCOPED_BOOT_TIMING("UEngine::InitializeObjectReferences");
 	DECLARE_SCOPE_CYCLE_COUNTER(TEXT("UEngine::InitializeObjectReferences"), STAT_InitializeObjectReferences, STATGROUP_LoadTime);
 
@@ -12088,7 +12089,7 @@ void UEngine::TickWorldTravel(FWorldContext& Context, float DeltaSeconds)
 
 bool UEngine::LoadMap( FWorldContext& WorldContext, FURL URL, class UPendingNetGame* Pending, FString& Error )
 {
-	TRACE_LOADTIME_LOAD_MAP_SCOPE(*URL.Map);
+	TRACE_LOADTIME_REQUEST_GROUP_SCOPE(TEXT("LoadMap - %s"), *URL.Map);
 	STAT_ADD_CUSTOMMESSAGE_NAME( STAT_NamedMarker, *(FString( TEXT( "LoadMap - " ) + URL.Map )) );
 	TRACE_BOOKMARK(TEXT("LoadMap - %s"), *URL.Map);
 
