UnrealInsights:
Fixed and enabled PlatformFileTrace on Windows
Exposed file activity as a table
Gracefully handle load time code running outside the main or async loading thread instead of crashing insights

diff --git a/Engine/Source/Developer/TraceInsights/Private/Insights/Widgets/STimingView.cpp b/Engine/Source/Developer/TraceInsights/Private/Insights/Widgets/STimingView.cpp
--- a/Engine/Source/Developer/TraceInsights/Private/Insights/Widgets/STimingView.cpp
+++ b/Engine/Source/Developer/TraceInsights/Private/Insights/Widgets/STimingView.cpp
@@ -1760,7 +1760,7 @@ void STimingView::UpdateIo()
 			FileActivityProvider.EnumerateFileActivity([this](const Trace::FFileInfo& FileInfo, const Trace::IFileActivityProvider::Timeline& Timeline)
 			{
 				Timeline.EnumerateEvents(-std::numeric_limits<double>::infinity(), +std::numeric_limits<double>::infinity(),
-					[this, &FileInfo, &Timeline](double EventStartTime, double EventEndTime, uint32 EventDepth, const Trace::FFileActivity& FileActivity)
+					[this, &FileInfo, &Timeline](double EventStartTime, double EventEndTime, uint32 EventDepth, const Trace::FFileActivity* FileActivity)
 				{
 					if (bMergeIoLanes)
 					{
@@ -1770,7 +1770,7 @@ void STimingView::UpdateIo()
 					{
 						EventDepth = FileInfo.Id % 32; // simple layout
 					}
-					uint32 Type = ((uint32)FileActivity.ActivityType & 0x0F) | (FileActivity.Failed ? 0x80 : 0);
+					uint32 Type = ((uint32)FileActivity->ActivityType & 0x0F) | (FileActivity->Failed ? 0x80 : 0);
 					AllIoEvents.Add(FIoTimingEvent{ EventStartTime, EventEndTime, EventDepth, Type, FileInfo.Path });
 				});
 
diff --git a/Engine/Source/Developer/TraceServices/Private/Analyzers/LoadTimeTraceAnalysis.cpp b/Engine/Source/Developer/TraceServices/Private/Analyzers/LoadTimeTraceAnalysis.cpp
--- a/Engine/Source/Developer/TraceServices/Private/Analyzers/LoadTimeTraceAnalysis.cpp
+++ b/Engine/Source/Developer/TraceServices/Private/Analyzers/LoadTimeTraceAnalysis.cpp
@@ -9,6 +9,8 @@
 #include "Common/FormatArgs.h"
 #include <limits>
 
+PRAGMA_DISABLE_OPTIMIZATION
+
 FAsyncLoadingTraceAnalyzer::FAsyncLoadingTraceAnalyzer(Trace::IAnalysisSession& InSession, Trace::FLoadTimeProfilerProvider& InLoadTimeProfilerProvider)
 	: Session(InSession)
 	, LoadTimeProfilerProvider(InLoadTimeProfilerProvider)
@@ -30,11 +32,15 @@ TSharedRef<FAsyncLoadingTraceAnalyzer::FThreadState> FAsyncLoadingTraceAnalyzer:
 		{
 			MainThreadId = ThreadId;
 			LoadTimeProfilerProvider.SetMainThreadId(MainThreadId);
-			ThreadState->CpuTimeline = LoadTimeProfilerProvider.EditMainThreadCpuTimeline();
+			ThreadState->CpuTimeline = &LoadTimeProfilerProvider.EditMainThreadCpuTimeline();
 		}
 		else if (ThreadId == AsyncLoadingThreadId)
 		{
-			ThreadState->CpuTimeline = LoadTimeProfilerProvider.EditAsyncLoadingThreadCpuTimeline();
+			ThreadState->CpuTimeline = &LoadTimeProfilerProvider.EditAsyncLoadingThreadCpuTimeline();
+		}
+		else
+		{
+			ThreadState->CpuTimeline = &LoadTimeProfilerProvider.EditAdditionalCpuTimeline(ThreadId);
 		}
 		return ThreadState;
 	}
@@ -353,6 +359,59 @@ void FAsyncLoadingTraceAnalyzer::OnEvent(uint16 RouteId, const FOnEventContext&
 	}
 }
 
+void FAsyncLoadingTraceAnalyzer::FThreadState::EnterPackageScope(double Time, const Trace::FPackageInfo* PackageInfo, ELoadTimeProfilerPackageEventType EventType)
+{
+	FScopeStackEntry& StackEntry = CpuScopeStack[CpuScopeStackDepth++];
+	StackEntry.Event.Export = nullptr;
+	StackEntry.Event.ExportEventType = LoadTimeProfilerObjectEventType_None;
+	StackEntry.Event.Package = PackageInfo;
+	StackEntry.Event.PackageEventType = EventType;
+	CurrentEvent = StackEntry.Event;
+	CpuTimeline->AppendBeginEvent(Time, StackEntry.Event);
+}
+
+void FAsyncLoadingTraceAnalyzer::FThreadState::EnterExportScope(double Time, const Trace::FPackageExportInfo* ExportInfo, ELoadTimeProfilerObjectEventType EventType)
+{
+	FScopeStackEntry& StackEntry = CpuScopeStack[CpuScopeStackDepth++];
+	StackEntry.Event.Export = ExportInfo;
+	StackEntry.Event.ExportEventType = EventType;
+	StackEntry.Event.Package = ExportInfo ? ExportInfo->Package : nullptr;
+	StackEntry.Event.PackageEventType = CurrentEvent.PackageEventType;
+	if (EventType == LoadTimeProfilerObjectEventType_PostLoad && StackEntry.Event.PackageEventType == LoadTimeProfilerPackageEventType_None)
+	{
+		StackEntry.Event.PackageEventType = LoadTimeProfilerPackageEventType_DeferredPostLoad;
+	}
+	CurrentEvent = StackEntry.Event;
+	CpuTimeline->AppendBeginEvent(Time, StackEntry.Event);
+}
+
+void FAsyncLoadingTraceAnalyzer::FThreadState::LeaveScope(double Time)
+{
+	FScopeStackEntry& StackEntry = CpuScopeStack[--CpuScopeStackDepth];
+	CpuTimeline->AppendEndEvent(Time);
+	if (CpuScopeStackDepth > 0)
+	{
+		CurrentEvent = CpuScopeStack[CpuScopeStackDepth - 1].Event;
+	}
+	else
+	{
+		CurrentEvent = Trace::FLoadTimeProfilerCpuEvent();
+	}
+}
+
+Trace::FPackageExportInfo* FAsyncLoadingTraceAnalyzer::FThreadState::GetCurrentExportScope()
+{
+	if (CpuScopeStackDepth > 0)
+	{
+		FScopeStackEntry& StackEntry = CpuScopeStack[CpuScopeStackDepth - 1];
+		return const_cast<Trace::FPackageExportInfo*>(StackEntry.Event.Export);
+	}
+	else
+	{
+		return nullptr;
+	}
+}
+
 void FAsyncLoadingTraceAnalyzer::PackageRequestAssociation(const FOnEventContext& Context, TSharedRef<FAsyncPackageState> AsyncPackageState, TSharedRef<FRequestState> RequestState)
 {
 	if (!AsyncPackageState->Request)
diff --git a/Engine/Source/Developer/TraceServices/Private/Analyzers/LoadTimeTraceAnalysis.h b/Engine/Source/Developer/TraceServices/Private/Analyzers/LoadTimeTraceAnalysis.h
--- a/Engine/Source/Developer/TraceServices/Private/Analyzers/LoadTimeTraceAnalysis.h
+++ b/Engine/Source/Developer/TraceServices/Private/Analyzers/LoadTimeTraceAnalysis.h
@@ -79,60 +79,12 @@ private:
 		Trace::FLoadTimeProfilerCpuEvent CurrentEvent;
 		TArray<TSharedPtr<FRequestGroupState>> RequestGroupStack;
 		
-		TSharedPtr<Trace::FLoadTimeProfilerProvider::CpuTimelineInternal> CpuTimeline;
-
-		void EnterPackageScope(double Time, const Trace::FPackageInfo* PackageInfo, ELoadTimeProfilerPackageEventType EventType)
-		{
-			FScopeStackEntry& StackEntry = CpuScopeStack[CpuScopeStackDepth++];
-			StackEntry.Event.Export = nullptr;
-			StackEntry.Event.ExportEventType = LoadTimeProfilerObjectEventType_None;
-			StackEntry.Event.Package = PackageInfo;
-			StackEntry.Event.PackageEventType = EventType;
-			CurrentEvent = StackEntry.Event;
-			CpuTimeline->AppendBeginEvent(Time, StackEntry.Event);
-		}
-
-		void EnterExportScope(double Time, const Trace::FPackageExportInfo* ExportInfo, ELoadTimeProfilerObjectEventType EventType)
-		{
-			FScopeStackEntry& StackEntry = CpuScopeStack[CpuScopeStackDepth++];
-			StackEntry.Event.Export = ExportInfo;
-			StackEntry.Event.ExportEventType = EventType;
-			StackEntry.Event.Package = ExportInfo ? ExportInfo->Package : nullptr;
-			StackEntry.Event.PackageEventType = CurrentEvent.PackageEventType;
-			if (EventType == LoadTimeProfilerObjectEventType_PostLoad && StackEntry.Event.PackageEventType == LoadTimeProfilerPackageEventType_None)
-			{
-				StackEntry.Event.PackageEventType = LoadTimeProfilerPackageEventType_DeferredPostLoad;
-			}
-			CurrentEvent = StackEntry.Event;
-			CpuTimeline->AppendBeginEvent(Time, StackEntry.Event);
-		}
-
-		void LeaveScope(double Time)
-		{
-			FScopeStackEntry& StackEntry = CpuScopeStack[--CpuScopeStackDepth];
-			CpuTimeline->AppendEndEvent(Time);
-			if (CpuScopeStackDepth > 0)
-			{
-				CurrentEvent = CpuScopeStack[CpuScopeStackDepth - 1].Event;
-			}
-			else
-			{
-				CurrentEvent = Trace::FLoadTimeProfilerCpuEvent();
-			}
-		}
-
-		Trace::FPackageExportInfo* GetCurrentExportScope()
-		{
-			if (CpuScopeStackDepth > 0)
-			{
-				FScopeStackEntry& StackEntry = CpuScopeStack[CpuScopeStackDepth - 1];
-				return const_cast<Trace::FPackageExportInfo*>(StackEntry.Event.Export);
-			}
-			else
-			{
-				return nullptr;
-			}
-		}
+		Trace::FLoadTimeProfilerProvider::CpuTimelineInternal* CpuTimeline;
+
+		void EnterPackageScope(double Time, const Trace::FPackageInfo* PackageInfo, ELoadTimeProfilerPackageEventType EventType);
+		void EnterExportScope(double Time, const Trace::FPackageExportInfo* ExportInfo, ELoadTimeProfilerObjectEventType EventType);
+		void LeaveScope(double Time);
+		Trace::FPackageExportInfo* GetCurrentExportScope();
 	};
 
 	void PackageRequestAssociation(const FOnEventContext& Context, TSharedRef<FAsyncPackageState> AsyncPackageState, TSharedRef<FRequestState> RequestState);
diff --git a/Engine/Source/Developer/TraceServices/Private/Analyzers/PlatformFileTraceAnalysis.cpp b/Engine/Source/Developer/TraceServices/Private/Analyzers/PlatformFileTraceAnalysis.cpp
--- a/Engine/Source/Developer/TraceServices/Private/Analyzers/PlatformFileTraceAnalysis.cpp
+++ b/Engine/Source/Developer/TraceServices/Private/Analyzers/PlatformFileTraceAnalysis.cpp
@@ -34,11 +34,10 @@ void FPlatformFileTraceAnalyzer::OnEvent(uint16 RouteId, const FOnEventContext&
 	case RouteId_BeginOpen:
 	{
 		double Time = Context.SessionContext.TimestampFromCycle(EventData.GetValue<uint64>("Cycle"));
-		uint64 TempHandle = EventData.GetValue<uint64>("TempHandle");
 		uint32 ThreadId = EventData.GetValue<uint32>("ThreadId");
-		check(!PendingOpenMap.Contains(TempHandle));
+		//check(!PendingOpenMap.Contains(ThreadId));
 		uint32 FileIndex = FileActivityProvider.GetFileIndex(reinterpret_cast<const TCHAR*>(EventData.GetAttachment()));
-		FPendingActivity& Open = PendingOpenMap.Add(TempHandle);
+		FPendingActivity& Open = PendingOpenMap.Add(ThreadId);
 		Open.ActivityIndex = FileActivityProvider.BeginActivity(FileIndex, Trace::FileActivityType_Open, 0, 0, Time);
 		Open.FileIndex = FileIndex;
 		break;
@@ -46,46 +45,50 @@ void FPlatformFileTraceAnalyzer::OnEvent(uint16 RouteId, const FOnEventContext&
 	case RouteId_EndOpen:
 	{
 		double Time = Context.SessionContext.TimestampFromCycle(EventData.GetValue<uint64>("Cycle"));
-		uint64 TempHandle = EventData.GetValue<uint64>("TempHandle");
 		uint64 FileHandle = EventData.GetValue<uint64>("FileHandle");
-		check(PendingOpenMap.Contains(TempHandle));
-		FPendingActivity& Open = PendingOpenMap[TempHandle];
-		if (FileHandle == uint64(-1)) // TODO: Portable invalid file handle?
-		{
-			FileActivityProvider.EndActivity(Open.FileIndex, Open.ActivityIndex, Time, true);
-		}
-		else
+		uint32 ThreadId = EventData.GetValue<uint32>("ThreadId");
+		const FPendingActivity* FindOpen = PendingOpenMap.Find(ThreadId);
+		if (FindOpen)
 		{
-			check(!OpenFilesMap.Contains(FileHandle));
-			OpenFilesMap.Add(FileHandle, Open.FileIndex);
-			FileActivityProvider.EndActivity(Open.FileIndex, Open.ActivityIndex, Time, false);
+			if (FileHandle == uint64(-1)) // TODO: Portable invalid file handle?
+			{
+				FileActivityProvider.EndActivity(FindOpen->FileIndex, FindOpen->ActivityIndex, 0, Time, true);
+			}
+			else
+			{
+				OpenFilesMap.Add(FileHandle, FindOpen->FileIndex);
+				FileActivityProvider.EndActivity(FindOpen->FileIndex, FindOpen->ActivityIndex, 0, Time, false);
+			}
+			PendingOpenMap.Remove(ThreadId);
 		}
-		PendingOpenMap.Remove(TempHandle);
 		break;
 	}
 	case RouteId_BeginClose:
 	{
 		double Time = Context.SessionContext.TimestampFromCycle(EventData.GetValue<uint64>("Cycle"));
-		uint64 TempHandle = EventData.GetValue<uint64>("TempHandle");
 		uint64 FileHandle = EventData.GetValue<uint64>("FileHandle");
 		uint32 ThreadId = EventData.GetValue<uint32>("ThreadId");
-		check(OpenFilesMap.Contains(FileHandle));
-		check(!PendingCloseMap.Contains(TempHandle));
-		uint64 FileIndex = OpenFilesMap[FileHandle];
-		OpenFilesMap.Remove(FileHandle);
-		FPendingActivity& Close = PendingCloseMap.Add(TempHandle);
-		Close.ActivityIndex = FileActivityProvider.BeginActivity(FileIndex, Trace::FileActivityType_Close, 0, 0, Time);
-		Close.FileIndex = FileIndex;
+		//check(!PendingCloseMap.Contains(ThreadId));
+		const uint32* FindFileIndex = OpenFilesMap.Find(FileHandle);
+		if (FindFileIndex)
+		{
+			OpenFilesMap.Remove(FileHandle);
+			FPendingActivity& Close = PendingCloseMap.Add(ThreadId);
+			Close.ActivityIndex = FileActivityProvider.BeginActivity(*FindFileIndex, Trace::FileActivityType_Close, 0, 0, Time);
+			Close.FileIndex = *FindFileIndex;
+		}
 		break;
 	}
 	case RouteId_EndClose:
 	{
 		double Time = Context.SessionContext.TimestampFromCycle(EventData.GetValue<uint64>("Cycle"));
-		uint64 TempHandle = EventData.GetValue<uint64>("TempHandle");
-		check(PendingCloseMap.Contains(TempHandle));
-		FPendingActivity& Close = PendingCloseMap[TempHandle];
-		FileActivityProvider.EndActivity(Close.FileIndex, Close.ActivityIndex, Time, false);
-		PendingCloseMap.Remove(TempHandle);
+		uint32 ThreadId = EventData.GetValue<uint32>("ThreadId");
+		const FPendingActivity* FindClose = PendingCloseMap.Find(ThreadId);
+		if (FindClose)
+		{
+			FileActivityProvider.EndActivity(FindClose->FileIndex, FindClose->ActivityIndex, 0, Time, false);
+			PendingCloseMap.Remove(ThreadId);
+		}
 		break;
 	}
 	case RouteId_BeginRead:
@@ -96,8 +99,17 @@ void FPlatformFileTraceAnalyzer::OnEvent(uint16 RouteId, const FOnEventContext&
 		uint64 Offset = EventData.GetValue<uint64>("Offset");
 		uint64 Size = EventData.GetValue<uint64>("Size");
 		uint32 ThreadId = EventData.GetValue<uint32>("ThreadId");
-		check(OpenFilesMap.Contains(FileHandle));
-		uint32 FileIndex = OpenFilesMap[FileHandle];
+		const uint32* FindFileIndex = OpenFilesMap.Find(FileHandle);
+		uint32 FileIndex;
+		if (FindFileIndex)
+		{
+			FileIndex = *FindFileIndex;
+		}
+		else
+		{
+			FileIndex = FileActivityProvider.GetUnknownFileIndex();
+			OpenFilesMap.Add(FileHandle, FileIndex);
+		}
 		uint64 ReadIndex = FileActivityProvider.BeginActivity(FileIndex, Trace::FileActivityType_Read, Offset, Size, Time);
 		FPendingActivity& Read = ActiveReadsMap.Add(ReadHandle);
 		Read.FileIndex = FileIndex;
@@ -110,10 +122,12 @@ void FPlatformFileTraceAnalyzer::OnEvent(uint16 RouteId, const FOnEventContext&
 		uint64 ReadHandle = EventData.GetValue<uint64>("ReadHandle");
 		uint64 SizeRead = EventData.GetValue<uint64>("SizeRead");
 		uint32 ThreadId = EventData.GetValue<uint32>("ThreadId");
-		check(ActiveReadsMap.Contains(ReadHandle));
-		const FPendingActivity& Read = ActiveReadsMap[ReadHandle];
-		FileActivityProvider.EndActivity(Read.FileIndex, Read.ActivityIndex, Time, false);
-		ActiveReadsMap.Remove(ReadHandle);
+		const FPendingActivity* FindRead = ActiveReadsMap.Find(ReadHandle);
+		if (FindRead)
+		{
+			FileActivityProvider.EndActivity(FindRead->FileIndex, FindRead->ActivityIndex, SizeRead, Time, false);
+			ActiveReadsMap.Remove(ReadHandle);
+		}
 		break;
 	}
 	case RouteId_BeginWrite:
@@ -124,8 +138,17 @@ void FPlatformFileTraceAnalyzer::OnEvent(uint16 RouteId, const FOnEventContext&
 		uint64 Offset = EventData.GetValue<uint64>("Offset");
 		uint64 Size = EventData.GetValue<uint64>("Size");
 		uint32 ThreadId = EventData.GetValue<uint32>("ThreadId");
-		check(OpenFilesMap.Contains(FileHandle));
-		uint32 FileIndex = OpenFilesMap[FileHandle];
+		const uint32* FindFileIndex = OpenFilesMap.Find(FileHandle);
+		uint32 FileIndex;
+		if (FindFileIndex)
+		{
+			FileIndex = *FindFileIndex;
+		}
+		else
+		{
+			FileIndex = FileActivityProvider.GetUnknownFileIndex();
+			OpenFilesMap.Add(FileHandle, FileIndex);
+		}
 		uint64 WriteIndex = FileActivityProvider.BeginActivity(FileIndex, Trace::FileActivityType_Write, Offset, Size, Time);
 		FPendingActivity& Write = ActiveWritesMap.Add(WriteHandle);
 		Write.FileIndex = FileIndex;
@@ -138,10 +161,12 @@ void FPlatformFileTraceAnalyzer::OnEvent(uint16 RouteId, const FOnEventContext&
 		uint64 WriteHandle = EventData.GetValue<uint64>("WriteHandle");
 		uint64 SizeWritten = EventData.GetValue<uint64>("SizeWritten");
 		uint32 ThreadId = EventData.GetValue<uint32>("ThreadId");
-		check(ActiveWritesMap.Contains(WriteHandle));
-		const FPendingActivity& Write = ActiveWritesMap[WriteHandle];
-		FileActivityProvider.EndActivity(Write.FileIndex, Write.ActivityIndex, Time, false);
-		ActiveWritesMap.Remove(WriteHandle);
+		const FPendingActivity* FindWrite = ActiveWritesMap.Find(WriteHandle);
+		if (FindWrite)
+		{
+			FileActivityProvider.EndActivity(FindWrite->FileIndex, FindWrite->ActivityIndex, SizeWritten, Time, false);
+			ActiveWritesMap.Remove(WriteHandle);
+		}
 		break;
 	}
 	}
diff --git a/Engine/Source/Developer/TraceServices/Private/Analyzers/PlatformFileTraceAnalysis.h b/Engine/Source/Developer/TraceServices/Private/Analyzers/PlatformFileTraceAnalysis.h
--- a/Engine/Source/Developer/TraceServices/Private/Analyzers/PlatformFileTraceAnalysis.h
+++ b/Engine/Source/Developer/TraceServices/Private/Analyzers/PlatformFileTraceAnalysis.h
@@ -46,8 +46,8 @@ private:
 	Trace::IAnalysisSession& Session;
 	Trace::FFileActivityProvider& FileActivityProvider;
 	TMap<uint64, uint32> OpenFilesMap;
-	TMap<uint64, FPendingActivity> PendingOpenMap;
-	TMap<uint64, FPendingActivity> PendingCloseMap;
+	TMap<uint32, FPendingActivity> PendingOpenMap;
+	TMap<uint32, FPendingActivity> PendingCloseMap;
 	TMap<uint64, FPendingActivity> ActiveReadsMap;
 	TMap<uint64, FPendingActivity> ActiveWritesMap;
 };
diff --git a/Engine/Source/Developer/TraceServices/Private/Model/FileActivity.cpp b/Engine/Source/Developer/TraceServices/Private/Model/FileActivity.cpp
--- a/Engine/Source/Developer/TraceServices/Private/Model/FileActivity.cpp
+++ b/Engine/Source/Developer/TraceServices/Private/Model/FileActivity.cpp
@@ -3,6 +3,7 @@
 #include "Model/FileActivity.h"
 #include "Model/IntervalTimeline.h"
 #include "AnalysisServicePrivate.h"
+#include <limits>
 
 namespace Trace
 {
@@ -10,54 +11,88 @@ namespace Trace
 FFileActivityProvider::FFileActivityProvider(IAnalysisSession& InSession)
 	: Session(InSession)
 	, Files(InSession.GetLinearAllocator(), 1024)
+	, FileActivities(InSession.GetLinearAllocator(), 1024)
+	, FileActivityTable(FileActivities)
 {
 }
 
-void FFileActivityProvider::EnumerateFileActivity(TFunctionRef<bool(const FFileInfo&, const ITimeline<FFileActivity>&)> Callback) const
+void FFileActivityProvider::EnumerateFileActivity(TFunctionRef<bool(const FFileInfo&, const Timeline&)> Callback) const
 {
 	for (uint64 FileIndex = 0, FileCount = Files.Num(); FileIndex < FileCount; ++FileIndex)
 	{
 		const FFileInfoInternal& FileInfoInternal = Files[FileIndex];
-		FFileInfo FileInfo;
-		FileInfo.Id = FileIndex;
-		FileInfo.Path = *FileInfoInternal.Path;
-		if (!Callback(FileInfo, *FileInfoInternal.ActivityTimeline))
+		if (!Callback(FileInfoInternal.FileInfo, *FileInfoInternal.ActivityTimeline))
 		{
 			return;
 		}
 	}
 }
 
+const ITable<FFileActivity>& FFileActivityProvider::GetFileActivityTable() const
+{
+	return FileActivityTable;
+}
+
 uint32 FFileActivityProvider::GetFileIndex(const TCHAR* Path)
 {
 	uint32 FileIndex = Files.Num();
 	FFileInfoInternal& FileInfo = Files.PushBack();
-	FileInfo.Path = Path;
+	FileInfo.FileInfo.Id = Files.Num() - 1;
+	FileInfo.FileInfo.Path = Session.StoreString(Path);
 	FileInfo.ActivityTimeline = MakeShared<TimelineInternal>(Session.GetLinearAllocator());
 	return FileIndex;
 }
 
+uint32 FFileActivityProvider::GetUnknownFileIndex()
+{
+	return GetFileIndex(TEXT("Unknown"));
+}
+
 uint64 FFileActivityProvider::BeginActivity(uint32 FileIndex, EFileActivityType Type, uint64 Offset, uint64 Size, double Time)
 {
 	FFileInfoInternal& FileInfo = Files[FileIndex];
-	FFileActivity FileActivity;
-	FileActivity.ActivityType = Type;
+	FFileActivity& FileActivity = FileActivities.PushBack();
+	FileActivity.File = &FileInfo.FileInfo;
 	FileActivity.Offset = Offset;
 	FileActivity.Size = Size;
+	FileActivity.StartTime = Time;
+	FileActivity.EndTime = std::numeric_limits<double>::infinity();
 	FileActivity.Failed = false;
-	return FileInfo.ActivityTimeline->AppendBeginEvent(Time, FileActivity);
+	FileActivity.ActivityType = Type;
+	return FileInfo.ActivityTimeline->AppendBeginEvent(Time, &FileActivity);
 }
 
-void FFileActivityProvider::EndActivity(uint32 FileIndex, uint64 ActivityIndex, double Time, bool Failed)
+void FFileActivityProvider::EndActivity(uint32 FileIndex, uint64 ActivityIndex, uint64 ActualSize, double Time, bool Failed)
 {
 	FFileInfoInternal& FileInfo = Files[FileIndex];
-	FFileActivity& Activity = FileInfo.ActivityTimeline->EndEvent(ActivityIndex, Time);
-	Activity.Failed = Failed;
+	FFileActivity* Activity = FileInfo.ActivityTimeline->EndEvent(ActivityIndex, Time);
+	Activity->EndTime = Time;
+	Activity->Failed = Failed;
+	if (!Failed)
+	{
+		Activity->Size = ActualSize;
+	}
 }
 
 const TCHAR* FFileActivityProvider::GetFilePath(uint32 FileIndex) const
 {
-	return *Files[FileIndex].Path;
+	return Files[FileIndex].FileInfo.Path;
+}
+
+const TCHAR* GetFileActivityTypeString(EFileActivityType ActivityType)
+{
+	switch (ActivityType)
+	{
+	case FileActivityType_Open:
+		return TEXT("Open");
+	case FileActivityType_Close:
+		return TEXT("Close");
+	case FileActivityType_Read:
+		return TEXT("Read");
+	case FileActivityType_Write:
+		return TEXT("Write");
+	}
+	return TEXT("Invalid");
 }
 
 }
\ No newline at end of file
diff --git a/Engine/Source/Developer/TraceServices/Private/Model/FileActivity.h b/Engine/Source/Developer/TraceServices/Private/Model/FileActivity.h
--- a/Engine/Source/Developer/TraceServices/Private/Model/FileActivity.h
+++ b/Engine/Source/Developer/TraceServices/Private/Model/FileActivity.h
@@ -8,6 +8,7 @@
 #include "TraceServices/Model/LoadTimeProfiler.h"
 #include "Containers/UnrealString.h"
 #include "Templates/SharedPointer.h"
+#include "Model/Tables.h"
 
 namespace Trace
 {
@@ -25,24 +26,38 @@ public:
 			EventsPerPage = 128
 		};
 	};
-	typedef TIntervalTimeline<FFileActivity, FTimelineSettings> TimelineInternal;
+	typedef TIntervalTimeline<FFileActivity*, FTimelineSettings> TimelineInternal;
 
 	FFileActivityProvider(IAnalysisSession& Session);
 	virtual void EnumerateFileActivity(TFunctionRef<bool(const FFileInfo&, const Timeline&)> Callback) const override;
+	virtual const ITable<FFileActivity>& GetFileActivityTable() const override;
 	uint32 GetFileIndex(const TCHAR* Path);
+	uint32 GetUnknownFileIndex();
 	uint64 BeginActivity(uint32 FileIndex, EFileActivityType Type, uint64 Offset, uint64 Size, double Time);
-	void EndActivity(uint32 FileIndex, uint64 ActivityIndex, double Time, bool Failed);
+	void EndActivity(uint32 FileIndex, uint64 ActivityIndex, uint64 ActualSize, double Time, bool Failed);
 	const TCHAR* GetFilePath(uint32 FileIndex) const;
 
 private:
 	struct FFileInfoInternal
 	{
-		FString Path;
+		FFileInfo FileInfo;
 		TSharedPtr<TimelineInternal> ActivityTimeline;
 	};
 
+	UE_TRACE_TABLE_LAYOUT_BEGIN(FFileActivityTableLayout, FFileActivity)
+		UE_TRACE_TABLE_PROJECTED_COLUMN(TableColumnType_CString, TEXT("File"), [](const FFileActivity& Row) { return Row.File ? Row.File->Path : TEXT("N/A"); })
+		UE_TRACE_TABLE_COLUMN(StartTime, TEXT("StartTime"))
+		UE_TRACE_TABLE_COLUMN(EndTime, TEXT("EndTime"))
+		UE_TRACE_TABLE_PROJECTED_COLUMN(TableColumnType_CString, TEXT("Type"), [](const FFileActivity& Row) { return GetFileActivityTypeString(Row.ActivityType); })
+		UE_TRACE_TABLE_COLUMN(Offset, TEXT("Offset"))
+		UE_TRACE_TABLE_COLUMN(Size, TEXT("Size"))
+		UE_TRACE_TABLE_COLUMN(Failed, TEXT("Failed"))
+	UE_TRACE_TABLE_LAYOUT_END()
+
 	IAnalysisSession& Session;
 	TPagedArray<FFileInfoInternal> Files;
+	TPagedArray<FFileActivity> FileActivities;
+	TTableView<FFileActivityTableLayout> FileActivityTable;
 };
 
 }
diff --git a/Engine/Source/Developer/TraceServices/Private/Model/LoadTimeProfiler.cpp b/Engine/Source/Developer/TraceServices/Private/Model/LoadTimeProfiler.cpp
--- a/Engine/Source/Developer/TraceServices/Private/Model/LoadTimeProfiler.cpp
+++ b/Engine/Source/Developer/TraceServices/Private/Model/LoadTimeProfiler.cpp
@@ -124,7 +124,7 @@ ITable<FPackagesTableRow>* FLoadTimeProfilerProvider::CreatePackageDetailsTable(
 				Row->SerializedExportsSize += Event.Export->SerialSize;
 			}
 
-			if (Event.PackageEventType == LoadTimeProfilerPackageEventType_CreateLinker)
+			if (!Event.Export && Event.PackageEventType == LoadTimeProfilerPackageEventType_CreateLinker)
 			{
 				Row->SerializedHeaderSize += Event.Package->Summary.TotalHeaderSize;
 			}
@@ -253,6 +253,20 @@ Trace::FPackageExportInfo& FLoadTimeProfilerProvider::CreateExport()
 	return Export;
 }
 
+FLoadTimeProfilerProvider::CpuTimelineInternal& FLoadTimeProfilerProvider::EditAdditionalCpuTimeline(uint32 ThreadId)
+{
+	if (AdditionalCpuTimelinesMap.Contains(ThreadId))
+	{
+		return AdditionalCpuTimelinesMap[ThreadId].Get();
+	}
+	else
+	{
+		TSharedRef<CpuTimelineInternal> Timeline = MakeShared<CpuTimelineInternal>(Session.GetLinearAllocator());
+		AdditionalCpuTimelinesMap.Add(ThreadId, Timeline);
+		return Timeline.Get();
+	}
+}
+
 uint64 FLoadTimeProfilerProvider::PackageSizeSum(const FLoadRequest& Row)
 {
 	uint64 Sum = 0;
diff --git a/Engine/Source/Developer/TraceServices/Private/Model/LoadTimeProfilerPrivate.h b/Engine/Source/Developer/TraceServices/Private/Model/LoadTimeProfilerPrivate.h
--- a/Engine/Source/Developer/TraceServices/Private/Model/LoadTimeProfilerPrivate.h
+++ b/Engine/Source/Developer/TraceServices/Private/Model/LoadTimeProfilerPrivate.h
@@ -33,8 +33,9 @@ public:
 	FLoadRequest& CreateRequest();
 	FPackageInfo& CreatePackage(const TCHAR* PackageName);
 	FPackageExportInfo& CreateExport();
-	TSharedRef<CpuTimelineInternal> EditMainThreadCpuTimeline() { return MainThreadCpuTimeline; }
-	TSharedRef<CpuTimelineInternal> EditAsyncLoadingThreadCpuTimeline() { return AsyncLoadingThreadCpuTimeline; }
+	CpuTimelineInternal& EditMainThreadCpuTimeline() { return MainThreadCpuTimeline.Get(); }
+	CpuTimelineInternal& EditAsyncLoadingThreadCpuTimeline() { return AsyncLoadingThreadCpuTimeline.Get(); }
+	CpuTimelineInternal& EditAdditionalCpuTimeline(uint32 ThreadId);
 	virtual uint32 GetMainThreadId() const override { return MainThreadId; }
 	void SetMainThreadId(uint32 ThreadId) { MainThreadId = ThreadId; }
 	virtual uint32 GetAsyncLoadingThreadId() const override { return AsyncLoadingThreadId; }
@@ -89,6 +90,7 @@ private:
 	TPagedArray<FPackageExportInfo> Exports;
 	TSharedRef<CpuTimelineInternal> MainThreadCpuTimeline;
 	TSharedRef<CpuTimelineInternal> AsyncLoadingThreadCpuTimeline;
+	TMap<uint32, TSharedRef<CpuTimelineInternal>> AdditionalCpuTimelinesMap;
 	uint32 MainThreadId = uint32(-1);
 	uint32 AsyncLoadingThreadId = uint32(-1);
 	TTableView<FRequestsTableLayout> RequestsTable;
diff --git a/Engine/Source/Developer/TraceServices/Public/TraceServices/Model/LoadTimeProfiler.h b/Engine/Source/Developer/TraceServices/Public/TraceServices/Model/LoadTimeProfiler.h
--- a/Engine/Source/Developer/TraceServices/Public/TraceServices/Model/LoadTimeProfiler.h
+++ b/Engine/Source/Developer/TraceServices/Public/TraceServices/Model/LoadTimeProfiler.h
@@ -13,8 +13,8 @@ namespace Trace
 
 struct FFileInfo
 {
-	uint32 Id;
-	const TCHAR* Path;
+	uint32 Id = uint32(-1);
+	const TCHAR* Path = nullptr;
 };
 
 enum EFileActivityType
@@ -25,24 +25,31 @@ enum EFileActivityType
 	FileActivityType_Write,
 
 	FileActivityType_Count,
+	FileActivityType_Invalid = FileActivityType_Count
 };
 
+TRACESERVICES_API const TCHAR* GetFileActivityTypeString(EFileActivityType ActivityType);
+
 struct FFileActivity
 {
-	uint64 Offset;
-	uint64 Size;
-	EFileActivityType ActivityType;
-	bool Failed;
+	const FFileInfo* File = nullptr;
+	double StartTime = 0.0;
+	double EndTime = 0.0;
+	uint64 Offset = 0;
+	uint64 Size = 0;
+	EFileActivityType ActivityType = FileActivityType_Invalid;
+	bool Failed = false;
 };
 
 class IFileActivityProvider
 	: public IProvider
 {
 public:
-	typedef ITimeline<FFileActivity> Timeline;
+	typedef ITimeline<FFileActivity*> Timeline;
 
 	virtual ~IFileActivityProvider() = default;
 	virtual void EnumerateFileActivity(TFunctionRef<bool(const FFileInfo&, const Timeline&)> Callback) const = 0;
+	virtual const ITable<FFileActivity>& GetFileActivityTable() const = 0;
 };
 
 struct FPackageSummaryInfo
diff --git a/Engine/Source/Runtime/Core/Private/ProfilingDebugging/PlatformFileTrace.cpp b/Engine/Source/Runtime/Core/Private/ProfilingDebugging/PlatformFileTrace.cpp
--- a/Engine/Source/Runtime/Core/Private/ProfilingDebugging/PlatformFileTrace.cpp
+++ b/Engine/Source/Runtime/Core/Private/ProfilingDebugging/PlatformFileTrace.cpp
@@ -6,32 +6,31 @@
 #include "Misc/CString.h"
 #include "HAL/PlatformTime.h"
 #include "HAL/PlatformTLS.h"
+#include "Misc/Parse.h"
 
-UE_TRACE_EVENT_BEGIN(PlatformFile, BeginOpen, Always)
+UE_TRACE_EVENT_BEGIN(PlatformFile, BeginOpen)
 	UE_TRACE_EVENT_FIELD(uint64, Cycle)
-	UE_TRACE_EVENT_FIELD(uint64, TempHandle)
 	UE_TRACE_EVENT_FIELD(uint32, ThreadId)
 UE_TRACE_EVENT_END()
 
-UE_TRACE_EVENT_BEGIN(PlatformFile, EndOpen, Always)
+UE_TRACE_EVENT_BEGIN(PlatformFile, EndOpen)
 	UE_TRACE_EVENT_FIELD(uint64, Cycle)
-	UE_TRACE_EVENT_FIELD(uint64, TempHandle)
 	UE_TRACE_EVENT_FIELD(uint64, FileHandle)
+	UE_TRACE_EVENT_FIELD(uint32, ThreadId)
 UE_TRACE_EVENT_END()
 
-UE_TRACE_EVENT_BEGIN(PlatformFile, BeginClose, Always)
+UE_TRACE_EVENT_BEGIN(PlatformFile, BeginClose)
 	UE_TRACE_EVENT_FIELD(uint64, Cycle)
-	UE_TRACE_EVENT_FIELD(uint64, TempHandle)
 	UE_TRACE_EVENT_FIELD(uint64, FileHandle)
 	UE_TRACE_EVENT_FIELD(uint32, ThreadId)
 UE_TRACE_EVENT_END()
 
-UE_TRACE_EVENT_BEGIN(PlatformFile, EndClose, Always)
+UE_TRACE_EVENT_BEGIN(PlatformFile, EndClose)
 	UE_TRACE_EVENT_FIELD(uint64, Cycle)
-	UE_TRACE_EVENT_FIELD(uint64, TempHandle)
+	UE_TRACE_EVENT_FIELD(uint32, ThreadId)
 UE_TRACE_EVENT_END()
 
-UE_TRACE_EVENT_BEGIN(PlatformFile, BeginRead, Always)
+UE_TRACE_EVENT_BEGIN(PlatformFile, BeginRead)
 	UE_TRACE_EVENT_FIELD(uint64, Cycle)
 	UE_TRACE_EVENT_FIELD(uint64, ReadHandle)
 	UE_TRACE_EVENT_FIELD(uint64, FileHandle)
@@ -40,14 +39,14 @@ UE_TRACE_EVENT_BEGIN(PlatformFile, BeginRead, Always)
 	UE_TRACE_EVENT_FIELD(uint32, ThreadId)
 UE_TRACE_EVENT_END()
 
-UE_TRACE_EVENT_BEGIN(PlatformFile, EndRead, Always)
+UE_TRACE_EVENT_BEGIN(PlatformFile, EndRead)
 	UE_TRACE_EVENT_FIELD(uint64, Cycle)
 	UE_TRACE_EVENT_FIELD(uint64, ReadHandle)
 	UE_TRACE_EVENT_FIELD(uint64, SizeRead)
 	UE_TRACE_EVENT_FIELD(uint32, ThreadId)
 UE_TRACE_EVENT_END()
 
-UE_TRACE_EVENT_BEGIN(PlatformFile, BeginWrite, Always)
+UE_TRACE_EVENT_BEGIN(PlatformFile, BeginWrite)
 	UE_TRACE_EVENT_FIELD(uint64, Cycle)
 	UE_TRACE_EVENT_FIELD(uint64, WriteHandle)
 	UE_TRACE_EVENT_FIELD(uint64, FileHandle)
@@ -56,45 +55,43 @@ UE_TRACE_EVENT_BEGIN(PlatformFile, BeginWrite, Always)
 	UE_TRACE_EVENT_FIELD(uint32, ThreadId)
 UE_TRACE_EVENT_END()
 
-UE_TRACE_EVENT_BEGIN(PlatformFile, EndWrite, Always)
+UE_TRACE_EVENT_BEGIN(PlatformFile, EndWrite)
 	UE_TRACE_EVENT_FIELD(uint64, Cycle)
 	UE_TRACE_EVENT_FIELD(uint64, WriteHandle)
 	UE_TRACE_EVENT_FIELD(uint64, SizeWritten)
 	UE_TRACE_EVENT_FIELD(uint32, ThreadId)
 UE_TRACE_EVENT_END()
 
-void FPlatformFileTrace::BeginOpen(uint64 TempHandle, const TCHAR* Path)
+void FPlatformFileTrace::BeginOpen(const TCHAR* Path)
 {
 	uint16 PathSize = (FCString::Strlen(Path) + 1) * sizeof(TCHAR);
 	UE_TRACE_LOG(PlatformFile, BeginOpen, PathSize)
 		<< BeginOpen.Cycle(FPlatformTime::Cycles64())
-		<< BeginOpen.TempHandle(TempHandle)
 		<< BeginOpen.Attachment(Path, PathSize)
 		<< BeginOpen.ThreadId(FPlatformTLS::GetCurrentThreadId());
 }
 
-void FPlatformFileTrace::EndOpen(uint64 TempHandle, uint64 FileHandle)
+void FPlatformFileTrace::EndOpen(uint64 FileHandle)
 {
 	UE_TRACE_LOG(PlatformFile, EndOpen)
 		<< EndOpen.Cycle(FPlatformTime::Cycles64())
-		<< EndOpen.TempHandle(TempHandle)
-		<< EndOpen.FileHandle(FileHandle);
+		<< EndOpen.FileHandle(FileHandle)
+		<< EndOpen.ThreadId(FPlatformTLS::GetCurrentThreadId());
 }
 
-void FPlatformFileTrace::BeginClose(uint64 TempHandle, uint64 FileHandle)
+void FPlatformFileTrace::BeginClose(uint64 FileHandle)
 {
 	UE_TRACE_LOG(PlatformFile, BeginClose)
 		<< BeginClose.Cycle(FPlatformTime::Cycles64())
-		<< BeginClose.TempHandle(TempHandle)
 		<< BeginClose.FileHandle(FileHandle)
 		<< BeginClose.ThreadId(FPlatformTLS::GetCurrentThreadId());
 }
 
-void FPlatformFileTrace::EndClose(uint64 TempHandle)
+void FPlatformFileTrace::EndClose()
 {
 	UE_TRACE_LOG(PlatformFile, EndClose)
 		<< EndClose.Cycle(FPlatformTime::Cycles64())
-		<< EndClose.TempHandle(TempHandle);
+		<< EndClose.ThreadId(FPlatformTLS::GetCurrentThreadId());
 }
 
 void FPlatformFileTrace::BeginRead(uint64 ReadHandle, uint64 FileHandle, uint64 Offset, uint64 Size)
@@ -137,4 +134,20 @@ void FPlatformFileTrace::EndWrite(uint64 WriteHandle, uint64 SizeWritten)
 		<< EndWrite.ThreadId(FPlatformTLS::GetCurrentThreadId());
 }
 
+void FPlatformFileTrace::Init(const TCHAR* CmdLine)
+{
+	if (FParse::Param(CmdLine, TEXT("filetrace")))
+	{
+		UE_TRACE_EVENT_IS_ENABLED(PlatformFile, BeginOpen);
+		UE_TRACE_EVENT_IS_ENABLED(PlatformFile, EndOpen);
+		UE_TRACE_EVENT_IS_ENABLED(PlatformFile, BeginClose);
+		UE_TRACE_EVENT_IS_ENABLED(PlatformFile, EndClose);
+		UE_TRACE_EVENT_IS_ENABLED(PlatformFile, BeginRead);
+		UE_TRACE_EVENT_IS_ENABLED(PlatformFile, EndRead);
+		UE_TRACE_EVENT_IS_ENABLED(PlatformFile, BeginWrite);
+		UE_TRACE_EVENT_IS_ENABLED(PlatformFile, EndWrite);
+		Trace::ToggleEvent(TEXT("PlatformFile"), true);
+	}
+}
+
 #endif
\ No newline at end of file
diff --git a/Engine/Source/Runtime/Core/Public/ProfilingDebugging/PlatformFileTrace.h b/Engine/Source/Runtime/Core/Public/ProfilingDebugging/PlatformFileTrace.h
--- a/Engine/Source/Runtime/Core/Public/ProfilingDebugging/PlatformFileTrace.h
+++ b/Engine/Source/Runtime/Core/Public/ProfilingDebugging/PlatformFileTrace.h
@@ -5,8 +5,8 @@
 #include "CoreTypes.h"
 #include "Trace/Trace.h"
 
-#if UE_TRACE_ENABLED && !UE_BUILD_SHIPPING
-#define PLATFORMFILETRACE_ENABLED 0
+#if UE_TRACE_ENABLED && PLATFORM_WINDOWS && !UE_BUILD_SHIPPING
+#define PLATFORMFILETRACE_ENABLED 1
 #else
 #define PLATFORMFILETRACE_ENABLED 0
 #endif
@@ -15,29 +15,32 @@
 
 struct FPlatformFileTrace
 {
-	static void BeginOpen(uint64 TempHandle, const TCHAR* Path);
-	static void EndOpen(uint64 TempHandle, uint64 FileHandle);
-	static void BeginClose(uint64 TempHandle, uint64 FileHandle);
-	static void EndClose(uint64 TempHandle);
+	CORE_API static void Init(const TCHAR* CmdLine);
+
+	static void BeginOpen(const TCHAR* Path);
+	static void EndOpen(uint64 FileHandle);
+	static void BeginClose(uint64 FileHandle);
+	static void EndClose();
 	static void BeginRead(uint64 ReadHandle, uint64 FileHandle, uint64 Offset, uint64 Size);
 	static void EndRead(uint64 ReadHandle, uint64 SizeRead);
 	static void BeginWrite(uint64 WriteHandle, uint64 FileHandle, uint64 Offset, uint64 Size);
 	static void EndWrite(uint64 WriteHandle, uint64 SizeWritten);
 };
 
+#define TRACE_PLATFORMFILE_INIT(CmdLine) \
+	FPlatformFileTrace::Init(CmdLine);
+
 #define TRACE_PLATFORMFILE_BEGIN_OPEN(Path) \
-	uint64 __TempHandle; \
-	FPlatformFileTrace::BeginOpen(uint64(&__TempHandle), Path);
+	FPlatformFileTrace::BeginOpen(Path);
 
 #define TRACE_PLATFORMFILE_END_OPEN(FileHandle) \
-	FPlatformFileTrace::EndOpen(uint64(&__TempHandle), uint64(FileHandle));
+	FPlatformFileTrace::EndOpen(uint64(FileHandle));
 
 #define TRACE_PLATFORMFILE_BEGIN_CLOSE(FileHandle) \
-	uint64 __TempHandle; \
-	FPlatformFileTrace::BeginClose(uint64(&__TempHandle), uint64(FileHandle));
+	FPlatformFileTrace::BeginClose(uint64(FileHandle));
 
 #define TRACE_PLATFORMFILE_END_CLOSE() \
-	FPlatformFileTrace::EndClose(uint64(&__TempHandle));
+	FPlatformFileTrace::EndClose();
 
 #define TRACE_PLATFORMFILE_BEGIN_READ(ReadHandle, FileHandle, Offset, Size) \
 	FPlatformFileTrace::BeginRead(uint64(ReadHandle), uint64(FileHandle), Offset, Size);
@@ -53,6 +56,7 @@ struct FPlatformFileTrace
 
 #else
 
+#define TRACE_PLATFORMFILE_INIT(CmdLine)
 #define TRACE_PLATFORMFILE_BEGIN_OPEN(Path)
 #define TRACE_PLATFORMFILE_END_OPEN(FileHandle)
 #define TRACE_PLATFORMFILE_BEGIN_CLOSE(FileHandle)
diff --git a/Engine/Source/Runtime/Launch/Private/LaunchEngineLoop.cpp b/Engine/Source/Runtime/Launch/Private/LaunchEngineLoop.cpp
--- a/Engine/Source/Runtime/Launch/Private/LaunchEngineLoop.cpp
+++ b/Engine/Source/Runtime/Launch/Private/LaunchEngineLoop.cpp
@@ -30,6 +30,7 @@
 #include "Stats/StatsMallocProfilerProxy.h"
 #include "Trace/Trace.h"
 #include "ProfilingDebugging/MiscTrace.h"
+#include "ProfilingDebugging/PlatformFileTrace.h"
 #if WITH_ENGINE
 #include "HAL/PlatformSplash.h"
 #endif
@@ -1161,6 +1162,7 @@ int32 FEngineLoop::PreInit(const TCHAR* CmdLine)
 #if CPUPROFILERTRACE_ENABLED
 	FCpuProfilerTrace::Init(FParse::Param(CmdLine, TEXT("cpuprofilertrace")));
 #endif
+	TRACE_PLATFORMFILE_INIT(CmdLine);
 
 	SCOPED_BOOT_TIMING("FEngineLoop::PreInit");
 
