Removed PRAGMA_DISABLE_OPTIMIZATIONS

diff --git a/Engine/Source/Developer/TraceServices/Private/Analyzers/LoadTimeTraceAnalysis.cpp b/Engine/Source/Developer/TraceServices/Private/Analyzers/LoadTimeTraceAnalysis.cpp
--- a/Engine/Source/Developer/TraceServices/Private/Analyzers/LoadTimeTraceAnalysis.cpp
+++ b/Engine/Source/Developer/TraceServices/Private/Analyzers/LoadTimeTraceAnalysis.cpp
@@ -9,8 +9,6 @@
 #include "Common/FormatArgs.h"
 #include <limits>
 
-PRAGMA_DISABLE_OPTIMIZATION
-
 FAsyncLoadingTraceAnalyzer::FAsyncLoadingTraceAnalyzer(Trace::IAnalysisSession& InSession, Trace::FLoadTimeProfilerProvider& InLoadTimeProfilerProvider)
 	: Session(InSession)
 	, LoadTimeProfilerProvider(InLoadTimeProfilerProvider)
