UnrealInsights:
Created CsvProfiler tracing module
Exposed CsvProfiler stats as counters

diff --git a/Engine/Source/Developer/TraceServices/Private/Analyzers/CsvProfilerTraceAnalysis.cpp b/Engine/Source/Developer/TraceServices/Private/Analyzers/CsvProfilerTraceAnalysis.cpp
new file mode 100644
--- /dev/null
+++ b/Engine/Source/Developer/TraceServices/Private/Analyzers/CsvProfilerTraceAnalysis.cpp
@@ -0,0 +1,151 @@
+// Copyright 1998-2019 Epic Games, Inc. All Rights Reserved.
+#include "CsvProfilerTraceAnalysis.h"
+#include "AnalysisServicePrivate.h"
+#include "TraceServices/Model/Counters.h"
+
+FCsvProfilerAnalyzer::FCsvProfilerAnalyzer(Trace::IAnalysisSession& InSession, Trace::ICounterProvider& InCounterProvider)
+	: Session(InSession)
+	, CounterProvider(InCounterProvider)
+{
+
+}
+
+void FCsvProfilerAnalyzer::OnAnalysisBegin(const FOnAnalysisContext& Context)
+{
+	auto& Builder = Context.InterfaceBuilder;
+
+	Builder.RouteEvent(RouteId_InlineStat, "CsvProfiler", "InlineStat");
+	Builder.RouteEvent(RouteId_DeclaredStat, "CsvProfiler", "DeclaredStat");
+	Builder.RouteEvent(RouteId_CustomStatInlineInt, "CsvProfiler", "CustomStatInlineInt");
+	Builder.RouteEvent(RouteId_CustomStatInlineFloat, "CsvProfiler", "CustomStatInlineFloat");
+	Builder.RouteEvent(RouteId_CustomStatDeclaredInt, "CsvProfiler", "CustomStatDeclaredInt");
+	Builder.RouteEvent(RouteId_CustomStatDeclaredFloat, "CsvProfiler", "CustomStatDeclaredFloat");
+}
+
+void FCsvProfilerAnalyzer::OnAnalysisEnd()
+{
+}
+
+void FCsvProfilerAnalyzer::OnEvent(uint16 RouteId, const FOnEventContext& Context)
+{
+	Trace::FAnalysisSessionEditScope _(Session);
+
+	const auto& EventData = Context.EventData;
+	switch (RouteId)
+	{
+	case RouteId_InlineStat:
+	{
+		uint64 StatNamePointer = EventData.GetValue<uint64>("StatNamePointer");
+		InlinePendingCountersMap.Add(StatNamePointer, Session.StoreString(ANSI_TO_TCHAR(reinterpret_cast<const ANSICHAR*>(EventData.GetAttachment()))));
+		break;
+	}
+	case RouteId_DeclaredStat:
+	{
+		uint32 StatId = EventData.GetValue<uint32>("StatId");
+		DeclaredPendingCountersMap.Add(StatId, Session.StoreString(reinterpret_cast<const TCHAR*>(EventData.GetAttachment())));
+		break;
+	}
+	case RouteId_CustomStatInlineInt:
+	{
+		Trace::ICounter* Counter = GetInlineStatCounter(EventData.GetValue<uint64>("StatNamePointer"), false);
+		HandleIntCounterEvent(Counter, Context);
+		break;
+	}
+	case RouteId_CustomStatInlineFloat:
+	{
+		Trace::ICounter* Counter = GetInlineStatCounter(EventData.GetValue<uint64>("StatNamePointer"), true);
+		HandleFloatCounterEvent(Counter, Context);
+		break;
+	}
+	case RouteId_CustomStatDeclaredInt:
+	{
+		Trace::ICounter* Counter = GetDeclaredStatCounter(EventData.GetValue<uint32>("StatId"), false);
+		HandleIntCounterEvent(Counter, Context);
+		break;
+	}
+	case RouteId_CustomStatDeclaredFloat:
+	{
+		Trace::ICounter* Counter = GetDeclaredStatCounter(EventData.GetValue<uint32>("StatId"), true);
+		HandleFloatCounterEvent(Counter, Context);
+		break;
+	}
+	}
+}
+
+void FCsvProfilerAnalyzer::HandleIntCounterEvent(Trace::ICounter* Counter, const FOnEventContext& Context)
+{
+	const auto& EventData = Context.EventData;
+	double Timestamp = Context.SessionContext.TimestampFromCycle(EventData.GetValue<uint64>("Cycle"));
+	uint8 OpType = EventData.GetValue<uint8>("OpType");
+	if (OpType == OpType_Accumulate)
+	{
+		Counter->AddValue(Timestamp, int64(EventData.GetValue<int32>("Value")));
+	}
+	else
+	{
+		Counter->SetValue(Timestamp, int64(EventData.GetValue<int32>("Value")));
+	}
+}
+
+void FCsvProfilerAnalyzer::HandleFloatCounterEvent(Trace::ICounter* Counter, const FOnEventContext& Context)
+{
+	const auto& EventData = Context.EventData;
+	double Timestamp = Context.SessionContext.TimestampFromCycle(EventData.GetValue<uint64>("Cycle"));
+	uint8 OpType = EventData.GetValue<uint8>("OpType");
+	if (OpType == OpType_Accumulate)
+	{
+		Counter->AddValue(Timestamp, double(EventData.GetValue<float>("Value")));
+	}
+	else
+	{
+		Counter->SetValue(Timestamp, double(EventData.GetValue<float>("Value")));
+	}
+}
+
+Trace::ICounter* FCsvProfilerAnalyzer::GetDeclaredStatCounter(uint32 StatId, bool bIsFloatingPoint)
+{
+	Trace::ICounter** FindIt;
+	FindIt = DeclaredCountersMap.Find(StatId);
+	if (!FindIt)
+	{
+		Trace::ICounter* Counter = CounterProvider.CreateCounter();
+		Counter->SetIsFloatingPoint(bIsFloatingPoint);
+		Counter->SetIsResetEveryFrame(true);
+		const TCHAR** FindName = DeclaredPendingCountersMap.Find(StatId);
+		if (FindName)
+		{
+			Counter->SetName(*FindName);
+			DeclaredPendingCountersMap.Remove(StatId);
+		}
+		DeclaredCountersMap.Add(StatId, Counter);
+		return Counter;
+	}
+	else
+	{
+		return *FindIt;
+	}
+}
+
+Trace::ICounter* FCsvProfilerAnalyzer::GetInlineStatCounter(uint64 StatNamePointer, bool bIsFloatingPoint)
+{
+	Trace::ICounter** FindIt;
+	FindIt = InlineCountersMap.Find(StatNamePointer);
+	if (!FindIt)
+	{
+		Trace::ICounter* Counter = CounterProvider.CreateCounter();
+		Counter->SetIsFloatingPoint(bIsFloatingPoint);
+		Counter->SetIsResetEveryFrame(true);
+		InlineCountersMap.Add(StatNamePointer, Counter);
+		const TCHAR** FindName = InlinePendingCountersMap.Find(StatNamePointer);
+		if (FindName)
+		{
+			Counter->SetName(*FindName);
+			InlinePendingCountersMap.Remove(StatNamePointer);
+		}
+		return Counter;
+	}
+	else
+	{
+		return *FindIt;
+	}
+}
diff --git a/Engine/Source/Developer/TraceServices/Private/Analyzers/CsvProfilerTraceAnalysis.h b/Engine/Source/Developer/TraceServices/Private/Analyzers/CsvProfilerTraceAnalysis.h
new file mode 100644
--- /dev/null
+++ b/Engine/Source/Developer/TraceServices/Private/Analyzers/CsvProfilerTraceAnalysis.h
@@ -0,0 +1,55 @@
+// Copyright 1998-2019 Epic Games, Inc. All Rights Reserved.
+
+#pragma once
+
+#include "Trace/Trace.h"
+#include "Trace/Analyzer.h"
+#include "Containers/Map.h"
+
+namespace Trace
+{
+	class IAnalysisSession;
+	class ICounter;
+	class ICounterProvider;
+}
+
+class FCsvProfilerAnalyzer
+	: public Trace::IAnalyzer
+{
+public:
+	FCsvProfilerAnalyzer(Trace::IAnalysisSession& Session, Trace::ICounterProvider& CounterProvider);
+	virtual void OnAnalysisBegin(const FOnAnalysisContext& Context) override;
+	virtual void OnEvent(uint16 RouteId, const FOnEventContext& Context) override;
+	virtual void OnAnalysisEnd() override;
+
+private:
+	enum : uint16
+	{
+		RouteId_InlineStat,
+		RouteId_DeclaredStat,
+		RouteId_CustomStatInlineInt,
+		RouteId_CustomStatInlineFloat,
+		RouteId_CustomStatDeclaredInt,
+		RouteId_CustomStatDeclaredFloat,
+	};
+
+	enum : uint8
+	{
+		OpType_Set,
+		OpType_Min,
+		OpType_Max,
+		OpType_Accumulate,
+	};
+
+	void HandleIntCounterEvent(Trace::ICounter* Counter, const FOnEventContext& Context);
+	void HandleFloatCounterEvent(Trace::ICounter* Counter, const FOnEventContext& Context);
+	Trace::ICounter* GetDeclaredStatCounter(uint32 StatId, bool bIsFloatingPoint);
+	Trace::ICounter* GetInlineStatCounter(uint64 StatNamePointer, bool bIsFloatingPoint);
+
+	Trace::IAnalysisSession& Session;
+	Trace::ICounterProvider& CounterProvider;
+	TMap<uint32, const TCHAR*> InlinePendingCountersMap;
+	TMap<uint32, Trace::ICounter*> InlineCountersMap;
+	TMap<uint64, const TCHAR*> DeclaredPendingCountersMap;
+	TMap<uint64, Trace::ICounter*> DeclaredCountersMap;
+};
diff --git a/Engine/Source/Developer/TraceServices/Private/Modules/CsvProfilerModule.cpp b/Engine/Source/Developer/TraceServices/Private/Modules/CsvProfilerModule.cpp
new file mode 100644
--- /dev/null
+++ b/Engine/Source/Developer/TraceServices/Private/Modules/CsvProfilerModule.cpp
@@ -0,0 +1,31 @@
+// Copyright 1998-2019 Epic Games, Inc. All Rights Reserved.
+
+#include "CsvProfilerModule.h"
+#include "Analyzers/CsvProfilerTraceAnalysis.h"
+#include "AnalysisServicePrivate.h"
+#include "TraceServices/ModuleService.h"
+
+namespace Trace
+{
+
+static const FName CsvProfilerModuleName("TraceModule_CsvProfiler");
+
+void FCsvProfilerModule::GetModuleInfo(FModuleInfo& OutModuleInfo)
+{
+	OutModuleInfo.Name = CsvProfilerModuleName;
+	OutModuleInfo.DisplayName = TEXT("CsvProfiler");
+}
+	
+
+void FCsvProfilerModule::OnAnalysisBegin(IAnalysisSession& Session)
+{
+	ICounterProvider& CounterProvider = EditCounterProvider(Session);
+	Session.AddAnalyzer(new FCsvProfilerAnalyzer(Session, CounterProvider));
+}
+
+void FCsvProfilerModule::GetLoggers(TArray<const TCHAR *>& OutLoggers)
+{
+	OutLoggers.Add(TEXT("CsvProfiler"));
+}
+
+}
diff --git a/Engine/Source/Developer/TraceServices/Private/Modules/CsvProfilerModule.h b/Engine/Source/Developer/TraceServices/Private/Modules/CsvProfilerModule.h
new file mode 100644
--- /dev/null
+++ b/Engine/Source/Developer/TraceServices/Private/Modules/CsvProfilerModule.h
@@ -0,0 +1,20 @@
+// Copyright 1998-2019 Epic Games, Inc. All Rights Reserved.
+
+#pragma once
+
+#include "TraceServices/ModuleService.h"
+#include "UObject/NameTypes.h"
+
+namespace Trace
+{
+
+class FCsvProfilerModule
+	: public IModule
+{
+public:
+	virtual void GetModuleInfo(FModuleInfo& OutModuleInfo) override;
+	virtual void OnAnalysisBegin(IAnalysisSession& Session) override;
+	virtual void GetLoggers(TArray<const TCHAR*>& OutLoggers) override;
+};
+
+}
diff --git a/Engine/Source/Developer/TraceServices/Private/TraceServicesModule.cpp b/Engine/Source/Developer/TraceServices/Private/TraceServicesModule.cpp
--- a/Engine/Source/Developer/TraceServices/Private/TraceServicesModule.cpp
+++ b/Engine/Source/Developer/TraceServices/Private/TraceServicesModule.cpp
@@ -9,7 +9,7 @@
 #include "Modules/TimingProfilerModule.h"
 #include "Modules/LoadTimeProfilerModule.h"
 #include "Modules/StatsModule.h"
-#include "Stats/StatsTrace.h"
+#include "Modules/CsvProfilerModule.h"
 
 class FTraceServicesModule
 	: public ITraceServicesModule
@@ -30,6 +30,7 @@ private:
 	Trace::FTimingProfilerModule TimingProfilerModule;
 	Trace::FLoadTimeProfilerModule LoadTimeProfilerModule;
 	Trace::FStatsModule StatsModule;
+	Trace::FCsvProfilerModule CsvProfilerModule;
 };
 
 TSharedPtr<Trace::ISessionService> FTraceServicesModule::GetSessionService()
@@ -68,10 +69,12 @@ void FTraceServicesModule::StartupModule()
 #if EXPERIMENTAL_STATSTRACE_ENABLED
 	IModularFeatures::Get().RegisterModularFeature(Trace::ModuleFeatureName, &StatsModule);
 #endif
+	IModularFeatures::Get().RegisterModularFeature(Trace::ModuleFeatureName, &CsvProfilerModule);
 }
 
 void FTraceServicesModule::ShutdownModule()
 {
+	IModularFeatures::Get().UnregisterModularFeature(Trace::ModuleFeatureName, &CsvProfilerModule);
 #if EXPERIMENTAL_STATSTRACE_ENABLED
 	IModularFeatures::Get().UnregisterModularFeature(Trace::ModuleFeatureName, &StatsModule);
 #endif
diff --git a/Engine/Source/Runtime/Core/Private/ProfilingDebugging/CsvProfiler.cpp b/Engine/Source/Runtime/Core/Private/ProfilingDebugging/CsvProfiler.cpp
--- a/Engine/Source/Runtime/Core/Private/ProfilingDebugging/CsvProfiler.cpp
+++ b/Engine/Source/Runtime/Core/Private/ProfilingDebugging/CsvProfiler.cpp
@@ -2983,6 +2983,7 @@ void FCsvProfiler::RecordEventAtTimestamp(int32 CategoryIndex, const FString& Ev
 
 void FCsvProfiler::RecordCustomStat(const char * StatName, uint32 CategoryIndex, float Value, const ECsvCustomStatOp CustomStatOp)
 {
+	TRACE_CSV_PROFILER_CUSTOM_STAT(StatName, Value, uint8(CustomStatOp));
 	if (GCsvProfilerIsCapturing && GCsvCategoriesEnabled[CategoryIndex])
 	{
 		LLM_SCOPE(ELLMTag::CsvProfiler);
@@ -2992,6 +2993,7 @@ void FCsvProfiler::RecordCustomStat(const char * StatName, uint32 CategoryIndex,
 
 void FCsvProfiler::RecordCustomStat(const FName& StatName, uint32 CategoryIndex, float Value, const ECsvCustomStatOp CustomStatOp)
 {
+	TRACE_CSV_PROFILER_CUSTOM_STAT(StatName, Value, uint8(CustomStatOp));
 	if (GCsvProfilerIsCapturing && GCsvCategoriesEnabled[CategoryIndex])
 	{
 		LLM_SCOPE(ELLMTag::CsvProfiler);
@@ -3001,6 +3003,7 @@ void FCsvProfiler::RecordCustomStat(const FName& StatName, uint32 CategoryIndex,
 
 void FCsvProfiler::RecordCustomStat(const char * StatName, uint32 CategoryIndex, int32 Value, const ECsvCustomStatOp CustomStatOp)
 {
+	TRACE_CSV_PROFILER_CUSTOM_STAT(StatName, Value, uint8(CustomStatOp));
 	if (GCsvProfilerIsCapturing && GCsvCategoriesEnabled[CategoryIndex])
 	{
 		LLM_SCOPE(ELLMTag::CsvProfiler);
@@ -3010,6 +3013,7 @@ void FCsvProfiler::RecordCustomStat(const char * StatName, uint32 CategoryIndex,
 
 void FCsvProfiler::RecordCustomStat(const FName& StatName, uint32 CategoryIndex, int32 Value, const ECsvCustomStatOp CustomStatOp)
 {
+	TRACE_CSV_PROFILER_CUSTOM_STAT(StatName, Value, uint8(CustomStatOp));
 	if (GCsvProfilerIsCapturing && GCsvCategoriesEnabled[CategoryIndex])
 	{
 		LLM_SCOPE(ELLMTag::CsvProfiler);
diff --git a/Engine/Source/Runtime/Core/Private/ProfilingDebugging/CsvProfilerTrace.cpp b/Engine/Source/Runtime/Core/Private/ProfilingDebugging/CsvProfilerTrace.cpp
new file mode 100644
--- /dev/null
+++ b/Engine/Source/Runtime/Core/Private/ProfilingDebugging/CsvProfilerTrace.cpp
@@ -0,0 +1,96 @@
+// Copyright 1998-2019 Epic Games, Inc. All Rights Reserved.
+#include "ProfilingDebugging/CsvProfilerTrace.h"
+#include "HAL/PlatformTime.h"
+#include "UObject/NameTypes.h"
+
+#if CSVPROFILERTRACE_ENABLED
+
+UE_TRACE_EVENT_BEGIN(CsvProfiler, InlineStat, Always)
+	UE_TRACE_EVENT_FIELD(const char*, StatNamePointer)
+UE_TRACE_EVENT_END()
+
+UE_TRACE_EVENT_BEGIN(CsvProfiler, DeclaredStat, Always)
+	UE_TRACE_EVENT_FIELD(uint32, StatId)
+UE_TRACE_EVENT_END()
+
+UE_TRACE_EVENT_BEGIN(CsvProfiler, CustomStatInlineInt)
+	UE_TRACE_EVENT_FIELD(uint64, Cycle)
+	UE_TRACE_EVENT_FIELD(const char*, StatNamePointer)
+	UE_TRACE_EVENT_FIELD(int32, Value)
+	UE_TRACE_EVENT_FIELD(uint8, OpType)
+UE_TRACE_EVENT_END()
+
+UE_TRACE_EVENT_BEGIN(CsvProfiler, CustomStatInlineFloat)
+	UE_TRACE_EVENT_FIELD(uint64, Cycle)
+	UE_TRACE_EVENT_FIELD(const char*, StatNamePointer)
+	UE_TRACE_EVENT_FIELD(float, Value)
+	UE_TRACE_EVENT_FIELD(uint8, OpType)
+UE_TRACE_EVENT_END()
+
+UE_TRACE_EVENT_BEGIN(CsvProfiler, CustomStatDeclaredInt)
+	UE_TRACE_EVENT_FIELD(uint64, Cycle)
+	UE_TRACE_EVENT_FIELD(uint32, StatId)
+	UE_TRACE_EVENT_FIELD(int32, Value)
+	UE_TRACE_EVENT_FIELD(uint8, OpType)
+UE_TRACE_EVENT_END()
+
+UE_TRACE_EVENT_BEGIN(CsvProfiler, CustomStatDeclaredFloat)
+	UE_TRACE_EVENT_FIELD(uint64, Cycle)
+	UE_TRACE_EVENT_FIELD(uint32, StatId)
+	UE_TRACE_EVENT_FIELD(float, Value)
+	UE_TRACE_EVENT_FIELD(uint8, OpType)
+UE_TRACE_EVENT_END()
+
+void FCsvProfilerTrace::OutputInlineStat(const char* StatName)
+{
+	uint16 NameSize = (strlen(StatName) + 1) * sizeof(char);
+	UE_TRACE_LOG(CsvProfiler, InlineStat, NameSize)
+		<< InlineStat.StatNamePointer(StatName)
+		<< InlineStat.Attachment(StatName, NameSize);
+}
+
+void FCsvProfilerTrace::OutputDeclaredStat(const FName& StatName, const TCHAR* StatNameString)
+{
+	uint16 NameSize = (FCString::Strlen(StatNameString) + 1) * sizeof(TCHAR);
+	UE_TRACE_LOG(CsvProfiler, DeclaredStat, NameSize)
+		<< DeclaredStat.StatId(StatName.GetComparisonIndex().ToUnstableInt())
+		<< DeclaredStat.Attachment(StatNameString, NameSize);
+}
+
+void FCsvProfilerTrace::OutputCustomStat(const char* StatName, int32 Value, uint8 OpType)
+{
+	UE_TRACE_LOG(CsvProfiler, CustomStatInlineInt)
+		<< CustomStatInlineInt.Cycle(FPlatformTime::Cycles64())
+		<< CustomStatInlineInt.StatNamePointer(StatName)
+		<< CustomStatInlineInt.Value(Value)
+		<< CustomStatInlineInt.OpType(OpType);
+}
+
+void FCsvProfilerTrace::OutputCustomStat(const FName& StatName, int32 Value, uint8 OpType)
+{
+	UE_TRACE_LOG(CsvProfiler, CustomStatDeclaredInt)
+		<< CustomStatDeclaredInt.Cycle(FPlatformTime::Cycles64())
+		<< CustomStatDeclaredInt.StatId(StatName.GetComparisonIndex().ToUnstableInt())
+		<< CustomStatDeclaredInt.Value(Value)
+		<< CustomStatDeclaredInt.OpType(OpType);
+}
+
+void FCsvProfilerTrace::OutputCustomStat(const char* StatName, float Value, uint8 OpType)
+{
+	UE_TRACE_LOG(CsvProfiler, CustomStatInlineFloat)
+		<< CustomStatInlineFloat.Cycle(FPlatformTime::Cycles64())
+		<< CustomStatInlineFloat.StatNamePointer(StatName)
+		<< CustomStatInlineFloat.Value(Value)
+		<< CustomStatInlineFloat.OpType(OpType);
+}
+
+void FCsvProfilerTrace::OutputCustomStat(const FName& StatName, float Value, uint8 OpType)
+{
+	UE_TRACE_LOG(CsvProfiler, CustomStatDeclaredFloat)
+		<< CustomStatDeclaredFloat.Cycle(FPlatformTime::Cycles64())
+		<< CustomStatDeclaredFloat.StatId(StatName.GetComparisonIndex().ToUnstableInt())
+		<< CustomStatDeclaredFloat.Value(Value)
+		<< CustomStatDeclaredFloat.OpType(OpType);
+}
+
+#endif
\ No newline at end of file
diff --git a/Engine/Source/Runtime/Core/Public/ProfilingDebugging/CsvProfiler.h b/Engine/Source/Runtime/Core/Public/ProfilingDebugging/CsvProfiler.h
--- a/Engine/Source/Runtime/Core/Public/ProfilingDebugging/CsvProfiler.h
+++ b/Engine/Source/Runtime/Core/Public/ProfilingDebugging/CsvProfiler.h
@@ -15,6 +15,7 @@
 #include "Async/TaskGraphInterfaces.h"
 #include "Misc/EnumClassFlags.h"
 #include "ProfilingDebugging/MiscTrace.h"
+#include "ProfilingDebugging/CsvProfilerTrace.h"
 
 // Whether to allow the CSV profiler in shipping builds.
 // Enable in a .Target.cs file if required.
@@ -47,21 +48,25 @@
 // Inline stats (no up front definition)
 #define CSV_SCOPED_TIMING_STAT(Category,StatName) \
 	FScopedCsvStat _ScopedCsvStat_ ## StatName (#StatName, CSV_CATEGORY_INDEX(Category)); \
-	TRACE_CPUPROFILER_EVENT_SCOPE_GROUP(StatName, CpuProfilerGroup_CsvProfiler)
+	TRACE_CPUPROFILER_EVENT_SCOPE_GROUP(StatName, CpuProfilerGroup_CsvProfiler);
 #define CSV_SCOPED_TIMING_STAT_GLOBAL(StatName) \
 	FScopedCsvStat _ScopedCsvStat_ ## StatName (#StatName, CSV_CATEGORY_INDEX_GLOBAL); \
-	TRACE_CPUPROFILER_EVENT_SCOPE_GROUP(StatName, CpuProfilerGroup_CsvProfiler)
+	TRACE_CPUPROFILER_EVENT_SCOPE_GROUP(StatName, CpuProfilerGroup_CsvProfiler);
 #define CSV_SCOPED_TIMING_STAT_EXCLUSIVE(StatName) \
 	FScopedCsvStatExclusive _ScopedCsvStatExclusive_ ## StatName (#StatName); \
-	TRACE_CPUPROFILER_EVENT_SCOPE_GROUP(StatName, CpuProfilerGroup_CsvProfiler)
+	TRACE_CPUPROFILER_EVENT_SCOPE_GROUP(StatName, CpuProfilerGroup_CsvProfiler);
 #define CSV_SCOPED_TIMING_STAT_EXCLUSIVE_CONDITIONAL(StatName,Condition) FScopedCsvStatExclusiveConditional _ScopedCsvStatExclusive_ ## StatName (#StatName,Condition);
 
 #define CSV_SCOPED_WAIT_CONDITIONAL(Condition)					FScopedCsvWaitConditional _ScopedCsvWait(Condition);
 #define CSV_SCOPED_SET_WAIT_STAT(StatName)						FScopedCsvSetWaitStat _ScopedCsvSetWaitStat ## StatName("EventWait/"#StatName);
 #define CSV_SCOPED_SET_WAIT_STAT_IGNORE()						FScopedCsvSetWaitStat _ScopedCsvSetWaitStat ## StatName();
 
-#define CSV_CUSTOM_STAT(Category,StatName,Value,Op)				FCsvProfiler::RecordCustomStat(#StatName, CSV_CATEGORY_INDEX(Category), Value, Op)
-#define CSV_CUSTOM_STAT_GLOBAL(StatName,Value,Op) 				FCsvProfiler::RecordCustomStat(#StatName, CSV_CATEGORY_INDEX_GLOBAL, Value, Op)
+#define CSV_CUSTOM_STAT(Category,StatName,Value,Op) \
+	TRACE_CSV_PROFILER_INLINE_STAT(#StatName); \
+	FCsvProfiler::RecordCustomStat(#StatName, CSV_CATEGORY_INDEX(Category), Value, Op);
+#define CSV_CUSTOM_STAT_GLOBAL(StatName,Value,Op) \
+	TRACE_CSV_PROFILER_INLINE_STAT(#StatName); \
+	FCsvProfiler::RecordCustomStat(#StatName, CSV_CATEGORY_INDEX_GLOBAL, Value, Op); 
 
 // Stats declared up front
 #define CSV_DEFINE_STAT(Category,StatName)						FCsvDeclaredStat _GCsvStat_##StatName((TCHAR*)TEXT(#StatName), CSV_CATEGORY_INDEX(Category));
@@ -143,7 +148,9 @@ struct FCsvDeclaredStat
 	FCsvDeclaredStat(TCHAR* InNameString, uint32 InCategoryIndex) 
 		: Name(InNameString)
 		, CategoryIndex(InCategoryIndex) 
-	{}
+	{
+		TRACE_CSVPROFILER_DECLARED_STAT(Name, InNameString);
+	}
 
 	FName Name;
 	uint32 CategoryIndex;
diff --git a/Engine/Source/Runtime/Core/Public/ProfilingDebugging/CsvProfilerTrace.h b/Engine/Source/Runtime/Core/Public/ProfilingDebugging/CsvProfilerTrace.h
new file mode 100644
--- /dev/null
+++ b/Engine/Source/Runtime/Core/Public/ProfilingDebugging/CsvProfilerTrace.h
@@ -0,0 +1,49 @@
+// Copyright 1998-2019 Epic Games, Inc. All Rights Reserved.
+
+#pragma once
+
+#include "CoreTypes.h"
+#include "Trace/Trace.h"
+
+#if UE_TRACE_ENABLED && !UE_BUILD_SHIPPING
+#define CSVPROFILERTRACE_ENABLED 1
+#else
+#define CSVPROFILERTRACE_ENABLED 0
+#endif
+
+#if CSVPROFILERTRACE_ENABLED
+
+struct FCsvProfilerTrace
+{
+	CORE_API static void OutputInlineStat(const char* StatName);
+	CORE_API static void OutputDeclaredStat(const FName& StatName, const TCHAR* StatNameString);
+	CORE_API static void OutputCustomStat(const char* StatName, int32 Value, uint8 OpType);
+	CORE_API static void OutputCustomStat(const FName& StatName, int32 Value, uint8 OpType);
+	CORE_API static void OutputCustomStat(const char* StatName, float Value, uint8 OpType);
+	CORE_API static void OutputCustomStat(const FName& StatName, float Value, uint8 OpType);
+};
+
+#define TRACE_CSV_PROFILER_INLINE_STAT(StatName) \
+	static bool PREPROCESSOR_JOIN(__CsvProfilerStat, __LINE__); \
+	if (!PREPROCESSOR_JOIN(__CsvProfilerStat, __LINE__)) { \
+		FCsvProfilerTrace::OutputInlineStat(StatName); \
+		PREPROCESSOR_JOIN(__CsvProfilerStat, __LINE__) = true; \
+	}
+
+#define TRACE_CSVPROFILER_DECLARED_STAT(StatName, StatNameString) \
+	static bool PREPROCESSOR_JOIN(__CsvProfilerStat, __LINE__); \
+	if (!PREPROCESSOR_JOIN(__CsvProfilerStat, __LINE__)) { \
+		FCsvProfilerTrace::OutputDeclaredStat(StatName, StatNameString); \
+		PREPROCESSOR_JOIN(__CsvProfilerStat, __LINE__) = true; \
+	}
+
+#define TRACE_CSV_PROFILER_CUSTOM_STAT(StatName, Value, OpType) \
+	FCsvProfilerTrace::OutputCustomStat(StatName, Value, OpType);
+
+#else
+
+#define TRACE_CSV_PROFILER_INLINE_STAT(...)
+#define TRACE_CSVPROFILER_DECLARED_STAT(...)
+#define TRACE_CSV_PROFILER_CUSTOM_STAT(...)
+
+#endif
\ No newline at end of file
