UnrealInsights:
Added ThreadId to file activities
Fixed SetColor log messages showing up in the log view

diff --git a/Engine/Source/Developer/TraceServices/Private/Analyzers/PlatformFileTraceAnalysis.cpp b/Engine/Source/Developer/TraceServices/Private/Analyzers/PlatformFileTraceAnalysis.cpp
--- a/Engine/Source/Developer/TraceServices/Private/Analyzers/PlatformFileTraceAnalysis.cpp
+++ b/Engine/Source/Developer/TraceServices/Private/Analyzers/PlatformFileTraceAnalysis.cpp
@@ -38,7 +38,7 @@ void FPlatformFileTraceAnalyzer::OnEvent(uint16 RouteId, const FOnEventContext&
 		//check(!PendingOpenMap.Contains(ThreadId));
 		uint32 FileIndex = FileActivityProvider.GetFileIndex(reinterpret_cast<const TCHAR*>(EventData.GetAttachment()));
 		FPendingActivity& Open = PendingOpenMap.Add(ThreadId);
-		Open.ActivityIndex = FileActivityProvider.BeginActivity(FileIndex, Trace::FileActivityType_Open, 0, 0, Time);
+		Open.ActivityIndex = FileActivityProvider.BeginActivity(FileIndex, Trace::FileActivityType_Open, ThreadId, 0, 0, Time);
 		Open.FileIndex = FileIndex;
 		break;
 	}
@@ -74,7 +74,7 @@ void FPlatformFileTraceAnalyzer::OnEvent(uint16 RouteId, const FOnEventContext&
 		{
 			OpenFilesMap.Remove(FileHandle);
 			FPendingActivity& Close = PendingCloseMap.Add(ThreadId);
-			Close.ActivityIndex = FileActivityProvider.BeginActivity(*FindFileIndex, Trace::FileActivityType_Close, 0, 0, Time);
+			Close.ActivityIndex = FileActivityProvider.BeginActivity(*FindFileIndex, Trace::FileActivityType_Close, ThreadId, 0, 0, Time);
 			Close.FileIndex = *FindFileIndex;
 		}
 		break;
@@ -110,7 +110,7 @@ void FPlatformFileTraceAnalyzer::OnEvent(uint16 RouteId, const FOnEventContext&
 			FileIndex = FileActivityProvider.GetUnknownFileIndex();
 			OpenFilesMap.Add(FileHandle, FileIndex);
 		}
-		uint64 ReadIndex = FileActivityProvider.BeginActivity(FileIndex, Trace::FileActivityType_Read, Offset, Size, Time);
+		uint64 ReadIndex = FileActivityProvider.BeginActivity(FileIndex, Trace::FileActivityType_Read, ThreadId, Offset, Size, Time);
 		FPendingActivity& Read = ActiveReadsMap.Add(ReadHandle);
 		Read.FileIndex = FileIndex;
 		Read.ActivityIndex = ReadIndex;
@@ -149,7 +149,7 @@ void FPlatformFileTraceAnalyzer::OnEvent(uint16 RouteId, const FOnEventContext&
 			FileIndex = FileActivityProvider.GetUnknownFileIndex();
 			OpenFilesMap.Add(FileHandle, FileIndex);
 		}
-		uint64 WriteIndex = FileActivityProvider.BeginActivity(FileIndex, Trace::FileActivityType_Write, Offset, Size, Time);
+		uint64 WriteIndex = FileActivityProvider.BeginActivity(FileIndex, Trace::FileActivityType_Write, ThreadId, Offset, Size, Time);
 		FPendingActivity& Write = ActiveWritesMap.Add(WriteHandle);
 		Write.FileIndex = FileIndex;
 		Write.ActivityIndex = WriteIndex;
diff --git a/Engine/Source/Developer/TraceServices/Private/Analyzers/PlatformFileTraceAnalysis.h b/Engine/Source/Developer/TraceServices/Private/Analyzers/PlatformFileTraceAnalysis.h
--- a/Engine/Source/Developer/TraceServices/Private/Analyzers/PlatformFileTraceAnalysis.h
+++ b/Engine/Source/Developer/TraceServices/Private/Analyzers/PlatformFileTraceAnalysis.h
@@ -41,6 +41,7 @@ private:
 	{
 		uint64 ActivityIndex;
 		uint32 FileIndex;
+		uint32 ThreadId;
 	};
 
 	Trace::IAnalysisSession& Session;
diff --git a/Engine/Source/Developer/TraceServices/Private/Model/FileActivity.cpp b/Engine/Source/Developer/TraceServices/Private/Model/FileActivity.cpp
--- a/Engine/Source/Developer/TraceServices/Private/Model/FileActivity.cpp
+++ b/Engine/Source/Developer/TraceServices/Private/Model/FileActivity.cpp
@@ -48,7 +48,7 @@ uint32 FFileActivityProvider::GetUnknownFileIndex()
 	return GetFileIndex(TEXT("Unknown"));
 }
 
-uint64 FFileActivityProvider::BeginActivity(uint32 FileIndex, EFileActivityType Type, uint64 Offset, uint64 Size, double Time)
+uint64 FFileActivityProvider::BeginActivity(uint32 FileIndex, EFileActivityType Type, uint32 ThreadId, uint64 Offset, uint64 Size, double Time)
 {
 	FFileInfoInternal& FileInfo = Files[FileIndex];
 	FFileActivity& FileActivity = FileActivities.PushBack();
@@ -57,6 +57,7 @@ uint64 FFileActivityProvider::BeginActivity(uint32 FileIndex, EFileActivityType
 	FileActivity.Size = Size;
 	FileActivity.StartTime = Time;
 	FileActivity.EndTime = std::numeric_limits<double>::infinity();
+	FileActivity.ThreadId = ThreadId;
 	FileActivity.Failed = false;
 	FileActivity.ActivityType = Type;
 	return FileInfo.ActivityTimeline->AppendBeginEvent(Time, &FileActivity);
diff --git a/Engine/Source/Developer/TraceServices/Private/Model/FileActivity.h b/Engine/Source/Developer/TraceServices/Private/Model/FileActivity.h
--- a/Engine/Source/Developer/TraceServices/Private/Model/FileActivity.h
+++ b/Engine/Source/Developer/TraceServices/Private/Model/FileActivity.h
@@ -33,7 +33,7 @@ public:
 	virtual const ITable<FFileActivity>& GetFileActivityTable() const override;
 	uint32 GetFileIndex(const TCHAR* Path);
 	uint32 GetUnknownFileIndex();
-	uint64 BeginActivity(uint32 FileIndex, EFileActivityType Type, uint64 Offset, uint64 Size, double Time);
+	uint64 BeginActivity(uint32 FileIndex, EFileActivityType Type, uint32 ThreadId, uint64 Offset, uint64 Size, double Time);
 	void EndActivity(uint32 FileIndex, uint64 ActivityIndex, uint64 ActualSize, double Time, bool Failed);
 	const TCHAR* GetFilePath(uint32 FileIndex) const;
 
@@ -48,6 +48,7 @@ private:
 		UE_TRACE_TABLE_PROJECTED_COLUMN(TableColumnType_CString, TEXT("File"), [](const FFileActivity& Row) { return Row.File ? Row.File->Path : TEXT("N/A"); })
 		UE_TRACE_TABLE_COLUMN(StartTime, TEXT("StartTime"))
 		UE_TRACE_TABLE_COLUMN(EndTime, TEXT("EndTime"))
+		UE_TRACE_TABLE_COLUMN(ThreadId, TEXT("ThreadId"))
 		UE_TRACE_TABLE_PROJECTED_COLUMN(TableColumnType_CString, TEXT("Type"), [](const FFileActivity& Row) { return GetFileActivityTypeString(Row.ActivityType); })
 		UE_TRACE_TABLE_COLUMN(Offset, TEXT("Offset"))
 		UE_TRACE_TABLE_COLUMN(Size, TEXT("Size"))
diff --git a/Engine/Source/Developer/TraceServices/Private/Model/LoadTimeProfilerPrivate.h b/Engine/Source/Developer/TraceServices/Private/Model/LoadTimeProfilerPrivate.h
--- a/Engine/Source/Developer/TraceServices/Private/Model/LoadTimeProfilerPrivate.h
+++ b/Engine/Source/Developer/TraceServices/Private/Model/LoadTimeProfilerPrivate.h
@@ -65,7 +65,7 @@ private:
 	UE_TRACE_TABLE_LAYOUT_END()
 
 	UE_TRACE_TABLE_LAYOUT_BEGIN(FExportsTableLayout, FExportsTableRow)
-		UE_TRACE_TABLE_PROJECTED_COLUMN(TableColumnType_CString, TEXT("Package"), [](const FExportsTableRow& Row) { return Row.ExportInfo->Package->Name; })
+		UE_TRACE_TABLE_PROJECTED_COLUMN(TableColumnType_CString, TEXT("Package"), [](const FExportsTableRow& Row) { return Row.ExportInfo->Package ? Row.ExportInfo->Package->Name : TEXT("[unknown]"); })
 		UE_TRACE_TABLE_PROJECTED_COLUMN(TableColumnType_CString, TEXT("Class"), [](const FExportsTableRow& Row) { return Row.ExportInfo->Class ? Row.ExportInfo->Class->Name : TEXT("[unknown]"); })
 		UE_TRACE_TABLE_PROJECTED_COLUMN(TableColumnType_CString, TEXT("EventType"), [](const FExportsTableRow& Row) { return GetLoadTimeProfilerObjectEventTypeString(Row.EventType); })
 		UE_TRACE_TABLE_COLUMN(SerializedSize, TEXT("SerializedSize"))
diff --git a/Engine/Source/Developer/TraceServices/Private/Model/Log.cpp b/Engine/Source/Developer/TraceServices/Private/Model/Log.cpp
--- a/Engine/Source/Developer/TraceServices/Private/Model/Log.cpp
+++ b/Engine/Source/Developer/TraceServices/Private/Model/Log.cpp
@@ -53,13 +53,16 @@ FLogMessageSpec& FLogProvider::GetMessageSpec(uint64 LogPoint)
 void FLogProvider::AppendMessage(uint64 LogPoint, double Time, const uint8* FormatArgs)
 {
 	Session.WriteAccessCheck();
-	check(SpecMap.Contains(LogPoint));
-	FLogMessageInternal& InternalMessage = Messages.PushBack();
-	InternalMessage.Time = Time;
-	InternalMessage.Spec = SpecMap[LogPoint];
-	FFormatArgsHelper::Format(FormatBuffer, FormatBufferSize - 1, InternalMessage.Spec->FormatString, FormatArgs);
-	InternalMessage.Message = Session.StoreString(FormatBuffer);
-	Session.UpdateDurationSeconds(Time);
+	FLogMessageSpec** FindSpec = SpecMap.Find(LogPoint);
+	if (FindSpec && (*FindSpec)->Verbosity != ELogVerbosity::SetColor)
+	{
+		FLogMessageInternal& InternalMessage = Messages.PushBack();
+		InternalMessage.Time = Time;
+		InternalMessage.Spec = *FindSpec;
+		FFormatArgsHelper::Format(FormatBuffer, FormatBufferSize - 1, InternalMessage.Spec->FormatString, FormatArgs);
+		InternalMessage.Message = Session.StoreString(FormatBuffer);
+		Session.UpdateDurationSeconds(Time);
+	}
 }
 
 uint64 FLogProvider::GetMessageCount() const
diff --git a/Engine/Source/Developer/TraceServices/Public/TraceServices/Model/LoadTimeProfiler.h b/Engine/Source/Developer/TraceServices/Public/TraceServices/Model/LoadTimeProfiler.h
--- a/Engine/Source/Developer/TraceServices/Public/TraceServices/Model/LoadTimeProfiler.h
+++ b/Engine/Source/Developer/TraceServices/Public/TraceServices/Model/LoadTimeProfiler.h
@@ -37,6 +37,7 @@ struct FFileActivity
 	double EndTime = 0.0;
 	uint64 Offset = 0;
 	uint64 Size = 0;
+	uint32 ThreadId;
 	EFileActivityType ActivityType = FileActivityType_Invalid;
 	bool Failed = false;
 };
