UnrealInsights:
Added trace based CsvProfiler implementation

diff --git a/Engine/Source/Developer/TraceServices/Private/Analyzers/CsvProfilerTraceAnalysis.cpp b/Engine/Source/Developer/TraceServices/Private/Analyzers/CsvProfilerTraceAnalysis.cpp
--- a/Engine/Source/Developer/TraceServices/Private/Analyzers/CsvProfilerTraceAnalysis.cpp
+++ b/Engine/Source/Developer/TraceServices/Private/Analyzers/CsvProfilerTraceAnalysis.cpp
@@ -2,28 +2,60 @@
 #include "CsvProfilerTraceAnalysis.h"
 #include "AnalysisServicePrivate.h"
 #include "TraceServices/Model/Counters.h"
+#include "TraceServices/Model/Frames.h"
 
-FCsvProfilerAnalyzer::FCsvProfilerAnalyzer(Trace::IAnalysisSession& InSession, Trace::ICounterProvider& InCounterProvider)
+FCsvProfilerAnalyzer::FCsvProfilerAnalyzer(Trace::IAnalysisSession& InSession, Trace::FCsvProfilerProvider& InCsvProfilerProvider, const Trace::IFrameProvider& InFrameProvider, const Trace::IThreadProvider& InThreadProvider)
 	: Session(InSession)
-	, CounterProvider(InCounterProvider)
+	, CsvProfilerProvider(InCsvProfilerProvider)
+	, FrameProvider(InFrameProvider)
+	, ThreadProvider(InThreadProvider)
 {
+}
 
+FCsvProfilerAnalyzer::~FCsvProfilerAnalyzer()
+{
+	OnAnalysisEnd();
 }
 
 void FCsvProfilerAnalyzer::OnAnalysisBegin(const FOnAnalysisContext& Context)
 {
 	auto& Builder = Context.InterfaceBuilder;
 
-	Builder.RouteEvent(RouteId_InlineStat, "CsvProfiler", "InlineStat");
-	Builder.RouteEvent(RouteId_DeclaredStat, "CsvProfiler", "DeclaredStat");
-	Builder.RouteEvent(RouteId_CustomStatInlineInt, "CsvProfiler", "CustomStatInlineInt");
-	Builder.RouteEvent(RouteId_CustomStatInlineFloat, "CsvProfiler", "CustomStatInlineFloat");
-	Builder.RouteEvent(RouteId_CustomStatDeclaredInt, "CsvProfiler", "CustomStatDeclaredInt");
-	Builder.RouteEvent(RouteId_CustomStatDeclaredFloat, "CsvProfiler", "CustomStatDeclaredFloat");
+	Builder.RouteEvent(RouteId_RegisterCategory, "CsvProfiler", "RegisterCategory");
+	Builder.RouteEvent(RouteId_DefineInlineStat, "CsvProfiler", "DefineInlineStat");
+	Builder.RouteEvent(RouteId_DefineDeclaredStat, "CsvProfiler", "DefineDeclaredStat");
+	Builder.RouteEvent(RouteId_BeginStat, "CsvProfiler", "BeginStat");
+	Builder.RouteEvent(RouteId_EndStat, "CsvProfiler", "EndStat");
+	Builder.RouteEvent(RouteId_BeginExclusiveStat, "CsvProfiler", "BeginExclusiveStat");
+	Builder.RouteEvent(RouteId_EndExclusiveStat, "CsvProfiler", "EndExclusiveStat");
+	Builder.RouteEvent(RouteId_CustomStatInt, "CsvProfiler", "CustomStatInt");
+	Builder.RouteEvent(RouteId_CustomStatFloat, "CsvProfiler", "CustomStatFloat");
+	Builder.RouteEvent(RouteId_Event, "CsvProfiler", "Event");
+	Builder.RouteEvent(RouteId_Metadata, "CsvProfiler", "Metadata");
+	Builder.RouteEvent(RouteId_BeginCapture, "CsvProfiler", "BeginCapture");
+	Builder.RouteEvent(RouteId_EndCapture, "CsvProfiler", "EndCapture");
 }
 
 void FCsvProfilerAnalyzer::OnAnalysisEnd()
 {
+	for (FStatSeriesInstance* StatSeriesInstance : StatSeriesInstanceArray)
+	{
+		delete StatSeriesInstance;
+	}
+	StatSeriesInstanceArray.Empty();
+	for (FStatSeriesDefinition* StatSeriesDefinition : StatSeriesDefinitionArray)
+	{
+		delete StatSeriesDefinition;
+	}
+	StatSeriesDefinitionArray.Empty();
+	StatSeriesMap.Empty();
+	StatSeriesStringMap.Empty();
+	for (auto& KV : ThreadStatesMap)
+	{
+		FThreadState* ThreadState = KV.Value;
+		delete ThreadState;
+	}
+	ThreadStatesMap.Empty();
 }
 
 void FCsvProfilerAnalyzer::OnEvent(uint16 RouteId, const FOnEventContext& Context)
@@ -33,119 +65,355 @@ void FCsvProfilerAnalyzer::OnEvent(uint16 RouteId, const FOnEventContext& Contex
 	const auto& EventData = Context.EventData;
 	switch (RouteId)
 	{
-	case RouteId_InlineStat:
+	case RouteId_RegisterCategory:
+	{
+		int32 CategoryIndex = EventData.GetValue<int32>("Index");
+		const TCHAR* Name = reinterpret_cast<const TCHAR*>(EventData.GetAttachment());
+		CategoryMap.Add(CategoryIndex, Session.StoreString(Name));;
+		break;
+	}
+	case RouteId_DefineInlineStat:
+	{
+		uint64 StatId = EventData.GetValue<uint64>("StatId");
+		int32 CategoryIndex = EventData.GetValue<int32>("CategoryIndex");
+		DefineStatSeries(StatId, ANSI_TO_TCHAR(reinterpret_cast<const ANSICHAR*>(EventData.GetAttachment())), CategoryIndex, true);
+		break;
+	}
+	case RouteId_DefineDeclaredStat:
+	{
+		uint64 StatId = EventData.GetValue<uint64>("StatId");
+		int32 CategoryIndex = EventData.GetValue<int32>("CategoryIndex");
+		DefineStatSeries(StatId, reinterpret_cast<const TCHAR*>(EventData.GetAttachment()), CategoryIndex, false);
+		break;
+	}
+	case RouteId_BeginStat:
+	{
+		HandleMarkerEvent(Context, false, true);
+		break;
+	}
+	case RouteId_EndStat:
+	{
+		HandleMarkerEvent(Context, false, false);
+		break;
+	}
+	case RouteId_BeginExclusiveStat:
+	{
+		HandleMarkerEvent(Context, true, true);
+		break;
+	}
+	case RouteId_EndExclusiveStat:
 	{
-		uint64 StatNamePointer = EventData.GetValue<uint64>("StatNamePointer");
-		InlinePendingCountersMap.Add(StatNamePointer, Session.StoreString(ANSI_TO_TCHAR(reinterpret_cast<const ANSICHAR*>(EventData.GetAttachment()))));
+		HandleMarkerEvent(Context, true, false);
 		break;
 	}
-	case RouteId_DeclaredStat:
+	case RouteId_CustomStatInt:
 	{
-		uint32 StatId = EventData.GetValue<uint32>("StatId");
-		DeclaredPendingCountersMap.Add(StatId, Session.StoreString(reinterpret_cast<const TCHAR*>(EventData.GetAttachment())));
+		HandleCustomStatEvent(Context, false);
 		break;
 	}
-	case RouteId_CustomStatInlineInt:
+	case RouteId_CustomStatFloat:
 	{
-		Trace::ICounter* Counter = GetInlineStatCounter(EventData.GetValue<uint64>("StatNamePointer"), false);
-		HandleIntCounterEvent(Counter, Context);
+		HandleCustomStatEvent(Context, true);
 		break;
 	}
-	case RouteId_CustomStatInlineFloat:
+	case RouteId_Event:
 	{
-		Trace::ICounter* Counter = GetInlineStatCounter(EventData.GetValue<uint64>("StatNamePointer"), true);
-		HandleFloatCounterEvent(Counter, Context);
+		HandleEventEvent(Context);
 		break;
 	}
-	case RouteId_CustomStatDeclaredInt:
+	case RouteId_Metadata:
 	{
-		Trace::ICounter* Counter = GetDeclaredStatCounter(EventData.GetValue<uint32>("StatId"), false);
-		HandleIntCounterEvent(Counter, Context);
+		const TCHAR* Key = reinterpret_cast<const TCHAR*>(EventData.GetAttachment());
+		const TCHAR* Value = reinterpret_cast<const TCHAR*>(EventData.GetAttachment() + EventData.GetValue<uint16>("ValueOffset"));
+		CsvProfilerProvider.SetMetadata(Session.StoreString(Key), Session.StoreString(Value));
 		break;
 	}
-	case RouteId_CustomStatDeclaredFloat:
+	case RouteId_BeginCapture:
 	{
-		Trace::ICounter* Counter = GetDeclaredStatCounter(EventData.GetValue<uint32>("StatId"), true);
-		HandleFloatCounterEvent(Counter, Context);
+		RenderThreadId = EventData.GetValue<uint32>("RenderThreadId");
+		RHIThreadId = EventData.GetValue<uint32>("RHIThreadId");
+		int32 CaptureStartFrame = GetFrameNumberForTimestamp(TraceFrameType_Game, Context.SessionContext.TimestampFromCycle(EventData.GetValue<uint64>("Cycle")));
+		bEnableCounts = EventData.GetValue<bool>("EnableCounts");
+		const TCHAR* Filename = Session.StoreString(reinterpret_cast<const TCHAR*>(EventData.GetAttachment()));
+		CsvProfilerProvider.StartCapture(Filename, CaptureStartFrame);
 		break;
 	}
+	case RouteId_EndCapture:
+	{
+		int32 CaptureEndFrame = GetFrameNumberForTimestamp(TraceFrameType_Game, Context.SessionContext.TimestampFromCycle(EventData.GetValue<uint64>("Cycle")));
+		CsvProfilerProvider.EndCapture(CaptureEndFrame);
+	}
 	}
 }
 
-void FCsvProfilerAnalyzer::HandleIntCounterEvent(Trace::ICounter* Counter, const FOnEventContext& Context)
+FCsvProfilerAnalyzer::FThreadState& FCsvProfilerAnalyzer::GetThreadState(uint32 ThreadId)
 {
-	const auto& EventData = Context.EventData;
-	double Timestamp = Context.SessionContext.TimestampFromCycle(EventData.GetValue<uint64>("Cycle"));
-	uint8 OpType = EventData.GetValue<uint8>("OpType");
-	if (OpType == OpType_Accumulate)
+	FThreadState** FindIt = ThreadStatesMap.Find(ThreadId);
+	if (FindIt)
 	{
-		Counter->AddValue(Timestamp, int64(EventData.GetValue<int32>("Value")));
+		return **FindIt;
 	}
-	else
+	FThreadState* ThreadState = new FThreadState();
+	if (ThreadId == RenderThreadId || ThreadId == RHIThreadId)
 	{
-		Counter->SetValue(Timestamp, int64(EventData.GetValue<int32>("Value")));
+		ThreadState->FrameType = TraceFrameType_Rendering;
 	}
+	ThreadState->ThreadName = ThreadId == RenderThreadId ? TEXT("RenderThread") : ThreadProvider.GetThreadName(ThreadId);
+	ThreadStatesMap.Add(ThreadId, ThreadState);
+	return *ThreadState;
 }
 
-void FCsvProfilerAnalyzer::HandleFloatCounterEvent(Trace::ICounter* Counter, const FOnEventContext& Context)
+FCsvProfilerAnalyzer::FStatSeriesDefinition* FCsvProfilerAnalyzer::CreateStatSeries(const TCHAR* Name, int32 CategoryIndex)
 {
-	const auto& EventData = Context.EventData;
-	double Timestamp = Context.SessionContext.TimestampFromCycle(EventData.GetValue<uint64>("Cycle"));
-	uint8 OpType = EventData.GetValue<uint8>("OpType");
-	if (OpType == OpType_Accumulate)
+	FStatSeriesDefinition* StatSeries = new FStatSeriesDefinition();
+	StatSeries->Name = Session.StoreString(Name);
+	StatSeries->CategoryIndex = CategoryIndex;
+	StatSeries->ColumnIndex = StatSeriesDefinitionArray.Num();
+	StatSeriesDefinitionArray.Add(StatSeries);
+	return StatSeries;
+}
+
+void FCsvProfilerAnalyzer::DefineStatSeries(uint64 StatId, const TCHAR* Name, int32 CategoryIndex, bool bIsInline)
+{
+	FStatSeriesDefinition** FindIt = StatSeriesMap.Find(StatId);
+	if (bIsInline && !FindIt)
 	{
-		Counter->AddValue(Timestamp, double(EventData.GetValue<float>("Value")));
+		TTuple<int32, FString> Key = TTuple<int32, FString>(CategoryIndex, Name);
+		FindIt = StatSeriesStringMap.Find(Key);
+		if (FindIt)
+		{
+			StatSeriesMap.Add(StatId, *FindIt);
+		}
 	}
-	else
+	if (!FindIt)
+	{
+		FStatSeriesDefinition* StatSeries = CreateStatSeries(Name, CategoryIndex);
+		StatSeriesMap.Add(StatId, StatSeries);
+		if (bIsInline)
+		{
+			TTuple<int32, FString> Key = TTuple<int32, FString>(CategoryIndex, Name);
+			StatSeriesStringMap.Add(Key, StatSeries);
+		}
+	}
+}
+
+const TCHAR* FCsvProfilerAnalyzer::GetStatSeriesName(const FStatSeriesDefinition* Definition, Trace::ECsvStatSeriesType Type, FThreadState& ThreadState, bool bIsCount)
+{
+	FString Name = Definition->Name;
+	if (Type == Trace::CsvStatSeriesType_Timer || Type == Trace::CsvStatSeriesType_ExclusiveTimer || bIsCount)
+	{
+		// Add a /<Threadname> prefix
+		Name = ThreadState.ThreadName + TEXT("/") + Name;
+	}
+
+	if (Definition->CategoryIndex > 0)
+	{
+		// Categorized stats are prefixed with <CATEGORY>/
+		Name = FString(CategoryMap[Definition->CategoryIndex]) + TEXT("/") + Name;
+	}
+	else if (Type == Trace::CsvStatSeriesType_ExclusiveTimer)
+	{
+		Name = FString(TEXT("Exclusive/")) + Name;
+	}
+
+	if (bIsCount)
 	{
-		Counter->SetValue(Timestamp, double(EventData.GetValue<float>("Value")));
+		// Add a counts prefix
+		Name = TEXT("COUNTS/") + Name;
 	}
+
+	return Session.StoreString(*Name);
 }
 
-Trace::ICounter* FCsvProfilerAnalyzer::GetDeclaredStatCounter(uint32 StatId, bool bIsFloatingPoint)
+FCsvProfilerAnalyzer::FStatSeriesInstance* FCsvProfilerAnalyzer::GetStatSeries(uint64 StatId, Trace::ECsvStatSeriesType Type, FThreadState& ThreadState)
 {
-	Trace::ICounter** FindIt;
-	FindIt = DeclaredCountersMap.Find(StatId);
+	FStatSeriesDefinition* Definition;
+	FStatSeriesDefinition** FindIt = StatSeriesMap.Find(StatId);
 	if (!FindIt)
 	{
-		Trace::ICounter* Counter = CounterProvider.CreateCounter();
-		Counter->SetIsFloatingPoint(bIsFloatingPoint);
-		Counter->SetIsResetEveryFrame(true);
-		const TCHAR** FindName = DeclaredPendingCountersMap.Find(StatId);
-		if (FindName)
+		Definition = CreateStatSeries(*FString::Printf(TEXT("[undefined%d]"), UndefinedStatSeriesCount++), 0);
+		StatSeriesMap.Add(StatId, Definition);
+	}
+	else
+	{
+		Definition = *FindIt;
+	}
+
+	if (ThreadState.StatSeries.Num() <= Definition->ColumnIndex)
+	{
+		ThreadState.StatSeries.AddZeroed(Definition->ColumnIndex + 1 - ThreadState.StatSeries.Num());
+	}
+	FStatSeriesInstance* Instance = ThreadState.StatSeries[Definition->ColumnIndex];
+	if (Instance)
+	{
+		return Instance;
+	}
+
+	Instance = new FStatSeriesInstance();
+	StatSeriesInstanceArray.Add(Instance);
+	ThreadState.StatSeries[Definition->ColumnIndex] = Instance;
+	Instance->ProviderHandle = CsvProfilerProvider.AddSeries(GetStatSeriesName(Definition, Type, ThreadState, false), Type);
+	Instance->ProviderCountHandle = CsvProfilerProvider.AddSeries(GetStatSeriesName(Definition, Type, ThreadState, true), Trace::CsvStatSeriesType_CustomStatInt);
+
+	return Instance;
+}
+
+void FCsvProfilerAnalyzer::HandleMarkerEvent(const FOnEventContext& Context, bool bIsExclusive, bool bIsBegin)
+{
+	FThreadState& ThreadState = GetThreadState(Context.EventData.GetValue<uint32>("ThreadId"));
+	uint64 StatId = Context.EventData.GetValue<uint64>("StatId");
+	FTimingMarker Marker;
+	Marker.StatSeries = GetStatSeries(StatId, Trace::CsvStatSeriesType_Timer, ThreadState);
+	Marker.bIsBegin = bIsBegin;
+	Marker.bIsExclusive = bIsExclusive;
+	Marker.Cycle = Context.EventData.GetValue<uint64>("Cycle");
+	HandleMarker(Context, ThreadState, Marker);
+}
+
+void FCsvProfilerAnalyzer::HandleMarker(const FOnEventContext& Context, FThreadState& ThreadState, const FTimingMarker& Marker)
+{
+	// Handle exclusive markers. This may insert an additional marker before this one
+	bool bInsertExtraMarker = false;
+	FTimingMarker InsertedMarker;
+	if (Marker.bIsExclusive & !Marker.bIsExclusiveInsertedMarker)
+	{
+		if (Marker.bIsBegin)
 		{
-			Counter->SetName(*FindName);
-			DeclaredPendingCountersMap.Remove(StatId);
+			if (ThreadState.ExclusiveMarkerStack.Num() > 0)
+			{
+				// Insert an artificial end marker to end the previous marker on the stack at the same timestamp
+				InsertedMarker = ThreadState.ExclusiveMarkerStack.Last();
+				InsertedMarker.bIsBegin = false;
+				InsertedMarker.bIsExclusiveInsertedMarker = true;
+				InsertedMarker.Cycle = Marker.Cycle;
+
+				bInsertExtraMarker = true;
+			}
+			ThreadState.ExclusiveMarkerStack.Add(Marker);
+		}
+		else
+		{
+			if (ThreadState.ExclusiveMarkerStack.Num() > 0)
+			{
+				ThreadState.ExclusiveMarkerStack.Pop(false);
+				if (ThreadState.ExclusiveMarkerStack.Num() > 0)
+				{
+					// Insert an artificial begin marker to resume the marker on the stack at the same timestamp
+					InsertedMarker = ThreadState.ExclusiveMarkerStack.Last();
+					InsertedMarker.bIsBegin = true;
+					InsertedMarker.bIsExclusiveInsertedMarker = true;
+					InsertedMarker.Cycle = Marker.Cycle;
+
+					bInsertExtraMarker = true;
+				}
+			}
 		}
-		DeclaredCountersMap.Add(StatId, Counter);
-		return Counter;
+	}
+	if (bInsertExtraMarker)
+	{
+		HandleMarker(Context, ThreadState, InsertedMarker);
+	}
+	
+	int32 FrameNumber = GetFrameNumberForTimestamp(ThreadState.FrameType, Context.SessionContext.TimestampFromCycle(Marker.Cycle));
+	if (Marker.bIsBegin)
+	{
+		ThreadState.MarkerStack.Push(Marker);
 	}
 	else
 	{
-		return *FindIt;
+		// Markers might not match up if they were truncated mid-frame, so we need to be robust to that
+		if (ThreadState.MarkerStack.Num() > 0)
+		{
+			// Find the start marker (might not actually be top of the stack, e.g if begin/end for two overlapping stats are independent)
+			bool bFoundStart = false;
+			FTimingMarker StartMarker;
+
+			for (int j = ThreadState.MarkerStack.Num() - 1; j >= 0; j--)
+			{
+				if (ThreadState.MarkerStack[j].StatSeries == Marker.StatSeries) // Note: only works with scopes!
+				{
+					StartMarker = ThreadState.MarkerStack[j];
+					ThreadState.MarkerStack.RemoveAt(j, 1, false);
+					bFoundStart = true;
+					break;
+				}
+			}
+
+			// TODO: if bFoundStart is false, this stat _never_ gets processed. Could we add it to a persistent list so it's considered next time?
+			// Example where this could go wrong: staggered/overlapping exclusive stats ( e.g Abegin, Bbegin, AEnd, BEnd ), where processing ends after AEnd
+			// AEnd would be missing 
+			if (FrameNumber >= 0 && bFoundStart)
+			{
+				check(Marker.StatSeries == StartMarker.StatSeries);
+				check(Marker.Cycle >= StartMarker.Cycle);
+				if (Marker.Cycle > StartMarker.Cycle)
+				{
+					uint64 ElapsedCycles = Marker.Cycle - StartMarker.Cycle;
+					CsvProfilerProvider.SetTimerValue(Marker.StatSeries->ProviderHandle, FrameNumber, Context.SessionContext.DurationFromCycleCount(ElapsedCycles) * 1000.0);
+
+					// Add the COUNT/ series if enabled. Ignore artificial markers (inserted above)
+					if (bEnableCounts && !Marker.bIsExclusiveInsertedMarker)
+					{
+						CsvProfilerProvider.SetCustomStatValue(Marker.StatSeries->ProviderCountHandle, FrameNumber, Trace::CsvOpType_Accumulate, 1);
+					}
+				}
+			}
+		}
 	}
 }
 
-Trace::ICounter* FCsvProfilerAnalyzer::GetInlineStatCounter(uint64 StatNamePointer, bool bIsFloatingPoint)
+void FCsvProfilerAnalyzer::HandleCustomStatEvent(const FOnEventContext& Context, bool bIsFloat)
 {
-	Trace::ICounter** FindIt;
-	FindIt = InlineCountersMap.Find(StatNamePointer);
-	if (!FindIt)
+	FThreadState& ThreadState = GetThreadState(Context.EventData.GetValue<uint32>("ThreadId"));
+	FStatSeriesInstance* StatSeries = GetStatSeries(Context.EventData.GetValue<uint64>("StatId"), Trace::CsvStatSeriesType_CustomStatFloat, ThreadState);
+	Trace::ECsvOpType OpType = static_cast<Trace::ECsvOpType>(Context.EventData.GetValue<uint8>("OpType"));
+	uint64 Cycle = Context.EventData.GetValue<uint64>("Cycle");
+	int32 FrameNumber = GetFrameNumberForTimestamp(ThreadState.FrameType, Context.SessionContext.TimestampFromCycle(Cycle));
+	if (FrameNumber >= 0)
+	{
+		if (bIsFloat)
+		{
+			CsvProfilerProvider.SetCustomStatValue(StatSeries->ProviderHandle, FrameNumber, OpType, Context.EventData.GetValue<float>("Value"));
+		}
+		else
+		{
+			CsvProfilerProvider.SetCustomStatValue(StatSeries->ProviderHandle, FrameNumber, OpType, Context.EventData.GetValue<int32>("Value"));
+		}
+
+		// Add the COUNT/ series if enabled
+		if (bEnableCounts)
+		{
+			CsvProfilerProvider.SetCustomStatValue(StatSeries->ProviderCountHandle, FrameNumber, Trace::CsvOpType_Accumulate, 1);
+		}
+	}
+}
+
+void FCsvProfilerAnalyzer::HandleEventEvent(const FOnEventContext& Context)
+{
+	FThreadState& ThreadState = GetThreadState(Context.EventData.GetValue<uint32>("ThreadId"));
+	uint64 Cycle = Context.EventData.GetValue<uint64>("Cycle");
+	int32 FrameNumber = GetFrameNumberForTimestamp(ThreadState.FrameType, Context.SessionContext.TimestampFromCycle(Cycle));
+	if (FrameNumber >= 0)
 	{
-		Trace::ICounter* Counter = CounterProvider.CreateCounter();
-		Counter->SetIsFloatingPoint(bIsFloatingPoint);
-		Counter->SetIsResetEveryFrame(true);
-		InlineCountersMap.Add(StatNamePointer, Counter);
-		const TCHAR** FindName = InlinePendingCountersMap.Find(StatNamePointer);
-		if (FindName)
+		FString EventText = reinterpret_cast<const TCHAR*>(Context.EventData.GetAttachment());
+		int32 CategoryIndex = Context.EventData.GetValue<int32>("CategoryIndex");
+		if (CategoryIndex > 0)
 		{
-			Counter->SetName(*FindName);
-			InlinePendingCountersMap.Remove(StatNamePointer);
+			EventText = FString(CategoryMap[CategoryIndex]) + TEXT("/") + EventText;
 		}
-		return Counter;
+		CsvProfilerProvider.AddEvent(FrameNumber, Session.StoreString(*EventText));
 	}
-	else
+}
+
+int32 FCsvProfilerAnalyzer::GetFrameNumberForTimestamp(ETraceFrameType FrameType, double Timestamp) const
+{
+	const TArray<double>& FrameStartTimes = FrameProvider.GetFrameStartTimes(FrameType);
+
+	if (FrameStartTimes.Num() == 0 || Timestamp < FrameStartTimes[0])
 	{
-		return *FindIt;
+		return 0;
 	}
+
+	uint32 Index = static_cast<uint32>(Algo::LowerBound(FrameStartTimes, Timestamp));
+	return Index + 1;
 }
diff --git a/Engine/Source/Developer/TraceServices/Private/Analyzers/CsvProfilerTraceAnalysis.h b/Engine/Source/Developer/TraceServices/Private/Analyzers/CsvProfilerTraceAnalysis.h
--- a/Engine/Source/Developer/TraceServices/Private/Analyzers/CsvProfilerTraceAnalysis.h
+++ b/Engine/Source/Developer/TraceServices/Private/Analyzers/CsvProfilerTraceAnalysis.h
@@ -5,19 +5,22 @@
 #include "Trace/Trace.h"
 #include "Trace/Analyzer.h"
 #include "Containers/Map.h"
+#include "Model/CsvProfilerPrivate.h"
+#include "ProfilingDebugging/MiscTrace.h"
 
 namespace Trace
 {
 	class IAnalysisSession;
-	class ICounter;
-	class ICounterProvider;
+	class IFrameProvider;
+	class IThreadProvider;
 }
 
 class FCsvProfilerAnalyzer
 	: public Trace::IAnalyzer
 {
 public:
-	FCsvProfilerAnalyzer(Trace::IAnalysisSession& Session, Trace::ICounterProvider& CounterProvider);
+	FCsvProfilerAnalyzer(Trace::IAnalysisSession& Session, Trace::FCsvProfilerProvider& CsvProfilerProvider, const Trace::IFrameProvider& FrameProvider, const Trace::IThreadProvider& ThreadProvider);
+	~FCsvProfilerAnalyzer();
 	virtual void OnAnalysisBegin(const FOnAnalysisContext& Context) override;
 	virtual void OnEvent(uint16 RouteId, const FOnEventContext& Context) override;
 	virtual void OnAnalysisEnd() override;
@@ -25,31 +28,90 @@ public:
 private:
 	enum : uint16
 	{
-		RouteId_InlineStat,
-		RouteId_DeclaredStat,
-		RouteId_CustomStatInlineInt,
-		RouteId_CustomStatInlineFloat,
-		RouteId_CustomStatDeclaredInt,
-		RouteId_CustomStatDeclaredFloat,
+		RouteId_RegisterCategory,
+		RouteId_DefineInlineStat,
+		RouteId_DefineDeclaredStat,
+		RouteId_BeginStat,
+		RouteId_EndStat,
+		RouteId_BeginExclusiveStat,
+		RouteId_EndExclusiveStat,
+		RouteId_CustomStatInt,
+		RouteId_CustomStatFloat,
+		RouteId_Event,
+		RouteId_Metadata,
+		RouteId_BeginCapture,
+		RouteId_EndCapture,
 	};
 
-	enum : uint8
+	union FStatSeriesID
 	{
-		OpType_Set,
-		OpType_Min,
-		OpType_Max,
-		OpType_Accumulate,
+		struct
+		{
+			uint64 IsFName : 1;
+			uint64 IsCountStat : 1;
+			uint64 CategoryIndex : 11;
+			uint64 FNameOrIndex : 51;
+		} Fields;
+		uint64 Hash;
 	};
 
-	void HandleIntCounterEvent(Trace::ICounter* Counter, const FOnEventContext& Context);
-	void HandleFloatCounterEvent(Trace::ICounter* Counter, const FOnEventContext& Context);
-	Trace::ICounter* GetDeclaredStatCounter(uint32 StatId, bool bIsFloatingPoint);
-	Trace::ICounter* GetInlineStatCounter(uint64 StatNamePointer, bool bIsFloatingPoint);
+	struct FStatSeriesDefinition
+	{
+		const TCHAR* Name = nullptr;
+		int32 CategoryIndex = -1;
+		int32 ColumnIndex = -1;
+	};
 
+	struct FStatSeriesInstance
+	{
+		uint64 ProviderHandle = uint64(-1);
+		uint64 ProviderCountHandle = uint64(-1);
+	};
+
+	struct FTimingMarker
+	{
+		const FStatSeriesInstance* StatSeries = nullptr;
+		uint64 Cycle = 0;
+		bool bIsBegin = false;
+		bool bIsExclusive = false;
+		bool bIsExclusiveInsertedMarker = false;
+	};
+
+	struct FThreadState
+	{
+		TArray<FTimingMarker> MarkerStack;
+		TArray<FTimingMarker> ExclusiveMarkerStack;
+		ETraceFrameType FrameType = TraceFrameType_Game;
+		TArray<FStatSeriesInstance*> StatSeries;
+		FString ThreadName;
+	};
+
+	FThreadState& GetThreadState(uint32 ThreadId);
+	FStatSeriesDefinition* CreateStatSeries(const TCHAR* Name, int32 CategoryIndex);
+	void DefineStatSeries(uint64 StatId, const TCHAR* Name, int32 CategoryIndex, bool bIsInline);
+	const TCHAR* GetStatSeriesName(const FStatSeriesDefinition* Definition, Trace::ECsvStatSeriesType Type, FThreadState& ThreadState, bool bIsCount);
+	FStatSeriesInstance* GetStatSeries(uint64 StatId, Trace::ECsvStatSeriesType Type, FThreadState& ThreadState);
+	void HandleMarkerEvent(const FOnEventContext& Context, bool bIsExclusive, bool bIsBegin);
+	void HandleMarker(const FOnEventContext& Context, FThreadState& ThreadState, const FTimingMarker& Marker);
+	void HandleCustomStatEvent(const FOnEventContext& Context, bool bIsFloat);
+	void HandleEventEvent(const FOnEventContext& Context);
+	int32 GetFrameNumberForTimestamp(ETraceFrameType FrameType, double Timestamp) const;
+	
 	Trace::IAnalysisSession& Session;
-	Trace::ICounterProvider& CounterProvider;
-	TMap<uint32, const TCHAR*> InlinePendingCountersMap;
-	TMap<uint32, Trace::ICounter*> InlineCountersMap;
-	TMap<uint64, const TCHAR*> DeclaredPendingCountersMap;
-	TMap<uint64, Trace::ICounter*> DeclaredCountersMap;
+	Trace::FCsvProfilerProvider& CsvProfilerProvider;
+	const Trace::IFrameProvider& FrameProvider;
+	const Trace::IThreadProvider& ThreadProvider;
+
+	TMap<uint32, FThreadState*> ThreadStatesMap;
+	TMap<int32, const TCHAR*> CategoryMap;
+	TMap<uint64, FStatSeriesDefinition*> StatSeriesMap;
+	TMap<TTuple<int32, FString>, FStatSeriesDefinition*> StatSeriesStringMap;
+	TArray<FStatSeriesDefinition*> StatSeriesDefinitionArray;
+	TArray<FStatSeriesInstance*> StatSeriesInstanceArray;
+	uint32 RenderThreadId = 0;
+	uint32 RHIThreadId = 0;
+	bool bEnableCounts = false;
+	uint32 UndefinedStatSeriesCount = 0;
+
+	TArray<uint64> FrameBoundaries[TraceFrameType_Count];
 };
diff --git a/Engine/Source/Developer/TraceServices/Private/Model/Counters.cpp b/Engine/Source/Developer/TraceServices/Private/Model/Counters.cpp
--- a/Engine/Source/Developer/TraceServices/Private/Model/Counters.cpp
+++ b/Engine/Source/Developer/TraceServices/Private/Model/Counters.cpp
@@ -9,7 +9,7 @@ namespace Trace
 
 const FName FCounterProvider::ProviderName("CounterProvider");
 
-FCounter::FCounter(ILinearAllocator& Allocator, const TPagedArray<double>& InFrameStartTimes, uint32 InId)
+FCounter::FCounter(ILinearAllocator& Allocator, const TArray<double>& InFrameStartTimes, uint32 InId)
 	: FrameStartTimes(InFrameStartTimes)
 	, IntCounterData(Allocator)
 	, DoubleCounterData(Allocator)
@@ -19,9 +19,10 @@ FCounter::FCounter(ILinearAllocator& Allocator, const TPagedArray<double>& InFra
 }
 
 template<typename CounterType, typename EnumerationType>
-static void EnumerateCounterValuesInternal(const TCounterData<CounterType>& CounterData, const TPagedArray<double>& FrameStartTimes, bool bResetEveryFrame, double IntervalStart, double IntervalEnd, TFunctionRef<void(double, EnumerationType)> Callback)
+static void EnumerateCounterValuesInternal(const TCounterData<CounterType>& CounterData, const TArray<double>& FrameStartTimes, bool bResetEveryFrame, double IntervalStart, double IntervalEnd, TFunctionRef<void(double, EnumerationType)> Callback)
 {
-	auto CounterIterator = bResetEveryFrame ? CounterData.GetIterator(FrameStartTimes) : CounterData.GetIterator();
+	TArray<double> NoFrameStartTimes;
+	auto CounterIterator = bResetEveryFrame ? CounterData.GetIterator(FrameStartTimes) : CounterData.GetIterator(NoFrameStartTimes);
 	while (CounterIterator)
 	{
 		const TTuple<double, CounterType>& Current = *CounterIterator;
@@ -124,14 +125,11 @@ void FCounter::SetValue(double Time, double Value)
 FCounterProvider::FCounterProvider(IAnalysisSession& InSession, IFrameProvider& InFrameProvider)
 	: Session(InSession)
 	, FrameProvider(InFrameProvider)
-	, FrameStartTimes(InSession.GetLinearAllocator(), 1024)
 {
-	FrameAddedDelegateHandle = FrameProvider.OnFrameAdded().AddRaw(this, &FCounterProvider::FrameAdded);
 }
 
 FCounterProvider::~FCounterProvider()
 {
-	FrameProvider.OnFrameAdded().Remove(FrameAddedDelegateHandle);
 	for (FCounter* Counter : Counters)
 	{
 		delete Counter;
@@ -159,19 +157,11 @@ bool FCounterProvider::ReadCounter(uint32 CounterId, TFunctionRef<void(const ICo
 
 ICounter* FCounterProvider::CreateCounter()
 {
-	FCounter* Counter = new FCounter(Session.GetLinearAllocator(), FrameStartTimes, Counters.Num());
+	FCounter* Counter = new FCounter(Session.GetLinearAllocator(), FrameProvider.GetFrameStartTimes(TraceFrameType_Game), Counters.Num());
 	Counters.Add(Counter);
 	return Counter;
 }
 
-void FCounterProvider::FrameAdded(const FFrame& Frame)
-{
-	if (Frame.FrameType == TraceFrameType_Game)
-	{
-		FrameStartTimes.PushBack() = Frame.StartTime;
-	}
-}
-
 const ICounterProvider& ReadCounterProvider(const IAnalysisSession& Session)
 {
 	return *Session.ReadProvider<ICounterProvider>(FCounterProvider::ProviderName);
diff --git a/Engine/Source/Developer/TraceServices/Private/Model/CountersPrivate.h b/Engine/Source/Developer/TraceServices/Private/Model/CountersPrivate.h
--- a/Engine/Source/Developer/TraceServices/Private/Model/CountersPrivate.h
+++ b/Engine/Source/Developer/TraceServices/Private/Model/CountersPrivate.h
@@ -24,8 +24,9 @@ template<typename ValueType>
 class TCounterDataIterator
 {
 public:
-	TCounterDataIterator(const TCounterData<ValueType>& Outer)
-		: TimestampsIterator(Outer.Timestamps.GetIterator())
+	TCounterDataIterator(const TCounterData<ValueType>& Outer, const TArray<double>& FrameStartTimes)
+		: FrameStartTimesIterator(FrameStartTimes.CreateConstIterator())
+		, TimestampsIterator(Outer.Timestamps.GetIterator())
 		, OpTypesIterator(Outer.OpTypes.GetIterator())
 		, OpArgumentsIterator(Outer.OpArguments.GetIterator())
 	{
@@ -33,12 +34,6 @@ public:
 		UpdateValue();
 	}
 
-	TCounterDataIterator(const TCounterData<ValueType>& Outer, const TPagedArray<double>& FrameStartTimes)
-		: TCounterDataIterator(Outer)
-	{
-		FrameStartTimesIterator = FrameStartTimes.GetIterator();
-	}
-
 	const TTuple<double, ValueType>& operator*() const
 	{
 		return Current;
@@ -97,7 +92,7 @@ private:
 		}
 	}
 
-	TPagedArray<double>::TIterator FrameStartTimesIterator;
+	TArray<double>::TConstIterator FrameStartTimesIterator;
 	TPagedArray<double>::TIterator TimestampsIterator;
 	TPagedArray<ECounterOpType>::TIterator OpTypesIterator;
 	typename TPagedArray<ValueType>::TIterator OpArgumentsIterator;
@@ -145,12 +140,7 @@ public:
 		OpArguments.Insert(InsertionIndex) = OpArgument;
 	}
 
-	TIterator GetIterator() const
-	{
-		return TIterator(*this);
-	}
-
-	TIterator GetIterator(TPagedArray<double> FrameStartTimes) const
+	TIterator GetIterator(TArray<double> FrameStartTimes) const
 	{
 		return TIterator(*this, FrameStartTimes);
 	}
@@ -168,7 +158,7 @@ class FCounter
 	: public ICounter
 {
 public:
-	FCounter(ILinearAllocator& Allocator, const TPagedArray<double>& FrameStartTimes, uint32 Id);
+	FCounter(ILinearAllocator& Allocator, const TArray<double>& FrameStartTimes, uint32 Id);
 	virtual uint32 GetId() const override { return Id; }
 	virtual const TCHAR* GetName() const override { return Name; }
 	virtual void SetName(const TCHAR* InName) override { Name = InName; }
@@ -187,7 +177,7 @@ public:
 	virtual void SetValue(double Time, double Value) override;
 
 private:
-	const TPagedArray<double>& FrameStartTimes;
+	const TArray<double>& FrameStartTimes;
 	TCounterData<int64> IntCounterData;
 	TCounterData<double> DoubleCounterData;
 	const TCHAR* Name = nullptr;
@@ -213,13 +203,9 @@ public:
 	virtual ICounter* CreateCounter() override;
 
 private:
-	void FrameAdded(const FFrame& Frame);
-
 	IAnalysisSession& Session;
 	IFrameProvider& FrameProvider;
 	TArray<FCounter*> Counters;
-	FDelegateHandle FrameAddedDelegateHandle;
-	TPagedArray<double> FrameStartTimes;
 };
 
 }
diff --git a/Engine/Source/Developer/TraceServices/Private/Model/CsvProfiler.cpp b/Engine/Source/Developer/TraceServices/Private/Model/CsvProfiler.cpp
new file mode 100644
--- /dev/null
+++ b/Engine/Source/Developer/TraceServices/Private/Model/CsvProfiler.cpp
@@ -0,0 +1,302 @@
+// Copyright 1998-2019 Epic Games, Inc. All Rights Reserved.
+
+#include "TraceServices/Model/CsvProfiler.h"
+#include "CsvProfilerPrivate.h"
+
+namespace Trace
+{
+
+uint64 FCsvProfilerProvider::FTableLayout::GetColumnCount() const
+{
+	return StatSeries.Num() + 1;
+}
+
+const TCHAR* FCsvProfilerProvider::FTableLayout::GetColumnName(uint64 ColumnIndex) const
+{
+	if (ColumnIndex == 0)
+	{
+		return TEXT("EVENTS");
+	}
+	else
+	{
+		return StatSeries[ColumnIndex - 1]->Name;
+	}
+}
+
+ETableColumnType FCsvProfilerProvider::FTableLayout::GetColumnType(uint64 ColumnIndex) const
+{
+	if (ColumnIndex == 0)
+	{
+		return TableColumnType_CString;
+	}
+	else
+	{
+		ECsvStatSeriesType SeriesType = StatSeries[ColumnIndex - 1]->Type;
+		if (SeriesType == CsvStatSeriesType_CustomStatInt)
+		{
+			return TableColumnType_Int;
+		}
+		else
+		{
+			return TableColumnType_Double;
+		}
+	}
+}
+
+bool FCsvProfilerProvider::FTableReader::IsValid() const
+{
+	return CurrentRowIndex < (Capture.EndFrame - Capture.StartFrame);
+}
+
+void FCsvProfilerProvider::FTableReader::NextRow()
+{
+	++CurrentRowIndex;
+}
+
+void FCsvProfilerProvider::FTableReader::SetRowIndex(uint64 RowIndex)
+{
+	CurrentRowIndex = RowIndex;
+}
+
+bool FCsvProfilerProvider::FTableReader::GetValueBool(uint64 ColumnIndex) const
+{
+	return false;
+}
+
+const FCsvProfilerProvider::FStatSeriesValue* FCsvProfilerProvider::FTableReader::GetValue(uint64 ColumnIndex) const
+{
+	if (!IsValid() || ColumnIndex == 0 || Capture.StatSeries.Num() <= ColumnIndex - 1)
+	{
+		return nullptr;
+	}
+	const FStatSeriesValue* FindIt = Capture.StatSeries[ColumnIndex - 1]->Values.Find(Capture.StartFrame + CurrentRowIndex);
+	return FindIt;
+}
+
+int64 FCsvProfilerProvider::FTableReader::GetValueInt(uint64 ColumnIndex) const
+{
+	const FStatSeriesValue* Value = GetValue(ColumnIndex);
+	if (!Value)
+	{
+		return 0;
+	}
+	return Value->Value.AsInt;
+}
+
+float FCsvProfilerProvider::FTableReader::GetValueFloat(uint64 ColumnIndex) const
+{
+	return static_cast<float>(GetValueDouble(ColumnIndex));
+}
+
+double FCsvProfilerProvider::FTableReader::GetValueDouble(uint64 ColumnIndex) const
+{
+	const FStatSeriesValue* Value = GetValue(ColumnIndex);
+	if (!Value)
+	{
+		return 0;
+	}
+	return Value->Value.AsDouble;
+}
+
+const TCHAR* FCsvProfilerProvider::FTableReader::GetValueCString(uint64 ColumnIndex) const
+{
+	if (ColumnIndex != 0)
+	{
+		return nullptr;
+	}
+	const FEvents* const* FindIt = Events.Find(Capture.StartFrame + CurrentRowIndex);
+	if (!FindIt)
+	{
+		return TEXT("");
+	}
+	return (*FindIt)->SemiColonSeparatedEvents;
+}
+
+uint64 FCsvProfilerProvider::FTable::GetRowCount() const
+{
+	return Capture.EndFrame - Capture.StartFrame;
+}
+
+IUntypedTableReader* FCsvProfilerProvider::FTable::CreateReader() const
+{
+	return new FCsvProfilerProvider::FTableReader(Capture, Events);
+}
+
+FCsvProfilerProvider::FCsvProfilerProvider(IAnalysisSession& InSession)
+	: Session(InSession)
+{
+
+}
+
+FCsvProfilerProvider::~FCsvProfilerProvider()
+{
+	for (FCapture* Capture : Captures)
+	{
+		delete Capture;
+	}
+	delete CurrentCapture;
+	for (FStatSeries* Series : StatSeries)
+	{
+		delete Series;
+	}
+}
+
+void FCsvProfilerProvider::EnumerateCaptures(TFunctionRef<void(const FCaptureInfo&)> Callback) const
+{
+	Session.ReadAccessCheck();
+
+	uint32 Index = 0;
+	for (const FCapture* Capture : Captures)
+	{
+		Callback({ Capture->Metadata, Capture->Filename, Index++, uint32(Capture->EndFrame - Capture->StartFrame) });
+	}
+}
+
+void FCsvProfilerProvider::StartCapture(const TCHAR* Filename, int32 FrameNumber)
+{
+	Session.WriteAccessCheck();
+
+	check(!CurrentCapture);
+	CurrentCapture = new FCapture(Events);
+	CurrentCapture->StartFrame = FrameNumber;
+	CurrentCapture->EndFrame = FrameNumber;
+	CurrentCapture->Filename = Filename;
+}
+
+void FCsvProfilerProvider::EndCapture(int32 FrameNumber)
+{
+	Session.WriteAccessCheck();
+
+	check(CurrentCapture);
+
+	CurrentCapture->EndFrame = FrameNumber;
+	check(CurrentCapture->EndFrame >= CurrentCapture->StartFrame);
+
+	CurrentCapture->Metadata = Metadata;
+
+	Captures.Add(CurrentCapture);
+	CurrentCapture = nullptr;
+}
+
+uint64 FCsvProfilerProvider::AddSeries(const TCHAR* Name, ECsvStatSeriesType Type)
+{
+	Session.WriteAccessCheck();
+
+	FStatSeries* Series = new FStatSeries();
+	Series->Name = Name;
+	Series->Type = Type;
+	int32 ColumnIndex = StatSeries.Num();
+	StatSeries.Add(Series);
+	return ColumnIndex;
+}
+
+FCsvProfilerProvider::FStatSeriesValue& FCsvProfilerProvider::GetValueRef(uint64 SeriesHandle, int32 FrameNumber)
+{
+	FStatSeries* Series = StatSeries[SeriesHandle];
+	if (CurrentCapture && !Series->Captures.Contains(Captures.Num()))
+	{
+		CurrentCapture->StatSeries.Add(Series);
+		Series->Captures.Add(Captures.Num());
+	}
+	return Series->Values.FindOrAdd(FrameNumber);
+}
+
+void FCsvProfilerProvider::SetTimerValue(uint64 SeriesHandle, int32 FrameNumber, double ElapsedTime)
+{
+	Session.WriteAccessCheck();
+
+	FStatSeriesValue& RowValue = GetValueRef(SeriesHandle, FrameNumber);
+	RowValue.Value.AsDouble += ElapsedTime;
+	RowValue.bIsValid = true;
+}
+
+void FCsvProfilerProvider::SetCustomStatValue(uint64 SeriesHandle, int32 FrameNumber, ECsvOpType OpType, int32 Value)
+{
+	Session.WriteAccessCheck();
+
+	FStatSeriesValue& RowValue = GetValueRef(SeriesHandle, FrameNumber);
+
+	if (!RowValue.bIsValid)
+	{
+		// The first op in a frame is always a set. Otherwise min/max don't work
+		OpType = CsvOpType_Set;
+	}
+
+	switch (OpType)
+	{
+	case CsvOpType_Set:
+		RowValue.Value.AsInt = Value;
+		break;
+	case CsvOpType_Min:
+		RowValue.Value.AsInt = FMath::Min(int64(Value), RowValue.Value.AsInt);
+		break;
+	case CsvOpType_Max:
+		RowValue.Value.AsInt = FMath::Max(int64(Value), RowValue.Value.AsInt);
+		break;
+	case CsvOpType_Accumulate:
+		RowValue.Value.AsInt += Value;
+		break;
+	}
+	RowValue.bIsValid = true;
+}
+
+void FCsvProfilerProvider::SetCustomStatValue(uint64 SeriesHandle, int32 FrameNumber, ECsvOpType OpType, float Value)
+{
+	Session.WriteAccessCheck();
+
+	FStatSeriesValue& RowValue = GetValueRef(SeriesHandle, FrameNumber);
+
+	if (!RowValue.bIsValid)
+	{
+		// The first op in a frame is always a set. Otherwise min/max don't work
+		OpType = CsvOpType_Set;
+	}
+
+	switch (OpType)
+	{
+	case CsvOpType_Set:
+		RowValue.Value.AsDouble = Value;
+		break;
+	case CsvOpType_Min:
+		RowValue.Value.AsDouble = FMath::Min(double(Value), RowValue.Value.AsDouble);
+		break;
+	case CsvOpType_Max:
+		RowValue.Value.AsDouble = FMath::Max(double(Value), RowValue.Value.AsDouble);
+		break;
+	case CsvOpType_Accumulate:
+		RowValue.Value.AsDouble += Value;
+		break;
+	}
+	RowValue.bIsValid = true;
+}
+
+void FCsvProfilerProvider::AddEvent(int32 FrameNumber, const TCHAR* Text)
+{
+	FEvents* FrameEvents;
+	FEvents** FindIt = Events.Find(FrameNumber);
+	if (!FindIt)
+	{
+		FrameEvents = new FEvents();
+		Events.Add(FrameNumber, FrameEvents);
+	}
+	else
+	{
+		FrameEvents = *FindIt;
+	}
+	if (FrameEvents->Events.Num())
+	{
+		FrameEvents->SemiColonSeparatedEvents = Session.StoreString(*(FString(FrameEvents->SemiColonSeparatedEvents) + ";" + Text));
+	}
+	else
+	{
+		FrameEvents->SemiColonSeparatedEvents = Text;
+	}
+	FrameEvents->Events.Add(Text);
+}
+
+void FCsvProfilerProvider::SetMetadata(const TCHAR* Key, const TCHAR* Value)
+{
+	Metadata.Add(Key, Value);
+}
+
+}
diff --git a/Engine/Source/Developer/TraceServices/Private/Model/CsvProfilerPrivate.h b/Engine/Source/Developer/TraceServices/Private/Model/CsvProfilerPrivate.h
new file mode 100644
--- /dev/null
+++ b/Engine/Source/Developer/TraceServices/Private/Model/CsvProfilerPrivate.h
@@ -0,0 +1,167 @@
+// Copyright 1998-2019 Epic Games, Inc. All Rights Reserved.
+
+#pragma once
+
+#include "TraceServices/Model/CsvProfiler.h"
+#include "TraceServices/Containers/Tables.h"
+#include "Containers/Map.h"
+
+namespace Trace
+{
+
+enum ECsvStatSeriesType
+{
+	CsvStatSeriesType_Timer,
+	CsvStatSeriesType_ExclusiveTimer,
+	CsvStatSeriesType_CustomStatInt,
+	CsvStatSeriesType_CustomStatFloat
+};
+
+enum ECsvOpType
+{
+	CsvOpType_Set,
+	CsvOpType_Min,
+	CsvOpType_Max,
+	CsvOpType_Accumulate,
+};
+
+class FCsvProfilerProvider
+	: public ICsvProfilerProvider
+{
+public:
+	FCsvProfilerProvider(IAnalysisSession& InSession);
+	virtual ~FCsvProfilerProvider();
+	virtual const IUntypedTable& GetTable(uint32 CaptureIndex) const override { return Captures[CaptureIndex]->Table; }
+	virtual void EnumerateCaptures(TFunctionRef<void(const FCaptureInfo&)> Callback) const override;
+	void StartCapture(const TCHAR* Filename, int32 FrameNumber);
+	void EndCapture(int32 FrameNumber);
+	uint64 AddSeries(const TCHAR* Name, ECsvStatSeriesType Type);
+	void SetTimerValue(uint64 SeriesHandle, int32 FrameNumber, double ElapsedTime);
+	void SetCustomStatValue(uint64 SeriesHandle, int32 FrameNumber, ECsvOpType OpType, int32 Value);
+	void SetCustomStatValue(uint64 SeriesHandle, int32 FrameNumber, ECsvOpType OpType, float Value);
+	void AddEvent(int32 FrameNumber, const TCHAR* Text);
+	void SetMetadata(const TCHAR* Key, const TCHAR* Value);
+
+private:
+	struct FStatSeriesValue
+	{
+		FStatSeriesValue() { Value.AsInt = 0; }
+		union
+		{
+			int64 AsInt;
+			double AsDouble;
+		} Value;
+		bool bIsValid = false;
+	};
+
+	struct FStatSeries
+	{
+		TMap<uint64, FStatSeriesValue> Values;
+		TSet<uint32> Captures;
+		const TCHAR* Name = nullptr;
+		ECsvStatSeriesType Type;
+	};
+
+	struct FEvents
+	{
+		TArray<const TCHAR*> Events;
+		const TCHAR* SemiColonSeparatedEvents = nullptr;
+	};
+
+	struct FCapture;
+
+	class FTableLayout
+		: public ITableLayout
+	{
+	public:
+		FTableLayout(const TArray<FStatSeries*>& InStatSeries)
+			: StatSeries(InStatSeries)
+		{
+
+		}
+
+		virtual uint64 GetColumnCount() const override;
+		virtual const TCHAR* GetColumnName(uint64 ColumnIndex) const override;
+		virtual ETableColumnType GetColumnType(uint64 ColumnIndex) const override;
+
+	private:
+		const TArray<FStatSeries*>& StatSeries;
+	};
+
+	class FTableReader
+		: public IUntypedTableReader
+	{
+	public:
+		FTableReader(const FCapture& InCapture, const TMap<int32, FEvents*>& InEvents)
+			: Capture(InCapture)
+			, Events(InEvents)
+			, CurrentRowIndex(0)
+		{
+		}
+
+		virtual bool IsValid() const override;
+		virtual void NextRow() override;
+		virtual void SetRowIndex(uint64 RowIndex) override;
+		virtual bool GetValueBool(uint64 ColumnIndex) const override;
+		virtual int64 GetValueInt(uint64 ColumnIndex) const override;
+		virtual float GetValueFloat(uint64 ColumnIndex) const override;
+		virtual double GetValueDouble(uint64 ColumnIndex) const override;
+		virtual const TCHAR* GetValueCString(uint64 ColumnIndex) const override;
+
+	private:
+		const FStatSeriesValue* GetValue(uint64 ColumnIndex) const;
+
+		const FCapture& Capture;
+		const TMap<int32, FEvents*>& Events;
+		uint64 CurrentRowIndex;
+	};
+
+	class FTable
+		: public IUntypedTable
+	{
+	public:
+		FTable(const FCapture& InCapture, const TMap<int32, FEvents*>& InEvents)
+			: Layout(InCapture.StatSeries)
+			, Capture(InCapture)
+			, Events(InEvents)
+		{
+
+		}
+
+		virtual const ITableLayout& GetLayout() const override { return Layout; }
+		virtual uint64 GetRowCount() const override;
+		virtual IUntypedTableReader* CreateReader() const override;
+
+	private:
+		FTableLayout Layout;
+		const FCapture& Capture;
+		const TMap<int32, FEvents*>& Events;
+	};
+
+	struct FCapture
+	{
+		FCapture(const TMap<int32, FEvents*>& Events)
+			: Table(*this, Events)
+		{
+			
+		}
+
+		FTable Table;
+		const TCHAR* Filename;
+		int32 StartFrame = -1;
+		int32 EndFrame = -1;
+		TArray<FStatSeries*> StatSeries;
+		TMap<const TCHAR*, const TCHAR*> Metadata;
+	};
+
+	FStatSeriesValue& GetValueRef(uint64 SeriesHandle, int32 FrameNumber);
+
+	IAnalysisSession& Session;
+	TArray<FStatSeries*> StatSeries;
+	TMap<int32, FEvents*> Events;
+	TMap<const TCHAR*, const TCHAR*> Metadata;
+	TArray<FCapture*> Captures;
+	FCapture* CurrentCapture = nullptr;
+};
+
+}
\ No newline at end of file
diff --git a/Engine/Source/Developer/TraceServices/Private/Model/Frames.cpp b/Engine/Source/Developer/TraceServices/Private/Model/Frames.cpp
--- a/Engine/Source/Developer/TraceServices/Private/Model/Frames.cpp
+++ b/Engine/Source/Developer/TraceServices/Private/Model/Frames.cpp
@@ -50,6 +50,9 @@ void FFrameProvider::BeginFrame(ETraceFrameType FrameType, double Time)
 	Frame.EndTime = std::numeric_limits<double>::infinity();
 	Frame.Index = Index;
 	Frame.FrameType = FrameType;
+
+	FrameStartTimes[FrameType].Add(Time);
+
 	Session.UpdateDurationSeconds(Time);
 }
 
@@ -58,7 +61,6 @@ void FFrameProvider::EndFrame(ETraceFrameType FrameType, double Time)
 	Session.WriteAccessCheck();
 	FFrame& Frame = Frames[FrameType][Frames[FrameType].Num() - 1];
 	Frame.EndTime = Time;
-	OnFrameAddedDelegate.Broadcast(Frame);
 	Session.UpdateDurationSeconds(Time);
 }
 
diff --git a/Engine/Source/Developer/TraceServices/Private/Model/FramesPrivate.h b/Engine/Source/Developer/TraceServices/Private/Model/FramesPrivate.h
--- a/Engine/Source/Developer/TraceServices/Private/Model/FramesPrivate.h
+++ b/Engine/Source/Developer/TraceServices/Private/Model/FramesPrivate.h
@@ -19,14 +19,14 @@ public:
 
 	virtual uint64 GetFrameCount(ETraceFrameType FrameType) const override;
 	virtual void EnumerateFrames(ETraceFrameType FrameType, uint64 Start, uint64 End, TFunctionRef<void(const FFrame&)> Callback) const override;
+	virtual const TArray<double>& GetFrameStartTimes(ETraceFrameType FrameType) const override { return FrameStartTimes[FrameType]; }
 	void BeginFrame(ETraceFrameType FrameType, double Time);
 	void EndFrame(ETraceFrameType FrameType, double Time);
-	virtual FOnFrameAdded& OnFrameAdded() override { return OnFrameAddedDelegate; }
 
 private:
 	IAnalysisSession& Session;
 	TArray<TPagedArray<FFrame>> Frames;
-	FOnFrameAdded OnFrameAddedDelegate;
+	TArray<double> FrameStartTimes[TraceFrameType_Count];
 };
 
 }
\ No newline at end of file
diff --git a/Engine/Source/Developer/TraceServices/Private/Model/Tables.h b/Engine/Source/Developer/TraceServices/Private/Model/Tables.h
--- a/Engine/Source/Developer/TraceServices/Private/Model/Tables.h
+++ b/Engine/Source/Developer/TraceServices/Private/Model/Tables.h
@@ -14,7 +14,7 @@
 	{ \
 	public: \
 		typedef InRowType RowType; \
-		virtual uint8 GetColumnCount() const override { return decltype(LastColumn)::Index; } \
+		virtual uint64 GetColumnCount() const override { return decltype(LastColumn)::Index; } \
 		template<uint8 IndexValue> \
 		constexpr ::Trace::ETableColumnType GetColumnTypeInternal() const { return TableColumnType_Invalid; } \
 		template<uint8 IndexValue> \
@@ -286,7 +286,7 @@ public:
 		return CurrentRow;
 	}
 
-	virtual bool GetValueBool(uint8 ColumnIndex) const override
+	virtual bool GetValueBool(uint64 ColumnIndex) const override
 	{
 		if (!CurrentRow)
 		{
@@ -307,7 +307,7 @@ public:
 		return false;
 	}
 
-	virtual int64 GetValueInt(uint8 ColumnIndex) const override
+	virtual int64 GetValueInt(uint64 ColumnIndex) const override
 	{
 		if (!CurrentRow)
 		{
@@ -328,7 +328,7 @@ public:
 		return 0;
 	}
 
-	virtual float GetValueFloat(uint8 ColumnIndex) const override
+	virtual float GetValueFloat(uint64 ColumnIndex) const override
 	{
 		if (!CurrentRow)
 		{
@@ -349,7 +349,7 @@ public:
 		return 0.0;
 	}
 
-	virtual double GetValueDouble(uint8 ColumnIndex) const override
+	virtual double GetValueDouble(uint64 ColumnIndex) const override
 	{
 		if (!CurrentRow)
 		{
@@ -370,7 +370,7 @@ public:
 		return 0.0;
 	}
 
-	virtual const TCHAR* GetValueCString(uint8 ColumnIndex) const override
+	virtual const TCHAR* GetValueCString(uint64 ColumnIndex) const override
 	{
 		if (!CurrentRow)
 		{
@@ -396,7 +396,7 @@ class TTableLayoutBase
 	: public ITableLayout
 {
 public:
-	constexpr ETableColumnType GetColumnTypeConstExpr(uint8 ColumnIndex) const
+	constexpr ETableColumnType GetColumnTypeConstExpr(uint64 ColumnIndex) const
 	{
 		const LayoutType* This = static_cast<const LayoutType*>(this);
 		switch (ColumnIndex)
@@ -421,12 +421,12 @@ public:
 		return ::Trace::TableColumnType_Invalid;
 	}
 
-	virtual ETableColumnType GetColumnType(uint8 ColumnIndex) const override
+	virtual ETableColumnType GetColumnType(uint64 ColumnIndex) const override
 	{
 		return GetColumnTypeConstExpr(ColumnIndex);
 	}
 
-	constexpr const TCHAR* GetColumnNameConstExpr(uint8 ColumnIndex) const
+	constexpr const TCHAR* GetColumnNameConstExpr(uint64 ColumnIndex) const
 	{
 		const LayoutType* This = static_cast<const LayoutType*>(this);
 		switch (ColumnIndex)
@@ -451,12 +451,12 @@ public:
 		return nullptr;
 	}
 
-	virtual const TCHAR* GetColumnName(uint8 ColumnIndex) const override
+	virtual const TCHAR* GetColumnName(uint64 ColumnIndex) const override
 	{
 		return GetColumnNameConstExpr(ColumnIndex);
 	}
 
-	FColumnValueContainer GetColumnValue(const RowType& Row, uint8 ColumnIndex) const
+	FColumnValueContainer GetColumnValue(const RowType& Row, uint64 ColumnIndex) const
 	{
 		const LayoutType* This = static_cast<const LayoutType*>(this);
 		switch (ColumnIndex)
diff --git a/Engine/Source/Developer/TraceServices/Private/Model/Threads.cpp b/Engine/Source/Developer/TraceServices/Private/Model/Threads.cpp
--- a/Engine/Source/Developer/TraceServices/Private/Model/Threads.cpp
+++ b/Engine/Source/Developer/TraceServices/Private/Model/Threads.cpp
@@ -101,6 +101,18 @@ void FThreadProvider::EnumerateThreads(TFunctionRef<void(const FThreadInfo &)> C
 	}
 }
 
+const TCHAR* FThreadProvider::GetThreadName(uint32 ThreadId) const
+{
+	Session.ReadAccessCheck();
+
+	const FThreadInfoInternal* const* FindIt = ThreadMap.Find(ThreadId);
+	if (!FindIt)
+	{
+		return TEXT("");
+	}
+	return (*FindIt)->Name;
+}
+
 void FThreadProvider::SortThreads()
 {
 	Algo::SortBy(SortedThreads, [](const FThreadInfoInternal* In) -> const FThreadInfoInternal&
diff --git a/Engine/Source/Developer/TraceServices/Private/Model/ThreadsPrivate.h b/Engine/Source/Developer/TraceServices/Private/Model/ThreadsPrivate.h
--- a/Engine/Source/Developer/TraceServices/Private/Model/ThreadsPrivate.h
+++ b/Engine/Source/Developer/TraceServices/Private/Model/ThreadsPrivate.h
@@ -23,6 +23,7 @@ public:
 	void SetThreadGroup(uint32 Id, const TCHAR* GroupName);
 	virtual uint64 GetModCount() const override { return ModCount; }
 	virtual void EnumerateThreads(TFunctionRef<void(const FThreadInfo&)> Callback) const override;
+	virtual const TCHAR* GetThreadName(uint32 ThreadId) const override;
 
 private:
 	struct FThreadInfoInternal
diff --git a/Engine/Source/Developer/TraceServices/Private/ModuleService.cpp b/Engine/Source/Developer/TraceServices/Private/ModuleService.cpp
--- a/Engine/Source/Developer/TraceServices/Private/ModuleService.cpp
+++ b/Engine/Source/Developer/TraceServices/Private/ModuleService.cpp
@@ -40,8 +40,13 @@ void FModuleService::GetAvailableModules(TArray<FModuleInfo>& OutModules)
 	for (const auto& KV : ModulesMap)
 	{
 		IModule* Module = KV.Value;
-		FModuleInfo& ModuleInfo = OutModules.AddDefaulted_GetRef();
-		Module->GetModuleInfo(ModuleInfo);
+		TArray<const TCHAR*> ModuleLoggers;
+		Module->GetLoggers(ModuleLoggers);
+		if (ModuleLoggers.Num())
+		{
+			FModuleInfo& ModuleInfo = OutModules.AddDefaulted_GetRef();
+			Module->GetModuleInfo(ModuleInfo);
+		}
 	}
 }
 	
@@ -80,7 +85,7 @@ void FModuleService::OnAnalysisBegin(IAnalysisSession& Session)
 	}
 }
 
-bool FModuleService::GetModuleLoggers(const FName& ModuleName, TArray<const TCHAR *>& OutLoggers)
+bool FModuleService::GetModuleLoggers(const FName& ModuleName, TArray<const TCHAR*>& OutLoggers)
 {
 	FScopeLock Lock(&CriticalSection);
 	Initialize();
diff --git a/Engine/Source/Developer/TraceServices/Private/Modules/CsvProfilerModule.cpp b/Engine/Source/Developer/TraceServices/Private/Modules/CsvProfilerModule.cpp
--- a/Engine/Source/Developer/TraceServices/Private/Modules/CsvProfilerModule.cpp
+++ b/Engine/Source/Developer/TraceServices/Private/Modules/CsvProfilerModule.cpp
@@ -9,6 +9,7 @@ namespace Trace
 {
 
 static const FName CsvProfilerModuleName("TraceModule_CsvProfiler");
+static const FName CsvProfilerProviderName("CsvProfilerProvider");
 
 void FCsvProfilerModule::GetModuleInfo(FModuleInfo& OutModuleInfo)
 {
@@ -19,13 +20,20 @@ void FCsvProfilerModule::GetModuleInfo(FModuleInfo& OutModuleInfo)
 
 void FCsvProfilerModule::OnAnalysisBegin(IAnalysisSession& Session)
 {
-	ICounterProvider& CounterProvider = EditCounterProvider(Session);
-	Session.AddAnalyzer(new FCsvProfilerAnalyzer(Session, CounterProvider));
+	const IFrameProvider& FrameProvider = ReadFrameProvider(Session);
+	const IThreadProvider& ThreadProvider = ReadThreadProvider(Session);
+	FCsvProfilerProvider* CsvProfilerProvider = new FCsvProfilerProvider(Session);
+	Session.AddProvider(CsvProfilerProviderName, CsvProfilerProvider);
+	Session.AddAnalyzer(new FCsvProfilerAnalyzer(Session, *CsvProfilerProvider, FrameProvider, ThreadProvider));
 }
 
 void FCsvProfilerModule::GetLoggers(TArray<const TCHAR *>& OutLoggers)
 {
-	OutLoggers.Add(TEXT("CsvProfiler"));
+}
+
+const ICsvProfilerProvider* ReadCsvProfilerProvider(const IAnalysisSession& Session)
+{
+	return Session.ReadProvider<ICsvProfilerProvider>(CsvProfilerProviderName);
 }
 
 }
diff --git a/Engine/Source/Developer/TraceServices/Public/TraceServices/Containers/Tables.h b/Engine/Source/Developer/TraceServices/Public/TraceServices/Containers/Tables.h
--- a/Engine/Source/Developer/TraceServices/Public/TraceServices/Containers/Tables.h
+++ b/Engine/Source/Developer/TraceServices/Public/TraceServices/Containers/Tables.h
@@ -19,9 +19,9 @@ class ITableLayout
 {
 public:
 	virtual ~ITableLayout() = default;
-	virtual uint8 GetColumnCount() const = 0;
-	virtual const TCHAR* GetColumnName(uint8 ColumnIndex) const = 0;
-	virtual ETableColumnType GetColumnType(uint8 ColumnIndex) const = 0;
+	virtual uint64 GetColumnCount() const = 0;
+	virtual const TCHAR* GetColumnName(uint64 ColumnIndex) const = 0;
+	virtual ETableColumnType GetColumnType(uint64 ColumnIndex) const = 0;
 };
 
 class IUntypedTableReader
@@ -31,11 +31,11 @@ public:
 	virtual bool IsValid() const = 0;
 	virtual void NextRow() = 0;
 	virtual void SetRowIndex(uint64 RowIndex) = 0;
-	virtual bool GetValueBool(uint8 ColumnIndex) const = 0;
-	virtual int64 GetValueInt(uint8 ColumnIndex) const = 0;
-	virtual float GetValueFloat(uint8 ColumnIndex) const = 0;
-	virtual double GetValueDouble(uint8 ColumnIndex) const = 0;
-	virtual const TCHAR* GetValueCString(uint8 ColumnIndex) const = 0;
+	virtual bool GetValueBool(uint64 ColumnIndex) const = 0;
+	virtual int64 GetValueInt(uint64 ColumnIndex) const = 0;
+	virtual float GetValueFloat(uint64 ColumnIndex) const = 0;
+	virtual double GetValueDouble(uint64 ColumnIndex) const = 0;
+	virtual const TCHAR* GetValueCString(uint64 ColumnIndex) const = 0;
 };
 
 template<typename RowType>
diff --git a/Engine/Source/Developer/TraceServices/Public/TraceServices/Model/CsvProfiler.h b/Engine/Source/Developer/TraceServices/Public/TraceServices/Model/CsvProfiler.h
new file mode 100644
--- /dev/null
+++ b/Engine/Source/Developer/TraceServices/Public/TraceServices/Model/CsvProfiler.h
@@ -0,0 +1,31 @@
+// Copyright 1998-2019 Epic Games, Inc. All Rights Reserved.
+
+#pragma once
+
+#include "TraceServices/Model/AnalysisSession.h"
+#include "Templates/Function.h"
+
+namespace Trace
+{
+
+class IUntypedTable;
+
+struct FCaptureInfo
+{
+	const TMap<const TCHAR*, const TCHAR*>& Metadata;
+	const TCHAR* Filename = nullptr;
+	uint32 Id = uint32(-1);
+	uint32 FrameCount = 0;
+};
+
+class ICsvProfilerProvider
+	: public IProvider
+{
+public:
+	virtual void EnumerateCaptures(TFunctionRef<void(const FCaptureInfo&)> Callback) const = 0;
+	virtual const IUntypedTable& GetTable(uint32 CaptureId) const = 0;
+};
+
+const ICsvProfilerProvider* ReadCsvProfilerProvider(const IAnalysisSession& Session);
+
+}
\ No newline at end of file
diff --git a/Engine/Source/Developer/TraceServices/Public/TraceServices/Model/Frames.h b/Engine/Source/Developer/TraceServices/Public/TraceServices/Model/Frames.h
--- a/Engine/Source/Developer/TraceServices/Public/TraceServices/Model/Frames.h
+++ b/Engine/Source/Developer/TraceServices/Public/TraceServices/Model/Frames.h
@@ -25,8 +25,7 @@ public:
 	virtual ~IFrameProvider() = default;
 	virtual uint64 GetFrameCount(ETraceFrameType FrameType) const = 0;
 	virtual void EnumerateFrames(ETraceFrameType FrameType, uint64 Start, uint64 End, TFunctionRef<void(const FFrame&)> Callback) const = 0;
-	DECLARE_MULTICAST_DELEGATE_OneParam(FOnFrameAdded, const FFrame&);
-	virtual FOnFrameAdded& OnFrameAdded() = 0;
+	virtual const TArray<double>& GetFrameStartTimes(ETraceFrameType FrameType) const = 0;
 };
 
 TRACESERVICES_API const IFrameProvider& ReadFrameProvider(const IAnalysisSession& Session);
diff --git a/Engine/Source/Developer/TraceServices/Public/TraceServices/Model/Threads.h b/Engine/Source/Developer/TraceServices/Public/TraceServices/Model/Threads.h
--- a/Engine/Source/Developer/TraceServices/Public/TraceServices/Model/Threads.h
+++ b/Engine/Source/Developer/TraceServices/Public/TraceServices/Model/Threads.h
@@ -22,6 +22,7 @@ public:
 	virtual ~IThreadProvider() = default;
 	virtual uint64 GetModCount() const = 0;
 	virtual void EnumerateThreads(TFunctionRef<void(const FThreadInfo&)> Callback) const = 0;
+	virtual const TCHAR* GetThreadName(uint32 ThreadId) const = 0;
 };
 
 TRACESERVICES_API const IThreadProvider& ReadThreadProvider(const IAnalysisSession& Session);
diff --git a/Engine/Source/Runtime/Core/Private/ProfilingDebugging/CsvProfiler.cpp b/Engine/Source/Runtime/Core/Private/ProfilingDebugging/CsvProfiler.cpp
--- a/Engine/Source/Runtime/Core/Private/ProfilingDebugging/CsvProfiler.cpp
+++ b/Engine/Source/Runtime/Core/Private/ProfilingDebugging/CsvProfiler.cpp
@@ -414,6 +414,7 @@ public:
 					CategoryNames[Index] = CategoryName;
 					CategoryNameToIndex.Add(CategoryName.ToLower(), Index);
 				}
+				TRACE_CSV_PROFILER_REGISTER_CATEGORY(Index, *CategoryName);
 			}
 		}
 		return Index;
@@ -1745,44 +1746,58 @@ public:
 
 	CSV_PROFILER_INLINE void AddTimestampBegin(const char* StatName, int32 CategoryIndex)
 	{
-		TimingMarkers.ReserveElement()->Init(GetStatID(StatName), CategoryIndex, FCsvStatBase::FFlags::TimestampBegin, FPlatformTime::Cycles64());
+		uint64 Cycles = FPlatformTime::Cycles64();
+		TRACE_CSV_PROFILER_BEGIN_STAT(StatName, CategoryIndex, Cycles);
+		TimingMarkers.ReserveElement()->Init(GetStatID(StatName), CategoryIndex, FCsvStatBase::FFlags::TimestampBegin, Cycles);
 		TimingMarkers.CommitElement();
 	}
 
 	CSV_PROFILER_INLINE void AddTimestampEnd(const char* StatName, int32 CategoryIndex)
 	{
-		TimingMarkers.ReserveElement()->Init(GetStatID(StatName), CategoryIndex, 0, FPlatformTime::Cycles64());
+		uint64 Cycles = FPlatformTime::Cycles64();
+		TRACE_CSV_PROFILER_END_STAT(StatName, CategoryIndex, Cycles);
+		TimingMarkers.ReserveElement()->Init(GetStatID(StatName), CategoryIndex, 0, Cycles);
 		TimingMarkers.CommitElement();
 	}
 
 	CSV_PROFILER_INLINE void AddTimestampExclusiveBegin(const char* StatName)
 	{
-		TimingMarkers.ReserveElement()->Init(GetStatID(StatName), CSV_CATEGORY_INDEX(Exclusive), FCsvStatBase::FFlags::TimestampBegin | FCsvStatBase::FFlags::IsExclusiveTimestamp, FPlatformTime::Cycles64());
+		uint64 Cycles = FPlatformTime::Cycles64();
+		TRACE_CSV_PROFILER_BEGIN_EXCLUSIVE_STAT(StatName, CSV_CATEGORY_INDEX(Exclusive), Cycles);
+		TimingMarkers.ReserveElement()->Init(GetStatID(StatName), CSV_CATEGORY_INDEX(Exclusive), FCsvStatBase::FFlags::TimestampBegin | FCsvStatBase::FFlags::IsExclusiveTimestamp, Cycles);
 		TimingMarkers.CommitElement();
 	}
 
 	CSV_PROFILER_INLINE void AddTimestampExclusiveEnd(const char* StatName)
 	{
-		TimingMarkers.ReserveElement()->Init(GetStatID(StatName), CSV_CATEGORY_INDEX(Exclusive), FCsvStatBase::FFlags::IsExclusiveTimestamp, FPlatformTime::Cycles64());
+		uint64 Cycles = FPlatformTime::Cycles64();
+		TRACE_CSV_PROFILER_END_EXCLUSIVE_STAT(StatName, CSV_CATEGORY_INDEX(Exclusive), Cycles);
+		TimingMarkers.ReserveElement()->Init(GetStatID(StatName), CSV_CATEGORY_INDEX(Exclusive), FCsvStatBase::FFlags::IsExclusiveTimestamp, Cycles);
 		TimingMarkers.CommitElement();
 	}
 
 	CSV_PROFILER_INLINE void AddTimestampBegin(const FName& StatName, int32 CategoryIndex)
 	{
-		TimingMarkers.ReserveElement()->Init(GetStatID(StatName), CategoryIndex, FCsvStatBase::FFlags::StatIDIsFName | FCsvStatBase::FFlags::TimestampBegin, FPlatformTime::Cycles64());
+		uint64 Cycles = FPlatformTime::Cycles64();
+		TRACE_CSV_PROFILER_BEGIN_STAT(StatName, CategoryIndex, Cycles);
+		TimingMarkers.ReserveElement()->Init(GetStatID(StatName), CategoryIndex, FCsvStatBase::FFlags::StatIDIsFName | FCsvStatBase::FFlags::TimestampBegin, Cycles);
 		TimingMarkers.CommitElement();
 	}
 
 	CSV_PROFILER_INLINE void AddTimestampEnd(const FName& StatName, int32 CategoryIndex)
 	{
-		TimingMarkers.ReserveElement()->Init(GetStatID(StatName), CategoryIndex, FCsvStatBase::FFlags::StatIDIsFName, FPlatformTime::Cycles64());
+		uint64 Cycles = FPlatformTime::Cycles64();
+		TRACE_CSV_PROFILER_END_STAT(StatName, CategoryIndex, Cycles);
+		TimingMarkers.ReserveElement()->Init(GetStatID(StatName), CategoryIndex, FCsvStatBase::FFlags::StatIDIsFName, Cycles);
 		TimingMarkers.CommitElement();
 	}
 
 	CSV_PROFILER_INLINE void AddCustomStat(const char* StatName, const int32 CategoryIndex, const float Value, const ECsvCustomStatOp CustomStatOp)
 	{
 		FCsvCustomStat* CustomStat = CustomStats.ReserveElement();
-		CustomStat->Init(GetStatID(StatName), CategoryIndex, FCsvStatBase::FFlags::IsCustomStat, FPlatformTime::Cycles64(), uint8(CustomStatOp));
+		uint64 Cycles = FPlatformTime::Cycles64();
+		TRACE_CSV_PROFILER_CUSTOM_STAT(StatName, CategoryIndex, Value, uint8(CustomStatOp), Cycles);
+		CustomStat->Init(GetStatID(StatName), CategoryIndex, FCsvStatBase::FFlags::IsCustomStat, Cycles, uint8(CustomStatOp));
 		CustomStat->Value.AsFloat = Value;
 		CustomStats.CommitElement();
 	}
@@ -1790,7 +1805,9 @@ public:
 	CSV_PROFILER_INLINE void AddCustomStat(const FName& StatName, const int32 CategoryIndex, const float Value, const ECsvCustomStatOp CustomStatOp)
 	{
 		FCsvCustomStat* CustomStat = CustomStats.ReserveElement();
-		CustomStat->Init(GetStatID(StatName), CategoryIndex, FCsvStatBase::FFlags::IsCustomStat | FCsvStatBase::FFlags::StatIDIsFName, FPlatformTime::Cycles64(), uint8(CustomStatOp));
+		uint64 Cycles = FPlatformTime::Cycles64();
+		TRACE_CSV_PROFILER_CUSTOM_STAT(StatName, CategoryIndex, Value, uint8(CustomStatOp), Cycles);
+		CustomStat->Init(GetStatID(StatName), CategoryIndex, FCsvStatBase::FFlags::IsCustomStat | FCsvStatBase::FFlags::StatIDIsFName, Cycles, uint8(CustomStatOp));
 		CustomStat->Value.AsFloat = Value;
 		CustomStats.CommitElement();
 	}
@@ -1798,7 +1815,9 @@ public:
 	CSV_PROFILER_INLINE void AddCustomStat(const char* StatName, const int32 CategoryIndex, const int32 Value, const ECsvCustomStatOp CustomStatOp)
 	{
 		FCsvCustomStat* CustomStat = CustomStats.ReserveElement();
-		CustomStat->Init(GetStatID(StatName), CategoryIndex, FCsvStatBase::FFlags::IsCustomStat | FCsvStatBase::FFlags::IsInteger, FPlatformTime::Cycles64(), uint8(CustomStatOp));
+		uint64 Cycles = FPlatformTime::Cycles64();
+		TRACE_CSV_PROFILER_CUSTOM_STAT(StatName, CategoryIndex, Value, uint8(CustomStatOp), Cycles);
+		CustomStat->Init(GetStatID(StatName), CategoryIndex, FCsvStatBase::FFlags::IsCustomStat | FCsvStatBase::FFlags::IsInteger, Cycles, uint8(CustomStatOp));
 		CustomStat->Value.AsInt = Value;
 		CustomStats.CommitElement();
 	}
@@ -1806,7 +1825,9 @@ public:
 	CSV_PROFILER_INLINE void AddCustomStat(const FName& StatName, const int32 CategoryIndex, const int32 Value, const ECsvCustomStatOp CustomStatOp)
 	{
 		FCsvCustomStat* CustomStat = CustomStats.ReserveElement();
-		CustomStat->Init(GetStatID(StatName), CategoryIndex, FCsvStatBase::FFlags::IsCustomStat | FCsvStatBase::FFlags::IsInteger | FCsvStatBase::FFlags::StatIDIsFName, FPlatformTime::Cycles64(), uint8(CustomStatOp));
+		uint64 Cycles = FPlatformTime::Cycles64();
+		TRACE_CSV_PROFILER_CUSTOM_STAT(StatName, CategoryIndex, Value, uint8(CustomStatOp), Cycles);
+		CustomStat->Init(GetStatID(StatName), CategoryIndex, FCsvStatBase::FFlags::IsCustomStat | FCsvStatBase::FFlags::IsInteger | FCsvStatBase::FFlags::StatIDIsFName, Cycles, uint8(CustomStatOp));
 		CustomStat->Value.AsInt = Value;
 		CustomStats.CommitElement();
 	}
@@ -1814,14 +1835,17 @@ public:
 	CSV_PROFILER_INLINE void AddEvent(const FString& EventText, const int32 CategoryIndex)
 	{
 		FCsvEvent* Event = Events.ReserveElement();
+		uint64 Cycles = FPlatformTime::Cycles64();
+		TRACE_CSV_PROFILER_EVENT(*EventText, CategoryIndex, Cycles);
 		Event->EventText = EventText;
-		Event->Timestamp = FPlatformTime::Cycles64();
+		Event->Timestamp = Cycles;
 		Event->CategoryIndex = CategoryIndex;
 		Events.CommitElement();
 	}
 
 	CSV_PROFILER_INLINE void AddEventWithTimestamp(const FString& EventText, const int32 CategoryIndex, const uint64 Timestamp)
 	{
+		TRACE_CSV_PROFILER_EVENT(*EventText, CategoryIndex, Timestamp);
 		FCsvEvent* Event = Events.ReserveElement();
 		Event->EventText = EventText;
 		Event->Timestamp = Timestamp;
@@ -2541,6 +2565,7 @@ void FCsvProfiler::BeginFrame()
 
 					// Initialize tls before setting the capturing flag to true.
 					FCsvProfilerThreadData::InitTls();
+					TRACE_CSV_PROFILER_BEGIN_CAPTURE(*Filename, GRenderThreadId, GRHIThreadId, GDefaultWaitStatName, GCsvStatCounts);
 					GCsvProfilerIsCapturing = true;
 				}
 			}
@@ -2628,6 +2653,8 @@ void FCsvProfiler::EndFrame()
 				GCsvProfilerIsWritingFile = true;
 				GCsvProfilerIsCapturing = false;
 
+				TRACE_CSV_PROFILER_END_CAPTURE();
+
 				if (!ProcessingThread)
 				{
 					// Suspend the hang and hitch heartbeats, as this is a long running task.
@@ -2952,6 +2979,8 @@ void FCsvProfiler::RecordEvent(int32 CategoryIndex, const FString& EventText)
 
 void FCsvProfiler::SetMetadata(const TCHAR* Key, const TCHAR* Value)
 {
+	TRACE_CSV_PROFILER_METADATA(Key, Value);
+
 	LLM_SCOPE(ELLMTag::CsvProfiler);
 
 	// Always gather CSV metadata, even if we're not currently capturing.
@@ -2983,7 +3012,6 @@ void FCsvProfiler::RecordEventAtTimestamp(int32 CategoryIndex, const FString& Ev
 
 void FCsvProfiler::RecordCustomStat(const char * StatName, uint32 CategoryIndex, float Value, const ECsvCustomStatOp CustomStatOp)
 {
-	TRACE_CSV_PROFILER_CUSTOM_STAT(StatName, Value, uint8(CustomStatOp));
 	if (GCsvProfilerIsCapturing && GCsvCategoriesEnabled[CategoryIndex])
 	{
 		LLM_SCOPE(ELLMTag::CsvProfiler);
@@ -2993,7 +3021,6 @@ void FCsvProfiler::RecordCustomStat(const char * StatName, uint32 CategoryIndex,
 
 void FCsvProfiler::RecordCustomStat(const FName& StatName, uint32 CategoryIndex, float Value, const ECsvCustomStatOp CustomStatOp)
 {
-	TRACE_CSV_PROFILER_CUSTOM_STAT(StatName, Value, uint8(CustomStatOp));
 	if (GCsvProfilerIsCapturing && GCsvCategoriesEnabled[CategoryIndex])
 	{
 		LLM_SCOPE(ELLMTag::CsvProfiler);
@@ -3003,7 +3030,6 @@ void FCsvProfiler::RecordCustomStat(const FName& StatName, uint32 CategoryIndex,
 
 void FCsvProfiler::RecordCustomStat(const char * StatName, uint32 CategoryIndex, int32 Value, const ECsvCustomStatOp CustomStatOp)
 {
-	TRACE_CSV_PROFILER_CUSTOM_STAT(StatName, Value, uint8(CustomStatOp));
 	if (GCsvProfilerIsCapturing && GCsvCategoriesEnabled[CategoryIndex])
 	{
 		LLM_SCOPE(ELLMTag::CsvProfiler);
@@ -3013,7 +3039,6 @@ void FCsvProfiler::RecordCustomStat(const char * StatName, uint32 CategoryIndex,
 
 void FCsvProfiler::RecordCustomStat(const FName& StatName, uint32 CategoryIndex, int32 Value, const ECsvCustomStatOp CustomStatOp)
 {
-	TRACE_CSV_PROFILER_CUSTOM_STAT(StatName, Value, uint8(CustomStatOp));
 	if (GCsvProfilerIsCapturing && GCsvCategoriesEnabled[CategoryIndex])
 	{
 		LLM_SCOPE(ELLMTag::CsvProfiler);
diff --git a/Engine/Source/Runtime/Core/Private/ProfilingDebugging/CsvProfilerTrace.cpp b/Engine/Source/Runtime/Core/Private/ProfilingDebugging/CsvProfilerTrace.cpp
--- a/Engine/Source/Runtime/Core/Private/ProfilingDebugging/CsvProfilerTrace.cpp
+++ b/Engine/Source/Runtime/Core/Private/ProfilingDebugging/CsvProfilerTrace.cpp
@@ -1,96 +1,281 @@
 // Copyright 1998-2019 Epic Games, Inc. All Rights Reserved.
 #include "ProfilingDebugging/CsvProfilerTrace.h"
+#include "ProfilingDebugging/CsvProfiler.h"
 #include "HAL/PlatformTime.h"
 #include "UObject/NameTypes.h"
 
 #if CSVPROFILERTRACE_ENABLED
 
-UE_TRACE_EVENT_BEGIN(CsvProfiler, InlineStat, Always)
-	UE_TRACE_EVENT_FIELD(const char*, StatNamePointer)
+UE_TRACE_EVENT_BEGIN(CsvProfiler, RegisterCategory, Always)
+	UE_TRACE_EVENT_FIELD(int32, Index)
 UE_TRACE_EVENT_END()
 
-UE_TRACE_EVENT_BEGIN(CsvProfiler, DeclaredStat, Always)
-	UE_TRACE_EVENT_FIELD(uint32, StatId)
+UE_TRACE_EVENT_BEGIN(CsvProfiler, DefineInlineStat, Always)
+	UE_TRACE_EVENT_FIELD(uint64, StatId)
+	UE_TRACE_EVENT_FIELD(int32, CategoryIndex)
 UE_TRACE_EVENT_END()
 
-UE_TRACE_EVENT_BEGIN(CsvProfiler, CustomStatInlineInt)
+UE_TRACE_EVENT_BEGIN(CsvProfiler, DefineDeclaredStat, Always)
+	UE_TRACE_EVENT_FIELD(uint64, StatId)
+	UE_TRACE_EVENT_FIELD(int32, CategoryIndex)
+UE_TRACE_EVENT_END()
+
+UE_TRACE_EVENT_BEGIN(CsvProfiler, BeginStat, Always)
+	UE_TRACE_EVENT_FIELD(uint64, StatId)
 	UE_TRACE_EVENT_FIELD(uint64, Cycle)
-	UE_TRACE_EVENT_FIELD(const char*, StatNamePointer)
-	UE_TRACE_EVENT_FIELD(int32, Value)
-	UE_TRACE_EVENT_FIELD(uint8, OpType)
+	UE_TRACE_EVENT_FIELD(uint32, ThreadId)
 UE_TRACE_EVENT_END()
 
-UE_TRACE_EVENT_BEGIN(CsvProfiler, CustomStatInlineFloat)
+UE_TRACE_EVENT_BEGIN(CsvProfiler, EndStat, Always)
+	UE_TRACE_EVENT_FIELD(uint64, StatId)
 	UE_TRACE_EVENT_FIELD(uint64, Cycle)
-	UE_TRACE_EVENT_FIELD(const char*, StatNamePointer)
-	UE_TRACE_EVENT_FIELD(float, Value)
-	UE_TRACE_EVENT_FIELD(uint8, OpType)
+	UE_TRACE_EVENT_FIELD(uint32, ThreadId)
+UE_TRACE_EVENT_END()
+
+UE_TRACE_EVENT_BEGIN(CsvProfiler, BeginExclusiveStat, Always)
+	UE_TRACE_EVENT_FIELD(uint64, StatId)
+	UE_TRACE_EVENT_FIELD(uint64, Cycle)
+	UE_TRACE_EVENT_FIELD(uint32, ThreadId)
+UE_TRACE_EVENT_END()
+
+UE_TRACE_EVENT_BEGIN(CsvProfiler, EndExclusiveStat, Always)
+	UE_TRACE_EVENT_FIELD(uint64, StatId)
+	UE_TRACE_EVENT_FIELD(uint64, Cycle)
+	UE_TRACE_EVENT_FIELD(uint32, ThreadId)
 UE_TRACE_EVENT_END()
 
-UE_TRACE_EVENT_BEGIN(CsvProfiler, CustomStatDeclaredInt)
+UE_TRACE_EVENT_BEGIN(CsvProfiler, CustomStatInt, Always)
+	UE_TRACE_EVENT_FIELD(uint64, StatId)
 	UE_TRACE_EVENT_FIELD(uint64, Cycle)
-	UE_TRACE_EVENT_FIELD(uint32, StatId)
+	UE_TRACE_EVENT_FIELD(uint32, ThreadId)
 	UE_TRACE_EVENT_FIELD(int32, Value)
 	UE_TRACE_EVENT_FIELD(uint8, OpType)
 UE_TRACE_EVENT_END()
 
-UE_TRACE_EVENT_BEGIN(CsvProfiler, CustomStatDeclaredFloat)
+UE_TRACE_EVENT_BEGIN(CsvProfiler, CustomStatFloat, Always)
+	UE_TRACE_EVENT_FIELD(uint64, StatId)
 	UE_TRACE_EVENT_FIELD(uint64, Cycle)
-	UE_TRACE_EVENT_FIELD(uint32, StatId)
+	UE_TRACE_EVENT_FIELD(uint32, ThreadId)
 	UE_TRACE_EVENT_FIELD(float, Value)
 	UE_TRACE_EVENT_FIELD(uint8, OpType)
 UE_TRACE_EVENT_END()
 
-void FCsvProfilerTrace::OutputInlineStat(const char* StatName)
+UE_TRACE_EVENT_BEGIN(CsvProfiler, Event, Always)
+	UE_TRACE_EVENT_FIELD(uint64, Cycle)
+	UE_TRACE_EVENT_FIELD(int32, CategoryIndex)
+	UE_TRACE_EVENT_FIELD(uint32, ThreadId)
+UE_TRACE_EVENT_END()
+
+UE_TRACE_EVENT_BEGIN(CsvProfiler, BeginCapture, Always)
+	UE_TRACE_EVENT_FIELD(uint64, Cycle)
+	UE_TRACE_EVENT_FIELD(uint32, RenderThreadId)
+	UE_TRACE_EVENT_FIELD(uint32, RHIThreadId)
+	UE_TRACE_EVENT_FIELD(bool, EnableCounts)
+UE_TRACE_EVENT_END()
+
+UE_TRACE_EVENT_BEGIN(CsvProfiler, EndCapture, Always)
+	UE_TRACE_EVENT_FIELD(uint64, Cycle)
+UE_TRACE_EVENT_END()
+
+UE_TRACE_EVENT_BEGIN(CsvProfiler, Metadata, Always)
+	UE_TRACE_EVENT_FIELD(uint16, ValueOffset)
+UE_TRACE_EVENT_END()
+
+struct FCsvProfilerTraceInternal
+{
+	union FStatId
+	{
+		struct
+		{
+			uint64 IsFName : 1;
+			uint64 CategoryIndex : 11;
+			uint64 FNameOrCString : 52;
+		} Fields;
+		uint64 Hash;
+	};
+
+	FORCEINLINE static uint64 GetStatId(const char* StatName, int32 CategoryIndex)
+	{
+		FStatId StatId;
+		StatId.Fields.IsFName = false;
+		StatId.Fields.CategoryIndex = CategoryIndex;
+		StatId.Fields.FNameOrCString = uint64(StatName);
+		return StatId.Hash;
+	}
+
+	FORCEINLINE static uint64 GetStatId(const FName& StatName, int32 CategoryIndex)
+	{
+		FStatId StatId;
+		StatId.Fields.IsFName = true;
+		StatId.Fields.CategoryIndex = CategoryIndex;
+		StatId.Fields.FNameOrCString = StatName.GetComparisonIndex().ToUnstableInt();
+		return StatId.Hash;
+	}
+};
+
+void FCsvProfilerTrace::OutputRegisterCategory(int32 Index, const TCHAR* Name)
+{
+	uint16 NameSize = (FCString::Strlen(Name) + 1) * sizeof(TCHAR);
+	UE_TRACE_LOG(CsvProfiler, RegisterCategory, NameSize)
+		<< RegisterCategory.Index(Index)
+		<< RegisterCategory.Attachment(Name, NameSize);
+}
+
+void FCsvProfilerTrace::OutputInlineStat(const char* StatName, int32 CategoryIndex)
 {
 	uint16 NameSize = (strlen(StatName) + 1) * sizeof(char);
-	UE_TRACE_LOG(CsvProfiler, InlineStat, NameSize)
-		<< InlineStat.StatNamePointer(StatName)
-		<< InlineStat.Attachment(StatName, NameSize);
+	UE_TRACE_LOG(CsvProfiler, DefineInlineStat, NameSize)
+		<< DefineInlineStat.StatId(FCsvProfilerTraceInternal::GetStatId(StatName, CategoryIndex))
+		<< DefineInlineStat.CategoryIndex(CategoryIndex)
+		<< DefineInlineStat.Attachment(StatName, NameSize);
+}
+
+CSV_DECLARE_CATEGORY_EXTERN(Exclusive);
+
+void FCsvProfilerTrace::OutputInlineStatExclusive(const char* StatName)
+{
+	OutputInlineStat(StatName, CSV_CATEGORY_INDEX(Exclusive));
+}
+
+void FCsvProfilerTrace::OutputDeclaredStat(const FName& StatName, int32 CategoryIndex)
+{
+	TCHAR NameString[NAME_SIZE];
+	StatName.GetPlainNameString(NameString);
+	uint16 NameSize = (FCString::Strlen(NameString) + 1) * sizeof(TCHAR);
+	UE_TRACE_LOG(CsvProfiler, DefineDeclaredStat, NameSize)
+		<< DefineDeclaredStat.StatId(FCsvProfilerTraceInternal::GetStatId(StatName, CategoryIndex))
+		<< DefineDeclaredStat.CategoryIndex(CategoryIndex)
+		<< DefineDeclaredStat.Attachment(NameString, NameSize);
+}
+
+void FCsvProfilerTrace::OutputBeginStat(const char* StatName, int32 CategoryIndex, uint64 Cycles)
+{
+	UE_TRACE_LOG(CsvProfiler, BeginStat)
+		<< BeginStat.StatId(FCsvProfilerTraceInternal::GetStatId(StatName, CategoryIndex))
+		<< BeginStat.Cycle(Cycles)
+		<< BeginStat.ThreadId(FPlatformTLS::GetCurrentThreadId());
+}
+
+void FCsvProfilerTrace::OutputBeginStat(const FName& StatName, int32 CategoryIndex, uint64 Cycles)
+{
+	UE_TRACE_LOG(CsvProfiler, BeginStat)
+		<< BeginStat.StatId(FCsvProfilerTraceInternal::GetStatId(StatName, CategoryIndex))
+		<< BeginStat.Cycle(Cycles)
+		<< BeginStat.ThreadId(FPlatformTLS::GetCurrentThreadId());
+}
+
+void FCsvProfilerTrace::OutputEndStat(const char* StatName, int32 CategoryIndex, uint64 Cycles)
+{
+	UE_TRACE_LOG(CsvProfiler, EndStat)
+		<< EndStat.StatId(FCsvProfilerTraceInternal::GetStatId(StatName, CategoryIndex))
+		<< EndStat.Cycle(Cycles)
+		<< EndStat.ThreadId(FPlatformTLS::GetCurrentThreadId());
+}
+
+void FCsvProfilerTrace::OutputEndStat(const FName& StatName, int32 CategoryIndex, uint64 Cycles)
+{
+	UE_TRACE_LOG(CsvProfiler, EndStat)
+		<< EndStat.StatId(FCsvProfilerTraceInternal::GetStatId(StatName, CategoryIndex))
+		<< EndStat.Cycle(Cycles)
+		<< EndStat.ThreadId(FPlatformTLS::GetCurrentThreadId());
+}
+
+void FCsvProfilerTrace::OutputBeginExclusiveStat(const char* StatName, int32 CategoryIndex, uint64 Cycles)
+{
+	UE_TRACE_LOG(CsvProfiler, BeginExclusiveStat)
+		<< BeginExclusiveStat.StatId(FCsvProfilerTraceInternal::GetStatId(StatName, CategoryIndex))
+		<< BeginExclusiveStat.Cycle(Cycles)
+		<< BeginExclusiveStat.ThreadId(FPlatformTLS::GetCurrentThreadId());
+}
+
+void FCsvProfilerTrace::OutputEndExclusiveStat(const char* StatName, int32 CategoryIndex, uint64 Cycles)
+{
+	UE_TRACE_LOG(CsvProfiler, EndExclusiveStat)
+		<< EndExclusiveStat.StatId(FCsvProfilerTraceInternal::GetStatId(StatName, CategoryIndex))
+		<< EndExclusiveStat.Cycle(Cycles)
+		<< EndExclusiveStat.ThreadId(FPlatformTLS::GetCurrentThreadId());
+}
+
+void FCsvProfilerTrace::OutputCustomStat(const char* StatName, int32 CategoryIndex, int32 Value, uint8 OpType, uint64 Cycles)
+{
+	UE_TRACE_LOG(CsvProfiler, CustomStatInt)
+		<< CustomStatInt.Cycle(Cycles)
+		<< CustomStatInt.StatId(FCsvProfilerTraceInternal::GetStatId(StatName, CategoryIndex))
+		<< CustomStatInt.ThreadId(FPlatformTLS::GetCurrentThreadId())
+		<< CustomStatInt.Value(Value)
+		<< CustomStatInt.OpType(OpType);
+}
+
+void FCsvProfilerTrace::OutputCustomStat(const FName& StatName, int32 CategoryIndex, int32 Value, uint8 OpType, uint64 Cycles)
+{
+	UE_TRACE_LOG(CsvProfiler, CustomStatInt)
+		<< CustomStatInt.Cycle(Cycles)
+		<< CustomStatInt.StatId(FCsvProfilerTraceInternal::GetStatId(StatName, CategoryIndex))
+		<< CustomStatInt.ThreadId(FPlatformTLS::GetCurrentThreadId())
+		<< CustomStatInt.Value(Value)
+		<< CustomStatInt.OpType(OpType);
+}
+
+void FCsvProfilerTrace::OutputCustomStat(const char* StatName, int32 CategoryIndex, float Value, uint8 OpType, uint64 Cycles)
+{
+	UE_TRACE_LOG(CsvProfiler, CustomStatFloat)
+		<< CustomStatFloat.Cycle(Cycles)
+		<< CustomStatFloat.StatId(FCsvProfilerTraceInternal::GetStatId(StatName, CategoryIndex))
+		<< CustomStatFloat.ThreadId(FPlatformTLS::GetCurrentThreadId())
+		<< CustomStatFloat.Value(Value)
+		<< CustomStatFloat.OpType(OpType);
 }
 
-void FCsvProfilerTrace::OutputDeclaredStat(const FName& StatName, const TCHAR* StatNameString)
+void FCsvProfilerTrace::OutputCustomStat(const FName& StatName, int32 CategoryIndex, float Value, uint8 OpType, uint64 Cycles)
 {
-	uint16 NameSize = (FCString::Strlen(StatNameString) + 1) * sizeof(TCHAR);
-	UE_TRACE_LOG(CsvProfiler, DeclaredStat, NameSize)
-		<< DeclaredStat.StatId(StatName.GetComparisonIndex().ToUnstableInt())
-		<< DeclaredStat.Attachment(StatNameString, NameSize);
+	UE_TRACE_LOG(CsvProfiler, CustomStatFloat)
+		<< CustomStatFloat.Cycle(Cycles)
+		<< CustomStatFloat.StatId(FCsvProfilerTraceInternal::GetStatId(StatName, CategoryIndex))
+		<< CustomStatFloat.ThreadId(FPlatformTLS::GetCurrentThreadId())
+		<< CustomStatFloat.Value(Value)
+		<< CustomStatFloat.OpType(OpType);
 }
 
-void FCsvProfilerTrace::OutputCustomStat(const char* StatName, int32 Value, uint8 OpType)
+void FCsvProfilerTrace::OutputBeginCapture(const TCHAR* Filename, uint32 RenderThreadId, uint32 RHIThreadId, const char* DefaultWaitStatName, bool bEnableCounts)
 {
-	UE_TRACE_LOG(CsvProfiler, CustomStatInlineInt)
-		<< CustomStatInlineInt.Cycle(FPlatformTime::Cycles64())
-		<< CustomStatInlineInt.StatNamePointer(StatName)
-		<< CustomStatInlineInt.Value(Value)
-		<< CustomStatInlineInt.OpType(OpType);
+	OutputInlineStat(DefaultWaitStatName, CSV_CATEGORY_INDEX(Exclusive));
+	uint16 NameSize = (FCString::Strlen(Filename) + 1) * sizeof(TCHAR);
+	UE_TRACE_LOG(CsvProfiler, BeginCapture, NameSize)
+		<< BeginCapture.Cycle(FPlatformTime::Cycles64())
+		<< BeginCapture.RenderThreadId(RenderThreadId)
+		<< BeginCapture.RHIThreadId(RHIThreadId)
+		<< BeginCapture.EnableCounts(bEnableCounts)
+		<< BeginCapture.Attachment(Filename, NameSize);
 }
 
-void FCsvProfilerTrace::OutputCustomStat(const FName& StatName, int32 Value, uint8 OpType)
+void FCsvProfilerTrace::OutputEvent(const TCHAR* Text, int32 CategoryIndex, uint64 Cycles)
 {
-	UE_TRACE_LOG(CsvProfiler, CustomStatDeclaredInt)
-		<< CustomStatDeclaredInt.Cycle(FPlatformTime::Cycles64())
-		<< CustomStatDeclaredInt.StatId(StatName.GetComparisonIndex().ToUnstableInt())
-		<< CustomStatDeclaredInt.Value(Value)
-		<< CustomStatDeclaredInt.OpType(OpType);
+	uint16 TextSize = (FCString::Strlen(Text) + 1) * sizeof(TCHAR);
+	UE_TRACE_LOG(CsvProfiler, Event, TextSize)
+		<< Event.Cycle(Cycles)
+		<< Event.CategoryIndex(CategoryIndex)
+		<< Event.ThreadId(FPlatformTLS::GetCurrentThreadId())
+		<< Event.Attachment(Text, TextSize);
 }
 
-void FCsvProfilerTrace::OutputCustomStat(const char* StatName, float Value, uint8 OpType)
+void FCsvProfilerTrace::OutputEndCapture()
 {
-	UE_TRACE_LOG(CsvProfiler, CustomStatInlineFloat)
-		<< CustomStatInlineFloat.Cycle(FPlatformTime::Cycles64())
-		<< CustomStatInlineFloat.StatNamePointer(StatName)
-		<< CustomStatInlineFloat.Value(Value)
-		<< CustomStatInlineFloat.OpType(OpType);
+	UE_TRACE_LOG(CsvProfiler, EndCapture)
+		<< EndCapture.Cycle(FPlatformTime::Cycles64());
 }
 
-void FCsvProfilerTrace::OutputCustomStat(const FName& StatName, float Value, uint8 OpType)
+void FCsvProfilerTrace::OutputMetadata(const TCHAR* Key, const TCHAR* Value)
 {
-	UE_TRACE_LOG(CsvProfiler, CustomStatDeclaredFloat)
-		<< CustomStatDeclaredFloat.Cycle(FPlatformTime::Cycles64())
-		<< CustomStatDeclaredFloat.StatId(StatName.GetComparisonIndex().ToUnstableInt())
-		<< CustomStatDeclaredFloat.Value(Value)
-		<< CustomStatDeclaredFloat.OpType(OpType);
+	uint16 KeySize = (FCString::Strlen(Key) + 1) * sizeof(TCHAR);
+	uint16 ValueSize = (FCString::Strlen(Value) + 1) * sizeof(TCHAR);
+	auto Attachment = [Key, KeySize, Value, ValueSize](uint8* Out)
+	{
+		memcpy(Out, Key, KeySize);
+		memcpy(Out + KeySize, Value, ValueSize);
+	};
+	UE_TRACE_LOG(CsvProfiler, Metadata, KeySize + ValueSize)
+		<< Metadata.ValueOffset(KeySize)
+		<< Metadata.Attachment(Attachment);
 }
 
 #endif
\ No newline at end of file
diff --git a/Engine/Source/Runtime/Core/Public/ProfilingDebugging/CsvProfiler.h b/Engine/Source/Runtime/Core/Public/ProfilingDebugging/CsvProfiler.h
--- a/Engine/Source/Runtime/Core/Public/ProfilingDebugging/CsvProfiler.h
+++ b/Engine/Source/Runtime/Core/Public/ProfilingDebugging/CsvProfiler.h
@@ -47,25 +47,31 @@
 
 // Inline stats (no up front definition)
 #define CSV_SCOPED_TIMING_STAT(Category,StatName) \
-	FScopedCsvStat _ScopedCsvStat_ ## StatName (#StatName, CSV_CATEGORY_INDEX(Category)); \
-	TRACE_CPUPROFILER_EVENT_SCOPE_GROUP(StatName, CpuProfilerGroup_CsvProfiler);
+	TRACE_CSV_PROFILER_INLINE_STAT(#StatName, CSV_CATEGORY_INDEX(Category)); \
+	FScopedCsvStat _ScopedCsvStat_ ## StatName (#StatName, CSV_CATEGORY_INDEX(Category));
 #define CSV_SCOPED_TIMING_STAT_GLOBAL(StatName) \
-	FScopedCsvStat _ScopedCsvStat_ ## StatName (#StatName, CSV_CATEGORY_INDEX_GLOBAL); \
-	TRACE_CPUPROFILER_EVENT_SCOPE_GROUP(StatName, CpuProfilerGroup_CsvProfiler);
+	TRACE_CSV_PROFILER_INLINE_STAT(#StatName, CSV_CATEGORY_INDEX_GLOBAL); \
+	FScopedCsvStat _ScopedCsvStat_ ## StatName (#StatName, CSV_CATEGORY_INDEX_GLOBAL);
 #define CSV_SCOPED_TIMING_STAT_EXCLUSIVE(StatName) \
-	FScopedCsvStatExclusive _ScopedCsvStatExclusive_ ## StatName (#StatName); \
-	TRACE_CPUPROFILER_EVENT_SCOPE_GROUP(StatName, CpuProfilerGroup_CsvProfiler);
-#define CSV_SCOPED_TIMING_STAT_EXCLUSIVE_CONDITIONAL(StatName,Condition) FScopedCsvStatExclusiveConditional _ScopedCsvStatExclusive_ ## StatName (#StatName,Condition);
+	TRACE_CSV_PROFILER_INLINE_STAT_EXCLUSIVE(#StatName); \
+	FScopedCsvStatExclusive _ScopedCsvStatExclusive_ ## StatName (#StatName);
+#define CSV_SCOPED_TIMING_STAT_EXCLUSIVE_CONDITIONAL(StatName,Condition) \
+	TRACE_CSV_PROFILER_INLINE_STAT_EXCLUSIVE(#StatName); \
+	FScopedCsvStatExclusiveConditional _ScopedCsvStatExclusive_ ## StatName (#StatName,Condition);
 
 #define CSV_SCOPED_WAIT_CONDITIONAL(Condition)					FScopedCsvWaitConditional _ScopedCsvWait(Condition);
-#define CSV_SCOPED_SET_WAIT_STAT(StatName)						FScopedCsvSetWaitStat _ScopedCsvSetWaitStat ## StatName("EventWait/"#StatName);
+
+#define CSV_SCOPED_SET_WAIT_STAT(StatName) \
+	TRACE_CSV_PROFILER_INLINE_STAT_EXCLUSIVE("EventWait/"#StatName); \
+	FScopedCsvSetWaitStat _ScopedCsvSetWaitStat ## StatName("EventWait/"#StatName);
+
 #define CSV_SCOPED_SET_WAIT_STAT_IGNORE()						FScopedCsvSetWaitStat _ScopedCsvSetWaitStat ## StatName();
 
 #define CSV_CUSTOM_STAT(Category,StatName,Value,Op) \
-	TRACE_CSV_PROFILER_INLINE_STAT(#StatName); \
+	TRACE_CSV_PROFILER_INLINE_STAT(#StatName, CSV_CATEGORY_INDEX(Category)); \
 	FCsvProfiler::RecordCustomStat(#StatName, CSV_CATEGORY_INDEX(Category), Value, Op);
 #define CSV_CUSTOM_STAT_GLOBAL(StatName,Value,Op) \
-	TRACE_CSV_PROFILER_INLINE_STAT(#StatName); \
+	TRACE_CSV_PROFILER_INLINE_STAT(#StatName, CSV_CATEGORY_INDEX_GLOBAL); \
 	FCsvProfiler::RecordCustomStat(#StatName, CSV_CATEGORY_INDEX_GLOBAL, Value, Op); 
 
 // Stats declared up front
@@ -149,7 +155,7 @@ struct FCsvDeclaredStat
 		: Name(InNameString)
 		, CategoryIndex(InCategoryIndex) 
 	{
-		TRACE_CSVPROFILER_DECLARED_STAT(Name, InNameString);
+		TRACE_CSV_PROFILER_DECLARED_STAT(Name, InCategoryIndex);
 	}
 
 	FName Name;
diff --git a/Engine/Source/Runtime/Core/Public/ProfilingDebugging/CsvProfilerTrace.h b/Engine/Source/Runtime/Core/Public/ProfilingDebugging/CsvProfilerTrace.h
--- a/Engine/Source/Runtime/Core/Public/ProfilingDebugging/CsvProfilerTrace.h
+++ b/Engine/Source/Runtime/Core/Public/ProfilingDebugging/CsvProfilerTrace.h
@@ -15,35 +15,87 @@
 
 struct FCsvProfilerTrace
 {
-	CORE_API static void OutputInlineStat(const char* StatName);
-	CORE_API static void OutputDeclaredStat(const FName& StatName, const TCHAR* StatNameString);
-	CORE_API static void OutputCustomStat(const char* StatName, int32 Value, uint8 OpType);
-	CORE_API static void OutputCustomStat(const FName& StatName, int32 Value, uint8 OpType);
-	CORE_API static void OutputCustomStat(const char* StatName, float Value, uint8 OpType);
-	CORE_API static void OutputCustomStat(const FName& StatName, float Value, uint8 OpType);
+	static void OutputRegisterCategory(int32 Index, const TCHAR* Name);
+	CORE_API static void OutputInlineStat(const char* StatName, int32 CategoryIndex);
+	CORE_API static void OutputInlineStatExclusive(const char* StatName);
+	CORE_API static void OutputDeclaredStat(const FName& StatName, int32 CategoryIndex);
+	static void OutputBeginStat(const char* StatName, int32 CategoryIndex, uint64 Cycles);
+	static void OutputBeginStat(const FName& StatName, int32 CategoryIndex, uint64 Cycles);
+	static void OutputEndStat(const char* StatName, int32 CategoryIndex, uint64 Cycles);
+	static void OutputEndStat(const FName& StatName, int32 CategoryIndex, uint64 Cycles);
+	static void OutputBeginExclusiveStat(const char* StatName, int32 CategoryIndex, uint64 Cycles);
+	static void OutputEndExclusiveStat(const char* StatName, int32 CategoryIndex, uint64 Cycles);
+	static void OutputCustomStat(const char* StatName, int32 CategoryIndex, int32 Value, uint8 OpType, uint64 Cycles);
+	static void OutputCustomStat(const FName& StatName, int32 CategoryIndex, int32 Value, uint8 OpType, uint64 Cycles);
+	static void OutputCustomStat(const char* StatName, int32 CategoryIndex, float Value, uint8 OpType, uint64 Cycles);
+	static void OutputCustomStat(const FName& StatName, int32 CategoryIndex, float Value, uint8 OpType, uint64 Cycles);
+	static void OutputEvent(const TCHAR* Text, int32 CategoryIndex, uint64 Cycles);
+	static void OutputBeginCapture(const TCHAR* Filename, uint32 RenderThreadId, uint32 RHIThreadId, const char* DefaultWaitStatName, bool bEnableCounts);
+	static void OutputEndCapture();
+	static void OutputMetadata(const TCHAR* Key, const TCHAR* Value);
 };
 
-#define TRACE_CSV_PROFILER_INLINE_STAT(StatName) \
+#define TRACE_CSV_PROFILER_REGISTER_CATEGORY(Index, Name) \
+	FCsvProfilerTrace::OutputRegisterCategory(Index, Name);
+
+#define TRACE_CSV_PROFILER_INLINE_STAT(StatName, CategoryIndex) \
 	static bool PREPROCESSOR_JOIN(__CsvProfilerStat, __LINE__); \
 	if (!PREPROCESSOR_JOIN(__CsvProfilerStat, __LINE__)) { \
-		FCsvProfilerTrace::OutputInlineStat(StatName); \
+		FCsvProfilerTrace::OutputInlineStat(StatName, CategoryIndex); \
 		PREPROCESSOR_JOIN(__CsvProfilerStat, __LINE__) = true; \
 	}
 
-#define TRACE_CSVPROFILER_DECLARED_STAT(StatName, StatNameString) \
+#define TRACE_CSV_PROFILER_INLINE_STAT_EXCLUSIVE(StatName) \
 	static bool PREPROCESSOR_JOIN(__CsvProfilerStat, __LINE__); \
 	if (!PREPROCESSOR_JOIN(__CsvProfilerStat, __LINE__)) { \
-		FCsvProfilerTrace::OutputDeclaredStat(StatName, StatNameString); \
+		FCsvProfilerTrace::OutputInlineStatExclusive(StatName); \
 		PREPROCESSOR_JOIN(__CsvProfilerStat, __LINE__) = true; \
 	}
 
-#define TRACE_CSV_PROFILER_CUSTOM_STAT(StatName, Value, OpType) \
-	FCsvProfilerTrace::OutputCustomStat(StatName, Value, OpType);
+#define TRACE_CSV_PROFILER_DECLARED_STAT(StatName, CategoryIndex) \
+	FCsvProfilerTrace::OutputDeclaredStat(StatName, CategoryIndex);
+	
+#define TRACE_CSV_PROFILER_BEGIN_STAT(StatName, CategoryIndex, Cycles) \
+	FCsvProfilerTrace::OutputBeginStat(StatName, CategoryIndex, Cycles);
+
+#define TRACE_CSV_PROFILER_END_STAT(StatName, CategoryIndex, Cycles) \
+	FCsvProfilerTrace::OutputEndStat(StatName, CategoryIndex, Cycles);
+
+#define TRACE_CSV_PROFILER_BEGIN_EXCLUSIVE_STAT(StatName, CategoryIndex, Cycles) \
+	FCsvProfilerTrace::OutputBeginExclusiveStat(StatName, CategoryIndex, Cycles);
+
+#define TRACE_CSV_PROFILER_END_EXCLUSIVE_STAT(StatName, CategoryIndex, Cycles) \
+	FCsvProfilerTrace::OutputEndExclusiveStat(StatName, CategoryIndex, Cycles);
+
+#define TRACE_CSV_PROFILER_CUSTOM_STAT(StatName, CategoryIndex, Value, OpType, Cycles) \
+	FCsvProfilerTrace::OutputCustomStat(StatName, CategoryIndex, Value, OpType, Cycles);
+
+#define TRACE_CSV_PROFILER_EVENT(Text, CategoryIndex, Cycles) \
+	FCsvProfilerTrace::OutputEvent(Text, CategoryIndex, Cycles);
+
+#define TRACE_CSV_PROFILER_BEGIN_CAPTURE(Filename, RenderThreadId, RHIThreadId, DefaultWaitStatName, EnableCounts) \
+	FCsvProfilerTrace::OutputBeginCapture(Filename, RenderThreadId, RHIThreadId, DefaultWaitStatName, EnableCounts);
+
+#define TRACE_CSV_PROFILER_END_CAPTURE() \
+	FCsvProfilerTrace::OutputEndCapture();
+
+#define TRACE_CSV_PROFILER_METADATA(Key, Value) \
+	FCsvProfilerTrace::OutputMetadata(Key, Value);
 
 #else
 
+#define TRACE_CSV_PROFILER_REGISTER_CATEGORY(...)
 #define TRACE_CSV_PROFILER_INLINE_STAT(...)
-#define TRACE_CSVPROFILER_DECLARED_STAT(...)
+#define TRACE_CSV_PROFILER_INLINE_STAT_EXCLUSIVE(...)
+#define TRACE_CSV_PROFILER_DECLARED_STAT(...)
+#define TRACE_CSV_PROFILER_BEGIN_STAT(...)
+#define TRACE_CSV_PROFILER_END_STAT(...)
+#define TRACE_CSV_PROFILER_BEGIN_EXCLUSIVE_STAT(...)
+#define TRACE_CSV_PROFILER_END_EXCLUSIVE_STAT(...)
 #define TRACE_CSV_PROFILER_CUSTOM_STAT(...)
+#define TRACE_CSV_PROFILER_EVENT(...)
+#define TRACE_CSV_PROFILER_BEGIN_CAPTURE(...)
+#define TRACE_CSV_PROFILER_END_CAPTURE(...)
+#define TRACE_CSV_PROFILER_METADATA(...)
 
 #endif
\ No newline at end of file
