Added first version of UnrealInsightsCLI with basic report generation support

diff --git a/Engine/Source/Developer/TraceServices/Private/Common/Utils.cpp b/Engine/Source/Developer/TraceServices/Private/Common/Utils.cpp
new file mode 100644
--- /dev/null
+++ b/Engine/Source/Developer/TraceServices/Private/Common/Utils.cpp
@@ -0,0 +1,73 @@
+// Copyright 1998-2019 Epic Games, Inc. All Rights Reserved.
+
+#include "Utils.h"
+#include "TraceServices/Containers/Tables.h"
+#include "Templates/SharedPointer.h"
+#include "HAL/FileManager.h"
+
+namespace Trace
+{
+
+void Table2Csv(const Trace::IUntypedTable& Table, const TCHAR* Filename)
+{
+	TSharedPtr<FArchive> OutputFile = MakeShareable(IFileManager::Get().CreateFileWriter(Filename));
+	check(OutputFile);
+	FString Header;
+	const Trace::ITableLayout& Layout = Table.GetLayout();
+	int32 ColumnCount = Layout.GetColumnCount();
+	for (int32 ColumnIndex = 0; ColumnIndex < ColumnCount; ++ColumnIndex)
+	{
+		Header += Layout.GetColumnName(ColumnIndex);
+		if (ColumnIndex < ColumnCount - 1)
+		{
+			Header += TEXT(",");
+		}
+		else
+		{
+			Header += TEXT("\n");
+		}
+	}
+	auto AnsiHeader = StringCast<ANSICHAR>(*Header);
+	OutputFile->Serialize((void*)AnsiHeader.Get(), AnsiHeader.Length());
+	TUniquePtr<Trace::IUntypedTableReader> TableReader(Table.CreateReader());
+	for (; TableReader->IsValid(); TableReader->NextRow())
+	{
+		FString Line;
+		for (int32 ColumnIndex = 0; ColumnIndex < ColumnCount; ++ColumnIndex)
+		{
+			switch (Layout.GetColumnType(ColumnIndex))
+			{
+			case Trace::TableColumnType_Bool:
+				Line += TableReader->GetValueBool(ColumnIndex) ? "true" : "false";
+				break;
+			case Trace::TableColumnType_Int:
+				Line += FString::Printf(TEXT("%lld"), TableReader->GetValueInt(ColumnIndex));
+				break;
+			case Trace::TableColumnType_Float:
+				Line += FString::Printf(TEXT("%f"), TableReader->GetValueFloat(ColumnIndex));
+				break;
+			case Trace::TableColumnType_Double:
+				Line += FString::Printf(TEXT("%f"), TableReader->GetValueDouble(ColumnIndex));
+				break;
+			case Trace::TableColumnType_CString:
+				FString ValueString = TableReader->GetValueCString(ColumnIndex);
+				ValueString.ReplaceInline(TEXT(","), TEXT(" "));
+				Line += ValueString;
+				break;
+			}
+			if (ColumnIndex < ColumnCount - 1)
+			{
+				Line += TEXT(",");
+			}
+			else
+			{
+				Line += TEXT("\n");
+			}
+		}
+		auto AnsiLine = StringCast<ANSICHAR>(*Line);
+		OutputFile->Serialize((void*)AnsiLine.Get(), AnsiLine.Length());
+	}
+	OutputFile->Close();
+}
+
+}
diff --git a/Engine/Source/Developer/TraceServices/Private/ModuleService.cpp b/Engine/Source/Developer/TraceServices/Private/ModuleService.cpp
--- a/Engine/Source/Developer/TraceServices/Private/ModuleService.cpp
+++ b/Engine/Source/Developer/TraceServices/Private/ModuleService.cpp
@@ -98,4 +98,15 @@ bool FModuleService::GetModuleLoggers(const FName& ModuleName, TArray<const TCHA
 	return true;
 }
 
+void FModuleService::GenerateReports(const IAnalysisSession& Session, const TCHAR* CmdLine, const TCHAR* OutputDirectory)
+{
+	FScopeLock Lock(&CriticalSection);
+	Initialize();
+	for (const auto& KV : ModulesMap)
+	{
+		IModule* Module = KV.Value;
+		Module->GenerateReports(Session, CmdLine, OutputDirectory);
+	}
+}
+
 }
diff --git a/Engine/Source/Developer/TraceServices/Private/ModuleServicePrivate.h b/Engine/Source/Developer/TraceServices/Private/ModuleServicePrivate.h
--- a/Engine/Source/Developer/TraceServices/Private/ModuleServicePrivate.h
+++ b/Engine/Source/Developer/TraceServices/Private/ModuleServicePrivate.h
@@ -19,6 +19,7 @@ public:
 	FModuleService();
 	virtual void GetAvailableModules(TArray<FModuleInfo>& OutModules) override;
 	virtual void SetModuleEnabled(const FName& ModuleName, bool bEnabled) override;
+	virtual void GenerateReports(const IAnalysisSession& Session, const TCHAR* CmdLine, const TCHAR* OutputDirectory) override;
 	void OnAnalysisBegin(IAnalysisSession& Session);
 	bool GetModuleLoggers(const FName& ModuleName, TArray<const TCHAR*>& OutLoggers);
 
diff --git a/Engine/Source/Developer/TraceServices/Private/Modules/CsvProfilerModule.cpp b/Engine/Source/Developer/TraceServices/Private/Modules/CsvProfilerModule.cpp
--- a/Engine/Source/Developer/TraceServices/Private/Modules/CsvProfilerModule.cpp
+++ b/Engine/Source/Developer/TraceServices/Private/Modules/CsvProfilerModule.cpp
@@ -27,8 +27,16 @@ void FCsvProfilerModule::OnAnalysisBegin(IAnalysisSession& Session)
 	Session.AddAnalyzer(new FCsvProfilerAnalyzer(Session, *CsvProfilerProvider, FrameProvider, ThreadProvider));
 }
 
-void FCsvProfilerModule::GetLoggers(TArray<const TCHAR *>& OutLoggers)
+void FCsvProfilerModule::GenerateReports(const IAnalysisSession& Session, const TCHAR* CmdLine, const TCHAR* OutputDirectory)
 {
+	const ICsvProfilerProvider* CsvProfilerProvider = Trace::ReadCsvProfilerProvider(Session);
+	if (CsvProfilerProvider)
+	{
+		CsvProfilerProvider->EnumerateCaptures([CsvProfilerProvider, &OutputDirectory](const FCaptureInfo& CaptureInfo)
+		{
+			Table2Csv(CsvProfilerProvider->GetTable(CaptureInfo.Id), *(FString(OutputDirectory) / TEXT("CsvProfiler") / CaptureInfo.Filename));
+		});
+	}
 }
 
 const ICsvProfilerProvider* ReadCsvProfilerProvider(const IAnalysisSession& Session)
diff --git a/Engine/Source/Developer/TraceServices/Private/Modules/CsvProfilerModule.h b/Engine/Source/Developer/TraceServices/Private/Modules/CsvProfilerModule.h
--- a/Engine/Source/Developer/TraceServices/Private/Modules/CsvProfilerModule.h
+++ b/Engine/Source/Developer/TraceServices/Private/Modules/CsvProfilerModule.h
@@ -14,7 +14,8 @@ class FCsvProfilerModule
 public:
 	virtual void GetModuleInfo(FModuleInfo& OutModuleInfo) override;
 	virtual void OnAnalysisBegin(IAnalysisSession& Session) override;
-	virtual void GetLoggers(TArray<const TCHAR*>& OutLoggers) override;
+	virtual void GetLoggers(TArray<const TCHAR*>& OutLoggers) override {};
+	virtual void GenerateReports(const IAnalysisSession& Session, const TCHAR* CmdLine, const TCHAR* OutputDirectory) override;
 };
 
 }
diff --git a/Engine/Source/Developer/TraceServices/Private/Modules/LoadTimeProfilerModule.cpp b/Engine/Source/Developer/TraceServices/Private/Modules/LoadTimeProfilerModule.cpp
--- a/Engine/Source/Developer/TraceServices/Private/Modules/LoadTimeProfilerModule.cpp
+++ b/Engine/Source/Developer/TraceServices/Private/Modules/LoadTimeProfilerModule.cpp
@@ -5,6 +5,7 @@
 #include "Analyzers/LoadTimeTraceAnalysis.h"
 #include "AnalysisServicePrivate.h"
 #include "Model/FileActivity.h"
+#include "HAL/FileManager.h"
 
 namespace Trace
 {
@@ -35,6 +36,57 @@ void FLoadTimeProfilerModule::GetLoggers(TArray<const TCHAR *>& OutLoggers)
 	OutLoggers.Add(TEXT("PlatformFile"));
 }
 
+void FLoadTimeProfilerModule::GenerateReports(const IAnalysisSession& Session, const TCHAR* CmdLine, const TCHAR* OutputDirectory)
+{
+	double CaptureStartTime = -DBL_MAX;
+	double CaptureEndTime = DBL_MAX;
+	const IBookmarkProvider& BookmarkProvider = Trace::ReadBookmarkProvider(Session);
+	FString BeginCaptureBookmarkName;
+	FParse::Value(CmdLine, TEXT("-BeginCaptureBookmark="), BeginCaptureBookmarkName);
+	FString EndCaptureBookmarkName;
+	FParse::Value(CmdLine, TEXT("-EndCaptureBookmark="), EndCaptureBookmarkName);
+	BookmarkProvider.EnumerateBookmarks(0.0, DBL_MAX, [BeginCaptureBookmarkName, EndCaptureBookmarkName, &CaptureStartTime, &CaptureEndTime](const FBookmark& Bookmark)
+	{
+		if (CaptureStartTime == -DBL_MAX && !BeginCaptureBookmarkName.IsEmpty() && BeginCaptureBookmarkName == Bookmark.Text)
+		{
+			CaptureStartTime = Bookmark.Time;
+		}
+		if (CaptureEndTime == DBL_MAX && !EndCaptureBookmarkName.IsEmpty() && EndCaptureBookmarkName == Bookmark.Text)
+		{
+			CaptureEndTime = Bookmark.Time;
+		}
+	});
+	if (CaptureStartTime == -DBL_MAX)
+	{
+		CaptureStartTime = 0.0;
+	}
+	if (CaptureEndTime == DBL_MAX)
+	{
+		CaptureEndTime = Session.GetDurationSeconds();
+	}
+	const ILoadTimeProfilerProvider* LoadTimeProfiler = Trace::ReadLoadTimeProfilerProvider(Session);
+	FString ReportDirectory = FString(OutputDirectory) / TEXT("LoadTimeProfiler");
+	if (LoadTimeProfiler)
+	{
+		TUniquePtr<ITable<FPackagesTableRow>> PackagesTable(LoadTimeProfiler->CreatePackageDetailsTable(CaptureStartTime, CaptureEndTime));
+		Table2Csv(*PackagesTable.Get(), *(ReportDirectory / TEXT("Packages.csv")));
+		TUniquePtr<ITable<FExportsTableRow>> ExportsTable(LoadTimeProfiler->CreateExportDetailsTable(CaptureStartTime, CaptureEndTime));
+		Table2Csv(*ExportsTable.Get(), *(ReportDirectory / TEXT("Exports.csv")));
+		Table2Csv(LoadTimeProfiler->GetRequestsTable(), *(ReportDirectory / TEXT("Requests.csv")));
+	}
+	const IFileActivityProvider* FileActivityProvider = Trace::ReadFileActivityProvider(Session);
+	if (FileActivityProvider)
+	{
+		Table2Csv(FileActivityProvider->GetFileActivityTable(), *(FString(ReportDirectory) / TEXT("FileActivity.csv")));
+	}
+	if (CaptureStartTime > 0.0 || CaptureEndTime < Session.GetDurationSeconds())
+	{
+		TSharedPtr<FArchive> CaptureSummaryOutputFile = MakeShareable(IFileManager::Get().CreateFileWriter(*(ReportDirectory / TEXT("CaptureSummary.txt"))));
+		FString SummaryLine = FString::Printf(TEXT("Capture start: %f\r\nCapture end: %f\r\nCapture duration: %f"), CaptureStartTime, CaptureEndTime, CaptureEndTime - CaptureStartTime);
+		CaptureSummaryOutputFile->Serialize(TCHAR_TO_ANSI(*SummaryLine), SummaryLine.Len());
+	}
+}
+
 const ILoadTimeProfilerProvider* ReadLoadTimeProfilerProvider(const IAnalysisSession& Session)
 {
 	return Session.ReadProvider<ILoadTimeProfilerProvider>(LoadTimeProfilerProviderName);
diff --git a/Engine/Source/Developer/TraceServices/Private/Modules/LoadTimeProfilerModule.h b/Engine/Source/Developer/TraceServices/Private/Modules/LoadTimeProfilerModule.h
--- a/Engine/Source/Developer/TraceServices/Private/Modules/LoadTimeProfilerModule.h
+++ b/Engine/Source/Developer/TraceServices/Private/Modules/LoadTimeProfilerModule.h
@@ -14,6 +14,7 @@ public:
 	virtual void GetModuleInfo(FModuleInfo& OutModuleInfo) override;
 	virtual void OnAnalysisBegin(IAnalysisSession& Session) override;
 	virtual void GetLoggers(TArray<const TCHAR *>& OutLoggers) override;
+	virtual void GenerateReports(const IAnalysisSession& Session, const TCHAR* CmdLine, const TCHAR* OutputDirectory) override;
 };
 
 }
diff --git a/Engine/Source/Developer/TraceServices/Private/Modules/StatsModule.h b/Engine/Source/Developer/TraceServices/Private/Modules/StatsModule.h
--- a/Engine/Source/Developer/TraceServices/Private/Modules/StatsModule.h
+++ b/Engine/Source/Developer/TraceServices/Private/Modules/StatsModule.h
@@ -15,6 +15,7 @@ public:
 	virtual void GetModuleInfo(FModuleInfo& OutModuleInfo) override;
 	virtual void OnAnalysisBegin(IAnalysisSession& Session) override;
 	virtual void GetLoggers(TArray<const TCHAR *>& OutLoggers) override;
+	virtual void GenerateReports(const IAnalysisSession& Session, const TCHAR* CmdLine, const TCHAR* OutputDirectory) override {}
 };
 
 }
diff --git a/Engine/Source/Developer/TraceServices/Private/Modules/TimingProfilerModule.h b/Engine/Source/Developer/TraceServices/Private/Modules/TimingProfilerModule.h
--- a/Engine/Source/Developer/TraceServices/Private/Modules/TimingProfilerModule.h
+++ b/Engine/Source/Developer/TraceServices/Private/Modules/TimingProfilerModule.h
@@ -15,6 +15,7 @@ public:
 	virtual void GetModuleInfo(FModuleInfo& OutModuleInfo) override;
 	virtual void OnAnalysisBegin(IAnalysisSession& Session) override;
 	virtual void GetLoggers(TArray<const TCHAR *>& OutLoggers) override;
+	virtual void GenerateReports(const IAnalysisSession& Session, const TCHAR* CmdLine, const TCHAR* OutputDirectory) override {}
 };
 
 }
diff --git a/Engine/Source/Developer/TraceServices/Public/TraceServices/Containers/Tables.h b/Engine/Source/Developer/TraceServices/Public/TraceServices/Containers/Tables.h
--- a/Engine/Source/Developer/TraceServices/Public/TraceServices/Containers/Tables.h
+++ b/Engine/Source/Developer/TraceServices/Public/TraceServices/Containers/Tables.h
@@ -64,4 +64,6 @@ public:
 	virtual ITableReader<RowType>* CreateReader() const = 0;
 };
 
+void Table2Csv(const IUntypedTable& Table, const TCHAR* Filename);
+
 }
\ No newline at end of file
diff --git a/Engine/Source/Developer/TraceServices/Public/TraceServices/ModuleService.h b/Engine/Source/Developer/TraceServices/Public/TraceServices/ModuleService.h
--- a/Engine/Source/Developer/TraceServices/Public/TraceServices/ModuleService.h
+++ b/Engine/Source/Developer/TraceServices/Public/TraceServices/ModuleService.h
@@ -27,6 +27,7 @@ public:
 	virtual void GetModuleInfo(FModuleInfo& OutModuleInfo) = 0;
 	virtual void OnAnalysisBegin(IAnalysisSession& Session) = 0;
 	virtual void GetLoggers(TArray<const TCHAR*>& OutLoggers) = 0;
+	virtual void GenerateReports(const IAnalysisSession& Session, const TCHAR* CmdLine, const TCHAR* OutputDirectory) = 0;
 };
 
 class IModuleService
@@ -35,6 +36,7 @@ public:
 	virtual ~IModuleService() = default;
 	virtual void GetAvailableModules(TArray<FModuleInfo>& OutModules) = 0;
 	virtual void SetModuleEnabled(const FName& ModuleName, bool bEnabled) = 0;
+	virtual void GenerateReports(const IAnalysisSession& Session, const TCHAR* CmdLine, const TCHAR* OutputDirectory) = 0;
 };
 
 }
diff --git a/Engine/Source/Programs/UnrealInsightsCLI/Private/UnrealInsightsCLI.cpp b/Engine/Source/Programs/UnrealInsightsCLI/Private/UnrealInsightsCLI.cpp
new file mode 100644
--- /dev/null
+++ b/Engine/Source/Programs/UnrealInsightsCLI/Private/UnrealInsightsCLI.cpp
@@ -0,0 +1,57 @@
+// Copyright 1998-2019 Epic Games, Inc. All Rights Reserved.
+
+
+#include "UnrealInsightsCLI.h"
+
+#include "RequiredProgramMainCPPInclude.h"
+#include "TraceServices/ITraceServicesModule.h"
+#include "TraceServices/AnalysisService.h"
+#include "TraceServices/SessionService.h"
+#include "TraceServices/ModuleService.h"
+
+DEFINE_LOG_CATEGORY_STATIC(LogUnrealInsightsCLI, Log, All);
+
+IMPLEMENT_APPLICATION(UnrealInsightsCLI, "UnrealInsightsCLI");
+
+INT32_MAIN_INT32_ARGC_TCHAR_ARGV()
+{
+	using namespace Trace;
+
+	GEngineLoop.PreInit(ArgC, ArgV);
+	ITraceServicesModule& TraceServicesModule = FModuleManager::LoadModuleChecked<ITraceServicesModule>("TraceServices");
+	TSharedPtr<IAnalysisService> AnalysisService = TraceServicesModule.GetAnalysisService();
+	TSharedPtr<Trace::ISessionService> SessionService = TraceServicesModule.GetSessionService();
+	TSharedPtr<IModuleService> ModuleService = TraceServicesModule.GetModuleService();
+
+	TArray<FString> Tokens, Switches;
+	FCommandLine::Parse(FCommandLine::Get(), Tokens, Switches);
+	if (Tokens.Num() < 1)
+	{
+		UE_LOG(LogUnrealInsightsCLI, Error, TEXT("Usage: UnrealInsightsCLI.exe TraceFile [Arguments...]"));
+		return 1;
+	}
+	const TCHAR* FileName = *Tokens[0];
+	TUniquePtr<IInDataStream> DataStream(SessionService->OpenSessionFromFile(FileName));
+	if (!DataStream)
+	{
+		UE_LOG(LogUnrealInsightsCLI, Error, TEXT("Trace file %s not found"), FileName);
+		return 1;
+	}
+	UE_LOG(LogUnrealInsightsCLI, Display, TEXT("Analyzing %s..."), FileName);
+	double StartTime = FPlatformTime::Seconds();
+	TSharedPtr<const IAnalysisSession> Session = AnalysisService->Analyze(TEXT("Session"), MoveTemp(DataStream));
+	double Duration = FPlatformTime::Seconds() - StartTime;
+	UE_LOG(LogUnrealInsightsCLI, Display, TEXT("Analysis done! Took %f seconds"), Duration);
+
+	FString OutputDirectory = FPaths::ProjectSavedDir() / TEXT("Reports") / FDateTime::Now().ToString(TEXT("%Y%m%d_%H%M%S"));
+	IPlatformFile& FileSystem = IPlatformFile::GetPlatformPhysical();
+	FileSystem.CreateDirectory(*OutputDirectory);
+	UE_LOG(LogUnrealInsightsCLI, Display, TEXT("Saving reports to %s"), *OutputDirectory);
+
+	FAnalysisSessionReadScope _(*Session);
+	ModuleService->GenerateReports(*Session, FCommandLine::Get(), *OutputDirectory);
+
+	GIsRequestingExit = true;
+	return 0;
+}
+
diff --git a/Engine/Source/Programs/UnrealInsightsCLI/Private/UnrealInsightsCLI.h b/Engine/Source/Programs/UnrealInsightsCLI/Private/UnrealInsightsCLI.h
new file mode 100644
--- /dev/null
+++ b/Engine/Source/Programs/UnrealInsightsCLI/Private/UnrealInsightsCLI.h
@@ -0,0 +1,6 @@
+// Copyright 1998-2019 Epic Games, Inc. All Rights Reserved.
+
+#pragma once
+
+#include "CoreMinimal.h"
+
diff --git a/Engine/Source/Programs/UnrealInsightsCLI/UnrealInsightsCLI.Build.cs b/Engine/Source/Programs/UnrealInsightsCLI/UnrealInsightsCLI.Build.cs
new file mode 100644
--- /dev/null
+++ b/Engine/Source/Programs/UnrealInsightsCLI/UnrealInsightsCLI.Build.cs
@@ -0,0 +1,16 @@
+// Copyright 1998-2019 Epic Games, Inc. All Rights Reserved.
+
+using UnrealBuildTool;
+
+public class UnrealInsightsCLI : ModuleRules
+{
+	public UnrealInsightsCLI(ReadOnlyTargetRules Target) : base(Target)
+	{
+		PublicIncludePaths.Add("Runtime/Launch/Public");
+		PrivateIncludePaths.Add("Runtime/Launch/Private");
+		PrivateDependencyModuleNames.Add("Core");
+		PrivateDependencyModuleNames.Add("Projects");
+		PrivateDependencyModuleNames.Add("TraceLog");
+		PrivateDependencyModuleNames.Add("TraceServices");
+	}
+}
diff --git a/Engine/Source/Programs/UnrealInsightsCLI/UnrealInsightsCLI.Target.cs b/Engine/Source/Programs/UnrealInsightsCLI/UnrealInsightsCLI.Target.cs
new file mode 100644
--- /dev/null
+++ b/Engine/Source/Programs/UnrealInsightsCLI/UnrealInsightsCLI.Target.cs
@@ -0,0 +1,22 @@
+// Copyright 1998-2019 Epic Games, Inc. All Rights Reserved.
+
+using UnrealBuildTool;
+using System.Collections.Generic;
+
+[SupportedPlatforms("Win64")]
+public class UnrealInsightsCLITarget : TargetRules
+{
+	public UnrealInsightsCLITarget(TargetInfo Target) : base(Target)
+	{
+		Type = TargetType.Program;
+		LinkType = TargetLinkType.Monolithic;
+		LaunchModuleName = "UnrealInsightsCLI";
+
+		bBuildDeveloperTools = false;
+		bBuildWithEditorOnlyData = true;
+		bCompileAgainstEngine = false;
+		bCompileAgainstCoreUObject = false;
+		bCompileAgainstApplicationCore = false;
+		bIsBuildingConsoleApplication = true;
+	}
+}
