UnrealInsights:
Added counters for CSV profiler stat series
Added previous and next value to counter enumerations (when available)

diff --git a/Engine/Source/Developer/TraceServices/Private/Analyzers/CsvProfilerTraceAnalysis.cpp b/Engine/Source/Developer/TraceServices/Private/Analyzers/CsvProfilerTraceAnalysis.cpp
--- a/Engine/Source/Developer/TraceServices/Private/Analyzers/CsvProfilerTraceAnalysis.cpp
+++ b/Engine/Source/Developer/TraceServices/Private/Analyzers/CsvProfilerTraceAnalysis.cpp
@@ -4,9 +4,10 @@
 #include "TraceServices/Model/Counters.h"
 #include "TraceServices/Model/Frames.h"
 
-FCsvProfilerAnalyzer::FCsvProfilerAnalyzer(Trace::IAnalysisSession& InSession, Trace::FCsvProfilerProvider& InCsvProfilerProvider, const Trace::IFrameProvider& InFrameProvider, const Trace::IThreadProvider& InThreadProvider)
+FCsvProfilerAnalyzer::FCsvProfilerAnalyzer(Trace::IAnalysisSession& InSession, Trace::FCsvProfilerProvider& InCsvProfilerProvider, Trace::ICounterProvider& InCounterProvider, const Trace::IFrameProvider& InFrameProvider, const Trace::IThreadProvider& InThreadProvider)
 	: Session(InSession)
 	, CsvProfilerProvider(InCsvProfilerProvider)
+	, CounterProvider(InCounterProvider)
 	, FrameProvider(InFrameProvider)
 	, ThreadProvider(InThreadProvider)
 {
@@ -132,7 +133,7 @@ void FCsvProfilerAnalyzer::OnEvent(uint16 RouteId, const FOnEventContext& Contex
 	{
 		RenderThreadId = EventData.GetValue<uint32>("RenderThreadId");
 		RHIThreadId = EventData.GetValue<uint32>("RHIThreadId");
-		int32 CaptureStartFrame = GetFrameNumberForTimestamp(TraceFrameType_Game, Context.SessionContext.TimestampFromCycle(EventData.GetValue<uint64>("Cycle")));
+		uint32 CaptureStartFrame = GetFrameNumberForTimestamp(TraceFrameType_Game, Context.SessionContext.TimestampFromCycle(EventData.GetValue<uint64>("Cycle")));
 		bEnableCounts = EventData.GetValue<bool>("EnableCounts");
 		const TCHAR* Filename = Session.StoreString(reinterpret_cast<const TCHAR*>(EventData.GetAttachment()));
 		CsvProfilerProvider.StartCapture(Filename, CaptureStartFrame);
@@ -140,7 +141,12 @@ void FCsvProfilerAnalyzer::OnEvent(uint16 RouteId, const FOnEventContext& Contex
 	}
 	case RouteId_EndCapture:
 	{
-		int32 CaptureEndFrame = GetFrameNumberForTimestamp(TraceFrameType_Game, Context.SessionContext.TimestampFromCycle(EventData.GetValue<uint64>("Cycle")));
+		uint32 CaptureEndFrame = GetFrameNumberForTimestamp(TraceFrameType_Game, Context.SessionContext.TimestampFromCycle(EventData.GetValue<uint64>("Cycle")));
+		for (FStatSeriesInstance* StatSeries : StatSeriesInstanceArray)
+		{
+			FlushAtEndOfCapture(*StatSeries, CaptureEndFrame);
+		}
+		
 		CsvProfilerProvider.EndCapture(CaptureEndFrame);
 	}
 	}
@@ -200,7 +206,7 @@ void FCsvProfilerAnalyzer::DefineStatSeries(uint64 StatId, const TCHAR* Name, in
 const TCHAR* FCsvProfilerAnalyzer::GetStatSeriesName(const FStatSeriesDefinition* Definition, Trace::ECsvStatSeriesType Type, FThreadState& ThreadState, bool bIsCount)
 {
 	FString Name = Definition->Name;
-	if (Type == Trace::CsvStatSeriesType_Timer || Type == Trace::CsvStatSeriesType_ExclusiveTimer || bIsCount)
+	if (Type == Trace::CsvStatSeriesType_Timer || bIsCount)
 	{
 		// Add a /<Threadname> prefix
 		Name = ThreadState.ThreadName + TEXT("/") + Name;
@@ -211,10 +217,6 @@ const TCHAR* FCsvProfilerAnalyzer::GetStatSeriesName(const FStatSeriesDefinition
 		// Categorized stats are prefixed with <CATEGORY>/
 		Name = FString(CategoryMap[Definition->CategoryIndex]) + TEXT("/") + Name;
 	}
-	else if (Type == Trace::CsvStatSeriesType_ExclusiveTimer)
-	{
-		Name = FString(TEXT("Exclusive/")) + Name;
-	}
 
 	if (bIsCount)
 	{
@@ -225,13 +227,13 @@ const TCHAR* FCsvProfilerAnalyzer::GetStatSeriesName(const FStatSeriesDefinition
 	return Session.StoreString(*Name);
 }
 
-FCsvProfilerAnalyzer::FStatSeriesInstance* FCsvProfilerAnalyzer::GetStatSeries(uint64 StatId, Trace::ECsvStatSeriesType Type, FThreadState& ThreadState)
+FCsvProfilerAnalyzer::FStatSeriesInstance& FCsvProfilerAnalyzer::GetStatSeries(uint64 StatId, Trace::ECsvStatSeriesType Type, FThreadState& ThreadState)
 {
 	FStatSeriesDefinition* Definition;
 	FStatSeriesDefinition** FindIt = StatSeriesMap.Find(StatId);
 	if (!FindIt)
 	{
-		Definition = CreateStatSeries(*FString::Printf(TEXT("[undefined%d]"), UndefinedStatSeriesCount++), 0);
+		Definition = CreateStatSeries(*FString::Printf(TEXT("[unknown%d]"), UndefinedStatSeriesCount++), 0);
 		StatSeriesMap.Add(StatId, Definition);
 	}
 	else
@@ -246,16 +248,22 @@ FCsvProfilerAnalyzer::FStatSeriesInstance* FCsvProfilerAnalyzer::GetStatSeries(u
 	FStatSeriesInstance* Instance = ThreadState.StatSeries[Definition->ColumnIndex];
 	if (Instance)
 	{
-		return Instance;
+		return *Instance;
 	}
 
 	Instance = new FStatSeriesInstance();
 	StatSeriesInstanceArray.Add(Instance);
 	ThreadState.StatSeries[Definition->ColumnIndex] = Instance;
-	Instance->ProviderHandle = CsvProfilerProvider.AddSeries(GetStatSeriesName(Definition, Type, ThreadState, false), Type);
+	const TCHAR* StatSeriesName = GetStatSeriesName(Definition, Type, ThreadState, false);
+	Instance->ProviderHandle = CsvProfilerProvider.AddSeries(StatSeriesName, Type);
 	Instance->ProviderCountHandle = CsvProfilerProvider.AddSeries(GetStatSeriesName(Definition, Type, ThreadState, true), Trace::CsvStatSeriesType_CustomStatInt);
+	Instance->Counter = CounterProvider.CreateCounter();
+	Instance->Counter->SetName(StatSeriesName);
+	Instance->Counter->SetIsFloatingPoint(Type != Trace::CsvStatSeriesType_CustomStatInt);
+	Instance->Type = Type;
+	Instance->FrameType = ThreadState.FrameType;
 
-	return Instance;
+	return *Instance;
 }
 
 void FCsvProfilerAnalyzer::HandleMarkerEvent(const FOnEventContext& Context, bool bIsExclusive, bool bIsBegin)
@@ -263,7 +271,7 @@ void FCsvProfilerAnalyzer::HandleMarkerEvent(const FOnEventContext& Context, boo
 	FThreadState& ThreadState = GetThreadState(Context.EventData.GetValue<uint32>("ThreadId"));
 	uint64 StatId = Context.EventData.GetValue<uint64>("StatId");
 	FTimingMarker Marker;
-	Marker.StatSeries = GetStatSeries(StatId, Trace::CsvStatSeriesType_Timer, ThreadState);
+	Marker.StatId = StatId;
 	Marker.bIsBegin = bIsBegin;
 	Marker.bIsExclusive = bIsExclusive;
 	Marker.Cycle = Context.EventData.GetValue<uint64>("Cycle");
@@ -314,7 +322,8 @@ void FCsvProfilerAnalyzer::HandleMarker(const FOnEventContext& Context, FThreadS
 		HandleMarker(Context, ThreadState, InsertedMarker);
 	}
 	
-	int32 FrameNumber = GetFrameNumberForTimestamp(ThreadState.FrameType, Context.SessionContext.TimestampFromCycle(Marker.Cycle));
+	double Timestamp = Context.SessionContext.TimestampFromCycle(Marker.Cycle);
+	uint32 FrameNumber = GetFrameNumberForTimestamp(ThreadState.FrameType, Timestamp);
 	if (Marker.bIsBegin)
 	{
 		ThreadState.MarkerStack.Push(Marker);
@@ -330,7 +339,7 @@ void FCsvProfilerAnalyzer::HandleMarker(const FOnEventContext& Context, FThreadS
 
 			for (int j = ThreadState.MarkerStack.Num() - 1; j >= 0; j--)
 			{
-				if (ThreadState.MarkerStack[j].StatSeries == Marker.StatSeries) // Note: only works with scopes!
+				if (ThreadState.MarkerStack[j].StatId == Marker.StatId) // Note: only works with scopes!
 				{
 					StartMarker = ThreadState.MarkerStack[j];
 					ThreadState.MarkerStack.RemoveAt(j, 1, false);
@@ -342,20 +351,15 @@ void FCsvProfilerAnalyzer::HandleMarker(const FOnEventContext& Context, FThreadS
 			// TODO: if bFoundStart is false, this stat _never_ gets processed. Could we add it to a persistent list so it's considered next time?
 			// Example where this could go wrong: staggered/overlapping exclusive stats ( e.g Abegin, Bbegin, AEnd, BEnd ), where processing ends after AEnd
 			// AEnd would be missing 
-			if (FrameNumber >= 0 && bFoundStart)
+			if (bFoundStart)
 			{
-				check(Marker.StatSeries == StartMarker.StatSeries);
+				check(Marker.StatId == StartMarker.StatId);
 				check(Marker.Cycle >= StartMarker.Cycle);
 				if (Marker.Cycle > StartMarker.Cycle)
 				{
 					uint64 ElapsedCycles = Marker.Cycle - StartMarker.Cycle;
-					CsvProfilerProvider.SetTimerValue(Marker.StatSeries->ProviderHandle, FrameNumber, Context.SessionContext.DurationFromCycleCount(ElapsedCycles) * 1000.0);
-
-					// Add the COUNT/ series if enabled. Ignore artificial markers (inserted above)
-					if (bEnableCounts && !Marker.bIsExclusiveInsertedMarker)
-					{
-						CsvProfilerProvider.SetCustomStatValue(Marker.StatSeries->ProviderCountHandle, FrameNumber, Trace::CsvOpType_Accumulate, 1);
-					}
+					FStatSeriesInstance& StatSeries = GetStatSeries(Marker.StatId, Trace::CsvStatSeriesType_Timer, ThreadState);
+					SetTimerValue(StatSeries, FrameNumber, Context.SessionContext.DurationFromCycleCount(ElapsedCycles) * 1000.0, !Marker.bIsExclusiveInsertedMarker);
 				}
 			}
 		}
@@ -365,26 +369,18 @@ void FCsvProfilerAnalyzer::HandleMarker(const FOnEventContext& Context, FThreadS
 void FCsvProfilerAnalyzer::HandleCustomStatEvent(const FOnEventContext& Context, bool bIsFloat)
 {
 	FThreadState& ThreadState = GetThreadState(Context.EventData.GetValue<uint32>("ThreadId"));
-	FStatSeriesInstance* StatSeries = GetStatSeries(Context.EventData.GetValue<uint64>("StatId"), Trace::CsvStatSeriesType_CustomStatFloat, ThreadState);
-	Trace::ECsvOpType OpType = static_cast<Trace::ECsvOpType>(Context.EventData.GetValue<uint8>("OpType"));
-	uint64 Cycle = Context.EventData.GetValue<uint64>("Cycle");
-	int32 FrameNumber = GetFrameNumberForTimestamp(ThreadState.FrameType, Context.SessionContext.TimestampFromCycle(Cycle));
-	if (FrameNumber >= 0)
+	FStatSeriesInstance& StatSeries = GetStatSeries(Context.EventData.GetValue<uint64>("StatId"), bIsFloat ? Trace::CsvStatSeriesType_CustomStatFloat : Trace::CsvStatSeriesType_CustomStatInt, ThreadState);
+	ECsvOpType OpType = static_cast<ECsvOpType>(Context.EventData.GetValue<uint8>("OpType"));
+	uint32 FrameNumber = GetFrameNumberForTimestamp(ThreadState.FrameType, Context.SessionContext.TimestampFromCycle(Context.EventData.GetValue<uint64>("Cycle")));
+	if (bIsFloat)
 	{
-		if (bIsFloat)
-		{
-			CsvProfilerProvider.SetCustomStatValue(StatSeries->ProviderHandle, FrameNumber, OpType, Context.EventData.GetValue<float>("Value"));
-		}
-		else
-		{
-			CsvProfilerProvider.SetCustomStatValue(StatSeries->ProviderHandle, FrameNumber, OpType, Context.EventData.GetValue<int32>("Value"));
-		}
-
-		// Add the COUNT/ series if enabled
-		if (bEnableCounts)
-		{
-			CsvProfilerProvider.SetCustomStatValue(StatSeries->ProviderCountHandle, FrameNumber, Trace::CsvOpType_Accumulate, 1);
-		}
+		float Value = Context.EventData.GetValue<float>("Value");
+		SetCustomStatValue(StatSeries, FrameNumber, OpType, Value);
+	}
+	else
+	{
+		int32 Value = Context.EventData.GetValue<int32>("Value");
+		SetCustomStatValue(StatSeries, FrameNumber, OpType, Value);
 	}
 }
 
@@ -392,28 +388,154 @@ void FCsvProfilerAnalyzer::HandleEventEvent(const FOnEventContext& Context)
 {
 	FThreadState& ThreadState = GetThreadState(Context.EventData.GetValue<uint32>("ThreadId"));
 	uint64 Cycle = Context.EventData.GetValue<uint64>("Cycle");
-	int32 FrameNumber = GetFrameNumberForTimestamp(ThreadState.FrameType, Context.SessionContext.TimestampFromCycle(Cycle));
-	if (FrameNumber >= 0)
+	uint32 FrameNumber = GetFrameNumberForTimestamp(ThreadState.FrameType, Context.SessionContext.TimestampFromCycle(Cycle));
+	FString EventText = reinterpret_cast<const TCHAR*>(Context.EventData.GetAttachment());
+	int32 CategoryIndex = Context.EventData.GetValue<int32>("CategoryIndex");
+	if (CategoryIndex > 0)
 	{
-		FString EventText = reinterpret_cast<const TCHAR*>(Context.EventData.GetAttachment());
-		int32 CategoryIndex = Context.EventData.GetValue<int32>("CategoryIndex");
-		if (CategoryIndex > 0)
-		{
-			EventText = FString(CategoryMap[CategoryIndex]) + TEXT("/") + EventText;
-		}
-		CsvProfilerProvider.AddEvent(FrameNumber, Session.StoreString(*EventText));
+		EventText = FString(CategoryMap[CategoryIndex]) + TEXT("/") + EventText;
 	}
+	CsvProfilerProvider.AddEvent(FrameNumber, Session.StoreString(*EventText));
 }
 
-int32 FCsvProfilerAnalyzer::GetFrameNumberForTimestamp(ETraceFrameType FrameType, double Timestamp) const
+uint32 FCsvProfilerAnalyzer::GetFrameNumberForTimestamp(ETraceFrameType FrameType, double Timestamp) const
 {
 	const TArray<double>& FrameStartTimes = FrameProvider.GetFrameStartTimes(FrameType);
 
 	if (FrameStartTimes.Num() == 0 || Timestamp < FrameStartTimes[0])
 	{
 		return 0;
+	} 
+	else if (Timestamp >= FrameStartTimes.Last())
+	{
+		return FrameStartTimes.Num();
+	}
+	else
+	{
+		uint32 Index = static_cast<uint32>(Algo::LowerBound(FrameStartTimes, Timestamp));
+		return Index + 1;
+	}
+}
+
+void FCsvProfilerAnalyzer::Flush(FStatSeriesInstance& StatSeries)
+{
+	double CounterTimestamp;
+	if (StatSeries.CurrentFrame == 0)
+	{
+		const Trace::FFrame* Frame = FrameProvider.GetFrame(StatSeries.FrameType, 0);
+		check(Frame);
+		CounterTimestamp = Frame->StartTime;
+	}
+	else
+	{
+		const Trace::FFrame* Frame = FrameProvider.GetFrame(StatSeries.FrameType, StatSeries.CurrentFrame - 1);
+		const Trace::FFrame* NextFrame = FrameProvider.GetFrame(StatSeries.FrameType, StatSeries.CurrentFrame);
+		check(NextFrame);
+		CounterTimestamp = Frame->EndTime;
+	}
+	if (StatSeries.Type == Trace::CsvStatSeriesType_CustomStatInt)
+	{
+		CsvProfilerProvider.SetValue(StatSeries.ProviderHandle, StatSeries.CurrentFrame, StatSeries.CurrentValue.Value.AsInt);
+		StatSeries.Counter->SetValue(CounterTimestamp, StatSeries.CurrentValue.Value.AsInt);
+	}
+	else
+	{
+		CsvProfilerProvider.SetValue(StatSeries.ProviderHandle, StatSeries.CurrentFrame, StatSeries.CurrentValue.Value.AsDouble);
+		StatSeries.Counter->SetValue(CounterTimestamp, StatSeries.CurrentValue.Value.AsDouble);
+	}
+	if (bEnableCounts)
+	{
+		CsvProfilerProvider.SetValue(StatSeries.ProviderCountHandle, StatSeries.CurrentFrame, StatSeries.CurrentCount);
+	}
+	StatSeries.CurrentValue = FStatSeriesValue();
+	StatSeries.CurrentCount = 0;
+}
+
+void FCsvProfilerAnalyzer::FlushIfNewFrame(FStatSeriesInstance& StatSeries, uint32 FrameNumber)
+{
+	if (FrameNumber != StatSeries.CurrentFrame && StatSeries.CurrentValue.bIsValid)
+	{
+		check(FrameNumber > StatSeries.CurrentFrame);
+		Flush(StatSeries);
+	}
+	StatSeries.CurrentFrame = FrameNumber;
+}
+
+void FCsvProfilerAnalyzer::FlushAtEndOfCapture(FStatSeriesInstance& StatSeries, uint32 CaptureEndFrame)
+{
+	if (StatSeries.CurrentValue.bIsValid && StatSeries.CurrentFrame < CaptureEndFrame)
+	{
+		Flush(StatSeries);
+	}
+}
+
+void FCsvProfilerAnalyzer::SetTimerValue(FStatSeriesInstance& StatSeries, uint32 FrameNumber, double ElapsedTime, bool bCount)
+{
+	FlushIfNewFrame(StatSeries, FrameNumber);
+
+	StatSeries.CurrentValue.Value.AsDouble += ElapsedTime;
+	StatSeries.CurrentValue.bIsValid = true;
+	if (bCount)
+	{
+		++StatSeries.CurrentCount;
+	}
+
+}
+
+void FCsvProfilerAnalyzer::SetCustomStatValue(FStatSeriesInstance& StatSeries, uint32 FrameNumber, ECsvOpType OpType, int32 Value)
+{
+	FlushIfNewFrame(StatSeries, FrameNumber);
+
+	if (!StatSeries.CurrentValue.bIsValid)
+	{
+		// The first op in a frame is always a set. Otherwise min/max don't work
+		OpType = CsvOpType_Set;
+	}
+
+	switch (OpType)
+	{
+	case CsvOpType_Set:
+		StatSeries.CurrentValue.Value.AsInt = Value;
+		break;
+	case CsvOpType_Min:
+		StatSeries.CurrentValue.Value.AsInt = FMath::Min(int64(Value), StatSeries.CurrentValue.Value.AsInt);
+		break;
+	case CsvOpType_Max:
+		StatSeries.CurrentValue.Value.AsInt = FMath::Max(int64(Value), StatSeries.CurrentValue.Value.AsInt);
+		break;
+	case CsvOpType_Accumulate:
+		StatSeries.CurrentValue.Value.AsInt += Value;
+		break;
+	}
+	StatSeries.CurrentValue.bIsValid = true;
+	++StatSeries.CurrentCount;
+}
+
+void FCsvProfilerAnalyzer::SetCustomStatValue(FStatSeriesInstance& StatSeries, uint32 FrameNumber, ECsvOpType OpType, float Value)
+{
+	FlushIfNewFrame(StatSeries, FrameNumber);
+
+	if (!StatSeries.CurrentValue.bIsValid)
+	{
+		// The first op in a frame is always a set. Otherwise min/max don't work
+		OpType = CsvOpType_Set;
 	}
 
-	uint32 Index = static_cast<uint32>(Algo::LowerBound(FrameStartTimes, Timestamp));
-	return Index + 1;
+	switch (OpType)
+	{
+	case CsvOpType_Set:
+		StatSeries.CurrentValue.Value.AsDouble = Value;
+		break;
+	case CsvOpType_Min:
+		StatSeries.CurrentValue.Value.AsDouble = FMath::Min(double(Value), StatSeries.CurrentValue.Value.AsDouble);
+		break;
+	case CsvOpType_Max:
+		StatSeries.CurrentValue.Value.AsDouble = FMath::Max(double(Value), StatSeries.CurrentValue.Value.AsDouble);
+		break;
+	case CsvOpType_Accumulate:
+		StatSeries.CurrentValue.Value.AsDouble += Value;
+		break;
+	}
+	StatSeries.CurrentValue.bIsValid = true;
+	++StatSeries.CurrentCount;
 }
diff --git a/Engine/Source/Developer/TraceServices/Private/Analyzers/CsvProfilerTraceAnalysis.h b/Engine/Source/Developer/TraceServices/Private/Analyzers/CsvProfilerTraceAnalysis.h
--- a/Engine/Source/Developer/TraceServices/Private/Analyzers/CsvProfilerTraceAnalysis.h
+++ b/Engine/Source/Developer/TraceServices/Private/Analyzers/CsvProfilerTraceAnalysis.h
@@ -13,13 +13,15 @@ namespace Trace
 	class IAnalysisSession;
 	class IFrameProvider;
 	class IThreadProvider;
+	class ICounterProvider;
+	class ICounter;
 }
 
 class FCsvProfilerAnalyzer
 	: public Trace::IAnalyzer
 {
 public:
-	FCsvProfilerAnalyzer(Trace::IAnalysisSession& Session, Trace::FCsvProfilerProvider& CsvProfilerProvider, const Trace::IFrameProvider& FrameProvider, const Trace::IThreadProvider& ThreadProvider);
+	FCsvProfilerAnalyzer(Trace::IAnalysisSession& Session, Trace::FCsvProfilerProvider& CsvProfilerProvider, Trace::ICounterProvider& CounterProvider, const Trace::IFrameProvider& FrameProvider, const Trace::IThreadProvider& ThreadProvider);
 	~FCsvProfilerAnalyzer();
 	virtual void OnAnalysisBegin(const FOnAnalysisContext& Context) override;
 	virtual void OnEvent(uint16 RouteId, const FOnEventContext& Context) override;
@@ -43,6 +45,14 @@ private:
 		RouteId_EndCapture,
 	};
 
+	enum ECsvOpType
+	{
+		CsvOpType_Set,
+		CsvOpType_Min,
+		CsvOpType_Max,
+		CsvOpType_Accumulate,
+	};
+
 	union FStatSeriesID
 	{
 		struct
@@ -55,6 +65,17 @@ private:
 		uint64 Hash;
 	};
 
+	struct FStatSeriesValue
+	{
+		FStatSeriesValue() { Value.AsInt = 0; }
+		union
+		{
+			int64 AsInt;
+			double AsDouble;
+		} Value;
+		bool bIsValid = false;
+	};
+
 	struct FStatSeriesDefinition
 	{
 		const TCHAR* Name = nullptr;
@@ -66,11 +87,17 @@ private:
 	{
 		uint64 ProviderHandle = uint64(-1);
 		uint64 ProviderCountHandle = uint64(-1);
+		Trace::ICounter* Counter = nullptr;
+		int64 CurrentFrame = -1;
+		FStatSeriesValue CurrentValue;
+		int64 CurrentCount = 0;
+		Trace::ECsvStatSeriesType Type = Trace::CsvStatSeriesType_CustomStatInt;
+		ETraceFrameType FrameType = TraceFrameType_Game;
 	};
 
 	struct FTimingMarker
 	{
-		const FStatSeriesInstance* StatSeries = nullptr;
+		uint64 StatId = 0;
 		uint64 Cycle = 0;
 		bool bIsBegin = false;
 		bool bIsExclusive = false;
@@ -90,15 +117,22 @@ private:
 	FStatSeriesDefinition* CreateStatSeries(const TCHAR* Name, int32 CategoryIndex);
 	void DefineStatSeries(uint64 StatId, const TCHAR* Name, int32 CategoryIndex, bool bIsInline);
 	const TCHAR* GetStatSeriesName(const FStatSeriesDefinition* Definition, Trace::ECsvStatSeriesType Type, FThreadState& ThreadState, bool bIsCount);
-	FStatSeriesInstance* GetStatSeries(uint64 StatId, Trace::ECsvStatSeriesType Type, FThreadState& ThreadState);
+	FStatSeriesInstance& GetStatSeries(uint64 StatId, Trace::ECsvStatSeriesType Type, FThreadState& ThreadState);
 	void HandleMarkerEvent(const FOnEventContext& Context, bool bIsExclusive, bool bIsBegin);
 	void HandleMarker(const FOnEventContext& Context, FThreadState& ThreadState, const FTimingMarker& Marker);
 	void HandleCustomStatEvent(const FOnEventContext& Context, bool bIsFloat);
 	void HandleEventEvent(const FOnEventContext& Context);
-	int32 GetFrameNumberForTimestamp(ETraceFrameType FrameType, double Timestamp) const;
+	uint32 GetFrameNumberForTimestamp(ETraceFrameType FrameType, double Timestamp) const;
+	void Flush(FStatSeriesInstance& StatSeries);
+	void FlushIfNewFrame(FStatSeriesInstance& StatSeries, uint32 FrameNumber);
+	void FlushAtEndOfCapture(FStatSeriesInstance& StatSeries, uint32 CaptureEndFrame);
+	void SetTimerValue(FStatSeriesInstance& StatSeries, uint32 FrameNumber, double ElapsedTime, bool bCount);
+	void SetCustomStatValue(FStatSeriesInstance& StatSeries, uint32 FrameNumber, ECsvOpType OpType, int32 Value);
+	void SetCustomStatValue(FStatSeriesInstance& StatSeries, uint32 FrameNumber, ECsvOpType OpType, float Value);
 	
 	Trace::IAnalysisSession& Session;
 	Trace::FCsvProfilerProvider& CsvProfilerProvider;
+	Trace::ICounterProvider& CounterProvider;
 	const Trace::IFrameProvider& FrameProvider;
 	const Trace::IThreadProvider& ThreadProvider;
 
diff --git a/Engine/Source/Developer/TraceServices/Private/Model/Counters.cpp b/Engine/Source/Developer/TraceServices/Private/Model/Counters.cpp
--- a/Engine/Source/Developer/TraceServices/Private/Model/Counters.cpp
+++ b/Engine/Source/Developer/TraceServices/Private/Model/Counters.cpp
@@ -21,21 +21,41 @@ FCounter::FCounter(ILinearAllocator& Allocator, const TArray<double>& InFrameSta
 template<typename CounterType, typename EnumerationType>
 static void EnumerateCounterValuesInternal(const TCounterData<CounterType>& CounterData, const TArray<double>& FrameStartTimes, bool bResetEveryFrame, double IntervalStart, double IntervalEnd, TFunctionRef<void(double, EnumerationType)> Callback)
 {
+	if (!CounterData.Num())
+	{
+		return;
+	}
 	TArray<double> NoFrameStartTimes;
 	auto CounterIterator = bResetEveryFrame ? CounterData.GetIterator(FrameStartTimes) : CounterData.GetIterator(NoFrameStartTimes);
+	bool bFirstValue = true;
+	bool bFirstEnumeratedValue = true;
+	CounterType LastValue = CounterType();
+	double LastTime = 0.0;
 	while (CounterIterator)
 	{
 		const TTuple<double, CounterType>& Current = *CounterIterator;
 		double Time = Current.template Get<0>();
-		CounterType CounterValue = Current.template Get<1>();
-		if (Time > IntervalEnd)
-		{
-			break;
-		}
+		CounterType CurrentValue = Current.template Get<1>();
 		if (Time >= IntervalStart)
 		{
-			Callback(Time, CounterValue);
+			if (bFirstEnumeratedValue)
+			{
+				if (!bFirstValue)
+				{
+					Callback(LastTime, LastValue);
+				}
+				bFirstEnumeratedValue = false;
+			}
+			Callback(Time, CurrentValue);
+			bFirstEnumeratedValue = false;
+			if (Time > IntervalEnd)
+			{
+				break;
+			}
 		}
+		LastTime = Time;
+		LastValue = CurrentValue;
+		bFirstValue = false;
 		++CounterIterator;
 	}
 }
diff --git a/Engine/Source/Developer/TraceServices/Private/Model/CountersPrivate.h b/Engine/Source/Developer/TraceServices/Private/Model/CountersPrivate.h
--- a/Engine/Source/Developer/TraceServices/Private/Model/CountersPrivate.h
+++ b/Engine/Source/Developer/TraceServices/Private/Model/CountersPrivate.h
@@ -140,6 +140,11 @@ public:
 		OpArguments.Insert(InsertionIndex) = OpArgument;
 	}
 
+	uint64 Num() const
+	{
+		return Timestamps.Num();
+	}
+
 	TIterator GetIterator(TArray<double> FrameStartTimes) const
 	{
 		return TIterator(*this, FrameStartTimes);
diff --git a/Engine/Source/Developer/TraceServices/Private/Model/CsvProfiler.cpp b/Engine/Source/Developer/TraceServices/Private/Model/CsvProfiler.cpp
--- a/Engine/Source/Developer/TraceServices/Private/Model/CsvProfiler.cpp
+++ b/Engine/Source/Developer/TraceServices/Private/Model/CsvProfiler.cpp
@@ -190,7 +190,7 @@ uint64 FCsvProfilerProvider::AddSeries(const TCHAR* Name, ECsvStatSeriesType Typ
 	return ColumnIndex;
 }
 
-FCsvProfilerProvider::FStatSeriesValue& FCsvProfilerProvider::GetValueRef(uint64 SeriesHandle, int32 FrameNumber)
+FCsvProfilerProvider::FStatSeriesValue& FCsvProfilerProvider::GetValueRef(uint64 SeriesHandle, uint32 FrameNumber)
 {
 	FStatSeries* Series = StatSeries[SeriesHandle];
 	if (CurrentCapture && !Series->Captures.Contains(Captures.Num()))
@@ -201,77 +201,26 @@ FCsvProfilerProvider::FStatSeriesValue& FCsvProfilerProvider::GetValueRef(uint64
 	return Series->Values.FindOrAdd(FrameNumber);
 }
 
-void FCsvProfilerProvider::SetTimerValue(uint64 SeriesHandle, int32 FrameNumber, double ElapsedTime)
+void FCsvProfilerProvider::SetValue(uint64 SeriesHandle, uint32 FrameNumber, int64 Value)
 {
 	Session.WriteAccessCheck();
 
 	FStatSeriesValue& RowValue = GetValueRef(SeriesHandle, FrameNumber);
-	RowValue.Value.AsDouble += ElapsedTime;
-	RowValue.bIsValid = true;
+	RowValue.Value.AsInt = Value;
 }
 
-void FCsvProfilerProvider::SetCustomStatValue(uint64 SeriesHandle, int32 FrameNumber, ECsvOpType OpType, int32 Value)
+void FCsvProfilerProvider::SetValue(uint64 SeriesHandle, uint32 FrameNumber, double Value)
 {
 	Session.WriteAccessCheck();
 
 	FStatSeriesValue& RowValue = GetValueRef(SeriesHandle, FrameNumber);
-
-	if (!RowValue.bIsValid)
-	{
-		// The first op in a frame is always a set. Otherwise min/max don't work
-		OpType = CsvOpType_Set;
-	}
-
-	switch (OpType)
-	{
-	case CsvOpType_Set:
-		RowValue.Value.AsInt = Value;
-		break;
-	case CsvOpType_Min:
-		RowValue.Value.AsInt = FMath::Min(int64(Value), RowValue.Value.AsInt);
-		break;
-	case CsvOpType_Max:
-		RowValue.Value.AsInt = FMath::Max(int64(Value), RowValue.Value.AsInt);
-		break;
-	case CsvOpType_Accumulate:
-		RowValue.Value.AsInt += Value;
-		break;
-	}
-	RowValue.bIsValid = true;
+	RowValue.Value.AsDouble = Value;
 }
 
-void FCsvProfilerProvider::SetCustomStatValue(uint64 SeriesHandle, int32 FrameNumber, ECsvOpType OpType, float Value)
+void FCsvProfilerProvider::AddEvent(uint32 FrameNumber, const TCHAR* Text)
 {
 	Session.WriteAccessCheck();
 
-	FStatSeriesValue& RowValue = GetValueRef(SeriesHandle, FrameNumber);
-
-	if (!RowValue.bIsValid)
-	{
-		// The first op in a frame is always a set. Otherwise min/max don't work
-		OpType = CsvOpType_Set;
-	}
-
-	switch (OpType)
-	{
-	case CsvOpType_Set:
-		RowValue.Value.AsDouble = Value;
-		break;
-	case CsvOpType_Min:
-		RowValue.Value.AsDouble = FMath::Min(double(Value), RowValue.Value.AsDouble);
-		break;
-	case CsvOpType_Max:
-		RowValue.Value.AsDouble = FMath::Max(double(Value), RowValue.Value.AsDouble);
-		break;
-	case CsvOpType_Accumulate:
-		RowValue.Value.AsDouble += Value;
-		break;
-	}
-	RowValue.bIsValid = true;
-}
-
-void FCsvProfilerProvider::AddEvent(int32 FrameNumber, const TCHAR* Text)
-{
 	FEvents* FrameEvents;
 	FEvents** FindIt = Events.Find(FrameNumber);
 	if (!FindIt)
@@ -296,6 +245,8 @@ void FCsvProfilerProvider::AddEvent(int32 FrameNumber, const TCHAR* Text)
 
 void FCsvProfilerProvider::SetMetadata(const TCHAR* Key, const TCHAR* Value)
 {
+	Session.WriteAccessCheck();
+
 	Metadata.Add(Key, Value);
 }
 
diff --git a/Engine/Source/Developer/TraceServices/Private/Model/CsvProfilerPrivate.h b/Engine/Source/Developer/TraceServices/Private/Model/CsvProfilerPrivate.h
--- a/Engine/Source/Developer/TraceServices/Private/Model/CsvProfilerPrivate.h
+++ b/Engine/Source/Developer/TraceServices/Private/Model/CsvProfilerPrivate.h
@@ -12,19 +12,10 @@ namespace Trace
 enum ECsvStatSeriesType
 {
 	CsvStatSeriesType_Timer,
-	CsvStatSeriesType_ExclusiveTimer,
 	CsvStatSeriesType_CustomStatInt,
 	CsvStatSeriesType_CustomStatFloat
 };
 
-enum ECsvOpType
-{
-	CsvOpType_Set,
-	CsvOpType_Min,
-	CsvOpType_Max,
-	CsvOpType_Accumulate,
-};
-
 class FCsvProfilerProvider
 	: public ICsvProfilerProvider
 {
@@ -36,10 +27,9 @@ public:
 	void StartCapture(const TCHAR* Filename, int32 FrameNumber);
 	void EndCapture(int32 FrameNumber);
 	uint64 AddSeries(const TCHAR* Name, ECsvStatSeriesType Type);
-	void SetTimerValue(uint64 SeriesHandle, int32 FrameNumber, double ElapsedTime);
-	void SetCustomStatValue(uint64 SeriesHandle, int32 FrameNumber, ECsvOpType OpType, int32 Value);
-	void SetCustomStatValue(uint64 SeriesHandle, int32 FrameNumber, ECsvOpType OpType, float Value);
-	void AddEvent(int32 FrameNumber, const TCHAR* Text);
+	void SetValue(uint64 SeriesHandle, uint32 FrameNumber, double Value);
+	void SetValue(uint64 SeriesHandle, uint32 FrameNumber, int64 Value);
+	void AddEvent(uint32 FrameNumber, const TCHAR* Text);
 	void SetMetadata(const TCHAR* Key, const TCHAR* Value);
 
 private:
@@ -51,7 +41,6 @@ private:
 			int64 AsInt;
 			double AsDouble;
 		} Value;
-		bool bIsValid = false;
 	};
 
 	struct FStatSeries
@@ -92,7 +81,7 @@ private:
 		: public IUntypedTableReader
 	{
 	public:
-		FTableReader(const FCapture& InCapture, const TMap<int32, FEvents*>& InEvents)
+		FTableReader(const FCapture& InCapture, const TMap<uint32, FEvents*>& InEvents)
 			: Capture(InCapture)
 			, Events(InEvents)
 			, CurrentRowIndex(0)
@@ -112,7 +101,7 @@ private:
 		const FStatSeriesValue* GetValue(uint64 ColumnIndex) const;
 
 		const FCapture& Capture;
-		const TMap<int32, FEvents*>& Events;
+		const TMap<uint32, FEvents*>& Events;
 		uint64 CurrentRowIndex;
 	};
 
@@ -120,7 +109,7 @@ private:
 		: public IUntypedTable
 	{
 	public:
-		FTable(const FCapture& InCapture, const TMap<int32, FEvents*>& InEvents)
+		FTable(const FCapture& InCapture, const TMap<uint32, FEvents*>& InEvents)
 			: Layout(InCapture.StatSeries)
 			, Capture(InCapture)
 			, Events(InEvents)
@@ -135,12 +124,12 @@ private:
 	private:
 		FTableLayout Layout;
 		const FCapture& Capture;
-		const TMap<int32, FEvents*>& Events;
+		const TMap<uint32, FEvents*>& Events;
 	};
 
 	struct FCapture
 	{
-		FCapture(const TMap<int32, FEvents*>& Events)
+		FCapture(const TMap<uint32, FEvents*>& Events)
 			: Table(*this, Events)
 		{
 			
@@ -148,17 +137,17 @@ private:
 
 		FTable Table;
 		const TCHAR* Filename;
-		int32 StartFrame = -1;
-		int32 EndFrame = -1;
+		uint32 StartFrame = 0;
+		uint32 EndFrame = 0;
 		TArray<FStatSeries*> StatSeries;
 		TMap<const TCHAR*, const TCHAR*> Metadata;
 	};
 
-	FStatSeriesValue& GetValueRef(uint64 SeriesHandle, int32 FrameNumber);
+	FStatSeriesValue& GetValueRef(uint64 SeriesHandle, uint32 FrameNumber);
 
 	IAnalysisSession& Session;
 	TArray<FStatSeries*> StatSeries;
-	TMap<int32, FEvents*> Events;
+	TMap<uint32, FEvents*> Events;
 	TMap<const TCHAR*, const TCHAR*> Metadata;
 	TArray<FCapture*> Captures;
 	FCapture* CurrentCapture = nullptr;
diff --git a/Engine/Source/Developer/TraceServices/Private/Model/Frames.cpp b/Engine/Source/Developer/TraceServices/Private/Model/Frames.cpp
--- a/Engine/Source/Developer/TraceServices/Private/Model/Frames.cpp
+++ b/Engine/Source/Developer/TraceServices/Private/Model/Frames.cpp
@@ -40,6 +40,21 @@ void FFrameProvider::EnumerateFrames(ETraceFrameType FrameType, uint64 Start, ui
 	}
 }
 
+const FFrame* FFrameProvider::GetFrame(ETraceFrameType FrameType, uint64 Index) const
+{
+	Session.ReadAccessCheck();
+
+	TPagedArray<FFrame> FramesOfType = Frames[FrameType];
+	if (Index < FramesOfType.Num())
+	{
+		return &FramesOfType[Index];
+	}
+	else
+	{
+		return nullptr;
+	}
+}
+
 void FFrameProvider::BeginFrame(ETraceFrameType FrameType, double Time)
 {
 	Session.WriteAccessCheck();
diff --git a/Engine/Source/Developer/TraceServices/Private/Model/FramesPrivate.h b/Engine/Source/Developer/TraceServices/Private/Model/FramesPrivate.h
--- a/Engine/Source/Developer/TraceServices/Private/Model/FramesPrivate.h
+++ b/Engine/Source/Developer/TraceServices/Private/Model/FramesPrivate.h
@@ -20,6 +20,7 @@ public:
 	virtual uint64 GetFrameCount(ETraceFrameType FrameType) const override;
 	virtual void EnumerateFrames(ETraceFrameType FrameType, uint64 Start, uint64 End, TFunctionRef<void(const FFrame&)> Callback) const override;
 	virtual const TArray<double>& GetFrameStartTimes(ETraceFrameType FrameType) const override { return FrameStartTimes[FrameType]; }
+	virtual const FFrame* GetFrame(ETraceFrameType FrameType, uint64 Index) const override;
 	void BeginFrame(ETraceFrameType FrameType, double Time);
 	void EndFrame(ETraceFrameType FrameType, double Time);
 
diff --git a/Engine/Source/Developer/TraceServices/Private/Modules/CsvProfilerModule.cpp b/Engine/Source/Developer/TraceServices/Private/Modules/CsvProfilerModule.cpp
--- a/Engine/Source/Developer/TraceServices/Private/Modules/CsvProfilerModule.cpp
+++ b/Engine/Source/Developer/TraceServices/Private/Modules/CsvProfilerModule.cpp
@@ -22,9 +22,10 @@ void FCsvProfilerModule::OnAnalysisBegin(IAnalysisSession& Session)
 {
 	const IFrameProvider& FrameProvider = ReadFrameProvider(Session);
 	const IThreadProvider& ThreadProvider = ReadThreadProvider(Session);
+	ICounterProvider& CounterProvider = EditCounterProvider(Session);
 	FCsvProfilerProvider* CsvProfilerProvider = new FCsvProfilerProvider(Session);
 	Session.AddProvider(CsvProfilerProviderName, CsvProfilerProvider);
-	Session.AddAnalyzer(new FCsvProfilerAnalyzer(Session, *CsvProfilerProvider, FrameProvider, ThreadProvider));
+	Session.AddAnalyzer(new FCsvProfilerAnalyzer(Session, *CsvProfilerProvider, CounterProvider, FrameProvider, ThreadProvider));
 }
 
 void FCsvProfilerModule::GenerateReports(const IAnalysisSession& Session, const TCHAR* CmdLine, const TCHAR* OutputDirectory)
diff --git a/Engine/Source/Developer/TraceServices/Public/TraceServices/Model/Frames.h b/Engine/Source/Developer/TraceServices/Public/TraceServices/Model/Frames.h
--- a/Engine/Source/Developer/TraceServices/Public/TraceServices/Model/Frames.h
+++ b/Engine/Source/Developer/TraceServices/Public/TraceServices/Model/Frames.h
@@ -26,6 +26,7 @@ public:
 	virtual uint64 GetFrameCount(ETraceFrameType FrameType) const = 0;
 	virtual void EnumerateFrames(ETraceFrameType FrameType, uint64 Start, uint64 End, TFunctionRef<void(const FFrame&)> Callback) const = 0;
 	virtual const TArray<double>& GetFrameStartTimes(ETraceFrameType FrameType) const = 0;
+	virtual const FFrame* GetFrame(ETraceFrameType FrameType, uint64 Index) const = 0;
 };
 
 TRACESERVICES_API const IFrameProvider& ReadFrameProvider(const IAnalysisSession& Session);
