First iteration of NetTraceAnalyzer/NetProfiler

diff --git a/Engine/Source/Developer/TraceServices/Private/AnalysisService.cpp b/Engine/Source/Developer/TraceServices/Private/AnalysisService.cpp
--- a/Engine/Source/Developer/TraceServices/Private/AnalysisService.cpp
+++ b/Engine/Source/Developer/TraceServices/Private/AnalysisService.cpp
@@ -15,6 +15,7 @@
 #include "Model/BookmarksPrivate.h"
 #include "Model/ThreadsPrivate.h"
 #include "Model/CountersPrivate.h"
+#include "Model/NetProfilerProvider.h"
 
 namespace Trace
 {
@@ -140,18 +141,28 @@ void FAnalysisService::AnalyzeInternal(TSharedRef<FAnalysisSession> AnalysisSess
 	{
 		IAnalysisSession& Session = AnalysisSession.Get();
 		Trace::FAnalysisSessionEditScope _(Session);
+
 		FBookmarkProvider* BookmarkProvider = new FBookmarkProvider(Session);
 		Session.AddProvider(FBookmarkProvider::ProviderName, BookmarkProvider);
+
 		FLogProvider* LogProvider = new FLogProvider(Session);
 		Session.AddProvider(FLogProvider::ProviderName, LogProvider);
+
 		FThreadProvider* ThreadProvider = new FThreadProvider(Session);
 		Session.AddProvider(FThreadProvider::ProviderName, ThreadProvider);
+
 		FFrameProvider* FrameProvider = new FFrameProvider(Session);
 		Session.AddProvider(FFrameProvider::ProviderName, FrameProvider);
+
 		FCounterProvider* CounterProvider = new FCounterProvider(Session, *FrameProvider);
 		Session.AddProvider(FCounterProvider::ProviderName, CounterProvider);
+
+		FNetProfilerProvider* NetProfilerProvider = new FNetProfilerProvider(Session);
+		Session.AddProvider(FNetProfilerProvider::ProviderName, NetProfilerProvider);
+
 		Session.AddAnalyzer(new FMiscTraceAnalyzer(Session, *ThreadProvider, *BookmarkProvider, *LogProvider, *FrameProvider));
 		Session.AddAnalyzer(new FLogTraceAnalyzer(Session, *LogProvider));
+
 		ModuleService.OnAnalysisBegin(Session);
 	}
 	
diff --git a/Engine/Source/Developer/TraceServices/Private/Analyzers/NetTraceAnalyzer.cpp b/Engine/Source/Developer/TraceServices/Private/Analyzers/NetTraceAnalyzer.cpp
new file mode 100644
--- /dev/null
+++ b/Engine/Source/Developer/TraceServices/Private/Analyzers/NetTraceAnalyzer.cpp
@@ -0,0 +1,350 @@
+// Copyright 1998-2019 Epic Games, Inc. All Rights Reserved.
+#include "NetTraceAnalyzer.h"
+#include "AnalysisServicePrivate.h"
+#include "Common/Utils.h"
+#include "TraceServices/Model/Threads.h"
+#include "Logging/LogMacros.h"
+
+FNetTraceAnalyzer::FNetTraceAnalyzer(Trace::IAnalysisSession& InSession, Trace::FNetProfilerProvider& InNetProfilerProvider)
+	: Session(InSession)
+	, NetProfilerProvider(InNetProfilerProvider)
+	, NetTraceVersion(0)
+	, NetTraceReporterVersion(0)
+{
+}
+
+void FNetTraceAnalyzer::OnAnalysisBegin(const FOnAnalysisContext& Context)
+{
+	auto& Builder = Context.InterfaceBuilder;
+
+	Builder.RouteEvent(RouteId_InitEvent, "NetTrace", "InitEvent");
+	Builder.RouteEvent(RouteId_InstanceDestroyedEvent, "NetTrace", "InstanceDestroyedEvent");
+	Builder.RouteEvent(RouteId_NameEvent, "NetTrace", "NameEvent");
+	Builder.RouteEvent(RouteId_PacketContentEvent, "NetTrace", "PacketContentEvent");
+	Builder.RouteEvent(RouteId_PacketEvent, "NetTrace", "PacketEvent");
+	Builder.RouteEvent(RouteId_PacketDroppedEvent, "NetTrace", "PacketDroppedEvent");
+	Builder.RouteEvent(RouteId_ConnectionCreatedEvent, "NetTrace", "ConnectionCreatedEvent");
+	// Add some default event types that we use for generic type events to make it easier to extend
+	// ConnectionAdded/Removed connections state /name etc?
+	Builder.RouteEvent(RouteId_ConnectionClosedEvent, "NetTrace", "ConnectionClosedEvent");
+	Builder.RouteEvent(RouteId_ObjectCreatedEvent, "NetTrace", "ObjectCreatedEvent");
+	Builder.RouteEvent(RouteId_ObjectDestroyedEvent, "NetTrace", "ObjectDestroyedEvent");
+}
+
+void FNetTraceAnalyzer::OnAnalysisEnd()
+{
+}
+
+void FNetTraceAnalyzer::OnEvent(uint16 RouteId, const FOnEventContext& Context)
+{
+	Trace::FAnalysisSessionEditScope _(Session);
+
+	// check that we always get the InitEvent before processing any other events
+	if (!ensure(RouteId == RouteId_InitEvent || NetTraceVersion > 0))
+	{
+		return;
+	}
+
+	const auto& EventData = Context.EventData;
+	switch (RouteId)
+	{
+		case RouteId_InitEvent:
+		{
+			// we always trace the version so that we make sure that we are backwards compatible with older trace stream
+			NetTraceVersion = EventData.GetValue<uint32>("NetTraceVersion");
+			NetTraceReporterVersion = EventData.GetValue<uint32>("NetTraceReporterVersion");
+			const uint64 TimestampCycles = EventData.GetValue<uint64>("Timestamp");
+
+			LastTimeStamp = Context.SessionContext.TimestampFromCycle(TimestampCycles);
+		}
+		break;
+
+		case RouteId_InstanceDestroyedEvent:
+		{
+			const uint8 GameInstanceId = EventData.GetValue<uint8>("ReplicationSystemId");
+			DestroyActiveGameInstanceState(GameInstanceId);
+		}
+		break;
+
+		case RouteId_NameEvent:
+		{
+			uint16 TraceNameId = EventData.GetValue<uint16>("NameId");
+			if (TracedNameIdToNetProfilerNameIdMap.Contains(TraceNameId))
+			{
+				// need to update the name
+				check(false);
+			}
+			else
+			{
+				TracedNameIdToNetProfilerNameIdMap.Add(TraceNameId, NetProfilerProvider.AddNetProfilerName(UTF8_TO_TCHAR(EventData.GetAttachment())));
+			}
+		}
+		break;
+
+		case RouteId_PacketContentEvent:
+		{
+			const uint16 ConnectionId = EventData.GetValue<uint16>("ConnectionId");
+			const uint8 GameInstanceId = EventData.GetValue<uint8>("ReplicationSystemId");
+			const uint8 PacketType =  EventData.GetValue<uint8>("PacketType");
+			
+			TSharedRef<FNetTraceGameInstanceState> GameInstanceState = GetOrCreateActiveGameInstanceState(GameInstanceId);
+			TSharedRef<FNetTraceConnectionState> ConnectionState = GetActiveConnectionState(GameInstanceId, ConnectionId);
+
+			Trace::ENetProfilerConnectionMode ConnectionMode = Trace::ENetProfilerConnectionMode(PacketType);
+
+			Trace::FNetProfilerConnectionData& ConnectionData = NetProfilerProvider.EditConnectionData(ConnectionState->ConnectionIndex, ConnectionMode);
+			++ConnectionData.ContentEventChangeCount;
+
+			TPagedArray<Trace::FNetProfilerContentEvent>& Events = ConnectionData.ContentEvents;
+		
+			// New Packet, update state
+			if (ConnectionState->CurrentPacketStartIndex == Events.Num())
+			{
+				ConnectionState->ConnectionMode = ConnectionMode;
+			}
+			else
+			{
+				// Corrupt data
+				check(ConnectionState->CurrentPacketStartIndex == Events.Num() || (ConnectionState->ConnectionMode == ConnectionMode));
+			}
+
+			// Decode batched events
+			uint64 BufferSize = EventData.GetAttachmentSize();
+			const uint8* BufferPtr = EventData.GetAttachment();
+			const uint8* BufferEnd = BufferPtr + BufferSize;
+			uint64 LastOffset = 0;
+			
+			while (BufferPtr < BufferEnd)
+			{
+				// Decode data
+				const uint64 DecodedNameOrObjectId = FTraceAnalyzerUtils::Decode7bit(BufferPtr);
+				const uint64 DecodedEventStartPos = FTraceAnalyzerUtils::Decode7bit(BufferPtr) + LastOffset;
+				LastOffset = DecodedEventStartPos;
+				const uint64 DecodedEventEndPos = FTraceAnalyzerUtils::Decode7bit(BufferPtr) + DecodedEventStartPos;
+				const uint8 DecodedEventType = *BufferPtr++;
+				const uint8 DecodedNestingLevel = *BufferPtr++;
+
+				// Store the event, should we store it immediately? probably.
+				Trace::FNetProfilerContentEvent& Event = Events.PushBack();
+
+				Event.StartPos = DecodedEventStartPos;
+				Event.EndPos = DecodedEventEndPos;
+				Event.Level = DecodedNestingLevel;
+
+				Event.ObjectInstanceIndex = 0;
+				Event.NameIndex = 0;
+
+				if (DecodedEventType == 0)
+				{
+					// Object index, need to lookup name indirectly
+					if (const FNetTraceActiveObjectState* ActiveObjectState = GameInstanceState->ActiveObjects.Find(DecodedNameOrObjectId))
+					{
+						Event.NameIndex = ActiveObjectState->NameIndex;
+						Event.ObjectInstanceIndex = ActiveObjectState->ObjectIndex;
+					}
+				}
+				else if (DecodedEventType == 1)
+				{
+					if (const uint32* NetProfilerNameIndex = TracedNameIdToNetProfilerNameIdMap.Find(DecodedNameOrObjectId))
+					{
+						Event.NameIndex = *NetProfilerNameIndex;
+					}
+				}
+			}
+			check(BufferPtr == BufferEnd);
+			
+		}
+		break;
+
+		case RouteId_PacketEvent:
+		{
+			const uint64 TimestampCycles = EventData.GetValue<uint64>("Timestamp");
+			const uint32 SequenceNumber = EventData.GetValue<uint32>("SequenceNumber");
+			const uint8 GameInstanceId = EventData.GetValue<uint8>("ReplicationSystemId");
+			const uint16 ConnectionId = EventData.GetValue<uint16>("ConnectionId");
+			
+			// Update LastTimestamp, latar on we will be able to get timestamps piggybacked from other analyzers
+			LastTimeStamp = Context.SessionContext.TimestampFromCycle(TimestampCycles);
+
+			TSharedRef<FNetTraceGameInstanceState> GameInstanceState = GetOrCreateActiveGameInstanceState(GameInstanceId);
+			TSharedRef<FNetTraceConnectionState> ConnectionState = GetActiveConnectionState(GameInstanceId, ConnectionId);
+
+			// Add the packet
+			Trace::FNetProfilerConnectionData& ConnectionData = NetProfilerProvider.EditConnectionData(ConnectionState->ConnectionIndex, ConnectionState->ConnectionMode);
+			Trace::FNetProfilerPacket& Packet = ConnectionData.Packets.PushBack();
+			++ConnectionData.PacketChangeCount;
+
+			// Fill in packet data a packet has at least 1 event
+			Packet.StartEventIndex = ConnectionState->CurrentPacketStartIndex;
+			Packet.EventCount = ConnectionData.ContentEvents.Num() - Packet.StartEventIndex;
+			Packet.TimeStamp = GetLastTimestamp();
+			Packet.SequenceNumber = SequenceNumber;
+			Packet.DeliveryStatus = Trace::ENetProfilerDeliveryStatus::Unknown;
+			Packet.ContentSizeInBits = Packet.EventCount ? ConnectionData.ContentEvents[Packet.StartEventIndex + Packet.EventCount - 1].EndPos : 0u;
+			Packet.TotalPacketSizeInBytes = (Packet.ContentSizeInBits + 7u) >> 3u;
+			Packet.DeliveryStatus = Trace::ENetProfilerDeliveryStatus::Delivered;
+
+			// Finalize received packet data and update data in the NetProfilerProvider
+			//NetProfilerProvider->FinalizePacket();
+
+			// Mark the beginning of a new packet
+			ConnectionState->CurrentPacketStartIndex = ConnectionData.ContentEvents.Num();
+		}
+		break;
+
+		case RouteId_PacketDroppedEvent:
+		{
+			const uint64 TimestampCycles = EventData.GetValue<uint64>("Timestamp");
+			const uint32 SequenceNumber = EventData.GetValue<uint32>("SequenceNumber");
+			const uint8 GameInstanceId = EventData.GetValue<uint8>("ReplicationSystemId");
+			const uint16 ConnectionId = EventData.GetValue<uint16>("ConnectionId");
+			const uint8 PacketType = EventData.GetValue<uint8>("PacketType");
+
+			// Update LastTimestamp, later on we will be able to get timestamps piggybacked from other analyzers
+			LastTimeStamp = Context.SessionContext.TimestampFromCycle(TimestampCycles);
+
+			TSharedRef<FNetTraceGameInstanceState> GameInstanceState = GetOrCreateActiveGameInstanceState(GameInstanceId);
+			TSharedRef<FNetTraceConnectionState> ConnectionState = GetActiveConnectionState(GameInstanceId, ConnectionId);
+
+			Trace::FNetProfilerConnectionData& ConnectionData = NetProfilerProvider.EditConnectionData(ConnectionState->ConnectionIndex, Trace::ENetProfilerConnectionMode(PacketType));
+
+			// Update packet delivery status
+			NetProfilerProvider.EditPacketDeliveryStatus(ConnectionState->ConnectionIndex, Trace::ENetProfilerConnectionMode(PacketType), SequenceNumber, Trace::ENetProfilerDeliveryStatus::Dropped);
+		}
+		break;
+
+		case RouteId_ConnectionCreatedEvent:
+		{
+			const uint8 GameInstanceId = EventData.GetValue<uint8>("ReplicationSystemId");
+			const uint16 ConnectionId = EventData.GetValue<uint16>("ConnectionId");
+
+			TSharedRef<FNetTraceGameInstanceState> GameInstanceState = GetOrCreateActiveGameInstanceState(GameInstanceId);
+			check(!GameInstanceState->ActiveConnections.Contains(ConnectionId));
+
+			// Add to both active connections and to persistent connections
+ 			Trace::FNetProfilerConnectionInternal& Connection = NetProfilerProvider.CreateConnection(GameInstanceState->GameInstanceIndex);
+			TSharedRef<FNetTraceConnectionState> ConnectionState = MakeShared<FNetTraceConnectionState>();
+			GameInstanceState->ActiveConnections.Add(ConnectionId, ConnectionState);
+
+			// Fill in Connection data
+			Connection.Connection.ConnectionId = ConnectionId;
+			Connection.Connection.LifeTime.Begin =  GetLastTimestamp();
+			ConnectionState->ConnectionIndex = Connection.Connection.ConnectionIndex;
+			ConnectionState->CurrentPacketStartIndex = 0;
+		}
+		break;
+
+		case RouteId_ConnectionClosedEvent:
+		{
+			const uint8 GameInstanceId = EventData.GetValue<uint8>("ReplicationSystemId");
+			const uint16 ConnectionId = EventData.GetValue<uint16>("ConnectionId");
+
+			TSharedRef<FNetTraceGameInstanceState> GameInstanceState = GetOrCreateActiveGameInstanceState(GameInstanceId);
+
+			if (TSharedRef<FNetTraceConnectionState>* ConnectionState = GameInstanceState->ActiveConnections.Find(ConnectionId))
+			{
+				if (Trace::FNetProfilerConnectionInternal* Connection = NetProfilerProvider.EditConnection((*ConnectionState)->ConnectionIndex))
+				{
+					// Update connection state
+					Connection->Connection.LifeTime.End =  GetLastTimestamp();
+				}
+				GameInstanceState->ActiveConnections.Remove(ConnectionId);
+			}
+			else
+			{
+				// Incomplete trace?  Ignore?
+				check(false);
+			}
+		}
+		break;
+
+		case RouteId_ObjectCreatedEvent:
+		{
+			const uint64 TypeId = EventData.GetValue<uint64>("TypeId");
+			const uint32 ObjectId = EventData.GetValue<uint32>("ObjectId");
+			const uint32 OwnerId = EventData.GetValue<uint32>("OwnerId");
+			const uint16 NameId = EventData.GetValue<uint16>("NameId");
+			const uint8 GameInstanceId = EventData.GetValue<uint8>("ReplicationSystemId");
+			const uint16 ConnectionId = EventData.GetValue<uint16>("ConnectionId");
+			
+			TSharedRef<FNetTraceGameInstanceState> GameInstanceState = GetOrCreateActiveGameInstanceState(GameInstanceId);
+			check(!GameInstanceState->ActiveObjects.Contains(ObjectId));
+	
+			// Add persistent object representation
+			Trace::FNetProfilerObjectInstance& ObjectInstance = NetProfilerProvider.CreateObject(GameInstanceState->GameInstanceIndex);
+
+			// Fill in object data
+			ObjectInstance.LifeTime.Begin =  GetLastTimestamp();
+			const uint32* NetProfilerNameIndex = TracedNameIdToNetProfilerNameIdMap.Find(NameId);
+			ObjectInstance.NameIndex = NetProfilerNameIndex ? *NetProfilerNameIndex : 0u;
+			ObjectInstance.NetId = ObjectId;
+			ObjectInstance.TypeId = TypeId;
+
+			// Add to active objects
+			GameInstanceState->ActiveObjects.Add(ObjectId, { ObjectInstance.ObjectIndex, ObjectInstance.NameIndex });
+		}
+		break;
+
+		case RouteId_ObjectDestroyedEvent:
+		{
+			// Remove from active instances and mark the end timestamp in the persistent instance list
+			const uint8 GameInstanceId = EventData.GetValue<uint8>("ReplicationSystemId");
+			const uint32 ObjectId = EventData.GetValue<uint32>("ObjectId");
+
+			TSharedRef<FNetTraceGameInstanceState> GameInstanceState = GetOrCreateActiveGameInstanceState(GameInstanceId);
+
+			FNetTraceActiveObjectState DestroyedObjectState;
+			if (GameInstanceState->ActiveObjects.RemoveAndCopyValue(ObjectId, DestroyedObjectState))
+			{
+				if (Trace::FNetProfilerObjectInstance* ObjectInstance = NetProfilerProvider.EditObject(GameInstanceId, DestroyedObjectState.ObjectIndex))
+				{
+					// Update object data
+					ObjectInstance->LifeTime.End = GetLastTimestamp();
+				}
+			}
+		}
+		break;
+	}
+}
+
+TSharedRef<FNetTraceAnalyzer::FNetTraceGameInstanceState> FNetTraceAnalyzer::GetOrCreateActiveGameInstanceState(uint32 GameInstanceId)
+{
+	if (TSharedRef<FNetTraceAnalyzer::FNetTraceGameInstanceState>* FoundState = ActiveGameInstances.Find(GameInstanceId))
+	{
+		return *FoundState;
+	}
+	else
+	{
+		// Persistent GameInstance
+		Trace::FNetProfilerGameInstanceInternal& GameInstance = NetProfilerProvider.CreateGameInstance();
+		GameInstance.Instance.GameInstanceId = GameInstanceId;
+		GameInstance.Instance.LifeTime.Begin = GetLastTimestamp();
+
+		// Active GameInstanceState
+		TSharedRef<FNetTraceGameInstanceState> GameInstanceState = MakeShared<FNetTraceGameInstanceState>();
+		ActiveGameInstances.Add(GameInstanceId, GameInstanceState);
+		GameInstanceState->GameInstanceIndex = GameInstance.Instance.GameInstanceIndex;
+
+		return GameInstanceState;
+	}
+}
+
+void FNetTraceAnalyzer::DestroyActiveGameInstanceState(uint32 GameInstanceId)
+{
+	if (TSharedRef<FNetTraceAnalyzer::FNetTraceGameInstanceState>* FoundState = ActiveGameInstances.Find(GameInstanceId))
+	{
+		// Mark as closed
+		if (Trace::FNetProfilerGameInstanceInternal* GameInstance = NetProfilerProvider.EditGameInstance((*FoundState)->GameInstanceIndex))
+		{
+			GameInstance->Instance.LifeTime.End = GetLastTimestamp();
+		}
+		ActiveGameInstances.Remove(GameInstanceId);
+	}
+}
+
+TSharedRef<FNetTraceAnalyzer::FNetTraceConnectionState> FNetTraceAnalyzer::GetActiveConnectionState(uint32 GameInstanceId, uint32 ConnectionId)
+{
+	TSharedRef<FNetTraceGameInstanceState> GameInstanceState = GetOrCreateActiveGameInstanceState(GameInstanceId);
+
+	return GameInstanceState->ActiveConnections.FindChecked(ConnectionId);
+}
diff --git a/Engine/Source/Developer/TraceServices/Private/Analyzers/NetTraceAnalyzer.h b/Engine/Source/Developer/TraceServices/Private/Analyzers/NetTraceAnalyzer.h
new file mode 100644
--- /dev/null
+++ b/Engine/Source/Developer/TraceServices/Private/Analyzers/NetTraceAnalyzer.h
@@ -0,0 +1,81 @@
+// Copyright 1998-2019 Epic Games, Inc. All Rights Reserved.
+
+#pragma once
+
+#include "Trace/Trace.h"
+#include "Trace/Analyzer.h"
+#include "Containers/UnrealString.h"
+#include "Model/NetProfilerProvider.h"
+
+namespace Trace
+{
+	class IAnalysisSession;
+}
+
+class FNetTraceAnalyzer
+	: public Trace::IAnalyzer
+{
+public:
+	FNetTraceAnalyzer(Trace::IAnalysisSession& Session, Trace::FNetProfilerProvider& NetProfilerProvider);
+	virtual void OnAnalysisBegin(const FOnAnalysisContext& Context) override;
+	virtual void OnEvent(uint16 RouteId, const FOnEventContext& Context) override;
+	virtual void OnAnalysisEnd() override;
+
+private:
+
+	enum : uint16
+	{
+		RouteId_InitEvent,
+		RouteId_InstanceDestroyedEvent,
+		RouteId_NameEvent,
+		RouteId_PacketContentEvent,
+		RouteId_PacketEvent,
+		RouteId_PacketDroppedEvent,
+		RouteId_ConnectionCreatedEvent,
+		RouteId_ConnectionClosedEvent,
+		RouteId_ObjectCreatedEvent,
+		RouteId_ObjectDestroyedEvent,
+	};
+
+	Trace::IAnalysisSession& Session;
+	Trace::FNetProfilerProvider& NetProfilerProvider;
+	uint32 NetTraceVersion;
+	uint32 NetTraceReporterVersion;
+
+	// Shared for trace
+	TMap<uint16, uint32> TracedNameIdToNetProfilerNameIdMap;
+
+	struct FNetTraceConnectionState
+	{
+		// Index into persistent connection array
+		uint32 ConnectionIndex;
+
+		// Current packet data
+		uint32 CurrentPacketStartIndex;
+		Trace::ENetProfilerConnectionMode ConnectionMode;
+	};
+
+	struct FNetTraceActiveObjectState
+	{
+		uint32 ObjectIndex;
+		uint32 NameIndex;
+	};
+
+	// Map from instance id to index in the persistent instance array
+	struct FNetTraceGameInstanceState
+	{
+		TMap<uint32, TSharedRef<FNetTraceConnectionState>> ActiveConnections;
+		TMap<uint32, FNetTraceActiveObjectState> ActiveObjects;
+
+		uint32 GameInstanceIndex;
+	};
+
+	TSharedRef<FNetTraceGameInstanceState> GetOrCreateActiveGameInstanceState(uint32 GameInstanceId);
+	void DestroyActiveGameInstanceState(uint32 GameInstanceId);
+
+	TSharedRef<FNetTraceConnectionState> GetActiveConnectionState(uint32 GameInstanceId, uint32 ConnectionId);
+	Trace::FNetProfilerTimeStamp GetLastTimestamp() const { return LastTimeStamp; }
+
+	TMap<uint32, TSharedRef<FNetTraceGameInstanceState>> ActiveGameInstances;
+	Trace::FNetProfilerTimeStamp LastTimeStamp;
+};
diff --git a/Engine/Source/Developer/TraceServices/Private/Model/NetProfilerProvider.cpp b/Engine/Source/Developer/TraceServices/Private/Model/NetProfilerProvider.cpp
new file mode 100644
--- /dev/null
+++ b/Engine/Source/Developer/TraceServices/Private/Model/NetProfilerProvider.cpp
@@ -0,0 +1,405 @@
+// Copyright 1998-2019 Epic Games, Inc. All Rights Reserved.
+
+#include "NetProfilerProvider.h"
+#include "AnalysisServicePrivate.h"
+#include "Common/StringStore.h"
+
+namespace Trace
+{
+
+const FName FNetProfilerProvider::ProviderName("NetProfilerProvider");
+
+FNetProfilerProvider::FNetProfilerProvider(IAnalysisSession& InSession)
+	: Session(InSession)
+	, Connections(InSession.GetLinearAllocator(), 4096)
+	, ConnectionChangeCount(0u)
+{
+}
+
+FNetProfilerProvider::~FNetProfilerProvider()
+{
+}
+
+uint32 FNetProfilerProvider::AddNetProfilerName(TCHAR* Name)
+{
+	Session.WriteAccessCheck();
+	
+	FNetProfilerName& NewName = Names.AddDefaulted_GetRef();
+	NewName.NameIndex = Names.Num() - 1;
+	NewName.Name = Session.StoreString(Name);
+
+	return NewName.NameIndex;
+}
+
+const FNetProfilerName* FNetProfilerProvider::GetNetProfilerName(uint32 NameIndex) const
+{
+	return NameIndex < (uint32)Names.Num() ? &Names[NameIndex] : nullptr;
+}
+
+Trace::FNetProfilerGameInstanceInternal& FNetProfilerProvider::CreateGameInstance()
+{
+	Session.WriteAccessCheck();
+
+	FNetProfilerGameInstanceInternal& GameInstance = GameInstances.AddDefaulted_GetRef();
+	GameInstance.Instance.GameInstanceIndex = GameInstances.Num() - 1;
+
+	GameInstance.Objects = (TPagedArray<FNetProfilerObjectInstance>*)Session.GetLinearAllocator().Allocate(sizeof(TPagedArray<FNetProfilerObjectInstance>));
+	new (GameInstance.Objects) TPagedArray<FNetProfilerObjectInstance>(Session.GetLinearAllocator(), 4096);
+	GameInstance.ObjectsChangeCount = 0u;
+
+	return GameInstance;
+}
+
+Trace::FNetProfilerGameInstanceInternal* FNetProfilerProvider::EditGameInstance(uint32 GameInstanceIndex)
+{
+	Session.WriteAccessCheck();
+
+	if (ensure(GameInstanceIndex < (uint32)GameInstances.Num()))
+	{
+		return &GameInstances[GameInstanceIndex];
+	}
+	else
+	{
+		return nullptr;
+	}
+}
+
+Trace::FNetProfilerConnectionInternal& FNetProfilerProvider::CreateConnection(uint32 GameInstanceIndex)
+{
+	Session.WriteAccessCheck();
+
+	FNetProfilerGameInstanceInternal* GameInstance = EditGameInstance(GameInstanceIndex);
+	check(GameInstance);
+
+	// Create new connection
+	uint32 ConnectionIndex = Connections.Num();
+	FNetProfilerConnectionInternal& Connection = Connections.PushBack();
+
+	Connection.Connection.ConnectionIndex = ConnectionIndex;
+	Connection.Connection.GameIntanceIndex = GameInstanceIndex;
+	Connection.Connection.bHasIncomingData = false;
+	Connection.Connection.bHasOutgoingData = false;
+
+	GameInstance->Connections.Push(ConnectionIndex);
+
+	// Allocate storage for packets and events
+	Connection.Data[ENetProfilerConnectionMode::Outgoing] = (FNetProfilerConnectionData*)Session.GetLinearAllocator().Allocate(sizeof(FNetProfilerConnectionData));
+	new (Connection.Data[ENetProfilerConnectionMode::Outgoing]) FNetProfilerConnectionData(Session.GetLinearAllocator());
+
+	Connection.Data[ENetProfilerConnectionMode::Incoming] = (FNetProfilerConnectionData*)Session.GetLinearAllocator().Allocate(sizeof(FNetProfilerConnectionData));
+	new (Connection.Data[ENetProfilerConnectionMode::Incoming]) FNetProfilerConnectionData(Session.GetLinearAllocator());
+
+	++ConnectionChangeCount;
+
+	return Connection;
+}
+
+Trace::FNetProfilerObjectInstance& FNetProfilerProvider::CreateObject(uint32 GameInstanceIndex)
+{
+	Session.WriteAccessCheck();
+
+	FNetProfilerGameInstanceInternal* GameInstance = EditGameInstance(GameInstanceIndex);
+	check(GameInstance);
+
+	Trace::FNetProfilerObjectInstance& Object = GameInstance->Objects->PushBack();
+	Object.ObjectIndex = GameInstance->Objects->Num() - 1;
+	++GameInstance->ObjectsChangeCount;
+
+	return Object;
+}
+
+Trace::FNetProfilerObjectInstance* FNetProfilerProvider::EditObject(uint32 GameInstanceIndex, uint32 ObjectIndex)
+{
+	Session.WriteAccessCheck();
+
+	FNetProfilerGameInstanceInternal* GameInstance = EditGameInstance(GameInstanceIndex);
+	check(GameInstance);
+	
+	if (ensure(ObjectIndex < (uint32)GameInstance->Objects->Num()))
+	{
+		++GameInstance->ObjectsChangeCount;
+		return &(*GameInstance->Objects)[ObjectIndex];
+	}
+	else
+	{
+		return nullptr;
+	}
+}
+
+Trace::FNetProfilerConnectionInternal* FNetProfilerProvider::EditConnection(uint32 ConnectionIndex)
+{
+	Session.WriteAccessCheck();
+	
+	if (ensure(ConnectionIndex < (uint32)Connections.Num()))
+	{
+		++ConnectionChangeCount;
+		return &Connections[ConnectionIndex];
+	}
+	else
+	{
+		return nullptr;
+	}
+}
+
+void FNetProfilerProvider::EditPacketDeliveryStatus(uint32 ConnectionIndex, ENetProfilerConnectionMode Mode, uint32 SequenceNumber, ENetProfilerDeliveryStatus DeliveryStatus)
+{
+	check(ConnectionIndex < (uint32)Connections.Num());
+
+	FNetProfilerConnectionInternal& Connection = Connections[ConnectionIndex];
+	Trace::FNetProfilerConnectionData& Data = *Connections[ConnectionIndex].Data[Mode];
+
+	// try to locate packet
+	uint32 PacketCount = Data.Packets.Num();
+	for (uint32 It=0; It < PacketCount; ++It)
+	{
+		const uint32 PacketIndex = PacketCount - It - 1u;
+		if (Data.Packets[PacketIndex].SequenceNumber == SequenceNumber)
+		{
+			Data.Packets[PacketIndex].DeliveryStatus = DeliveryStatus;
+			++Data.PacketChangeCount;
+			return;
+		}
+	}
+}
+
+Trace::FNetProfilerConnectionData& FNetProfilerProvider::EditConnectionData(uint32 ConnectionIndex, ENetProfilerConnectionMode Mode)
+{
+	check(ConnectionIndex < (uint32)Connections.Num());
+
+	FNetProfilerConnectionInternal& Connection = Connections[ConnectionIndex];
+
+	if (Mode == ENetProfilerConnectionMode::Incoming && !Connection.Connection.bHasIncomingData)
+	{
+		Connection.Connection.bHasIncomingData = true;
+		++ConnectionChangeCount;
+	}
+	if (Mode == ENetProfilerConnectionMode::Outgoing && !Connection.Connection.bHasOutgoingData)
+	{
+		Connection.Connection.bHasOutgoingData = true;
+		++ConnectionChangeCount;
+	}
+
+	return *Connections[ConnectionIndex].Data[Mode];
+}
+
+void FNetProfilerProvider::ReadNames(TFunctionRef<void(const FNetProfilerName*, uint64)> Callback) const
+{
+	Session.ReadAccessCheck();
+
+	Callback(Names.GetData(), Names.Num());
+}
+
+void FNetProfilerProvider::ReadName(uint32 NameIndex, TFunctionRef<void(const FNetProfilerName&)> Callback) const
+{
+	Session.ReadAccessCheck();
+	check(NameIndex < (uint32)Names.Num());
+
+	Callback(*GetNetProfilerName(NameIndex));
+}
+
+void FNetProfilerProvider::ReadGameInstances(TFunctionRef<void(const FNetProfilerGameInstance&)> Callback) const
+{
+	Session.ReadAccessCheck();
+
+	for (const FNetProfilerGameInstanceInternal& Instance : GameInstances)
+	{
+		Callback(Instance.Instance);
+	}
+}
+
+uint32 FNetProfilerProvider::GetConnectionCount(uint32 GameInstanceIndex) const
+{
+	Session.ReadAccessCheck();
+
+	check(GameInstanceIndex < (uint32)GameInstances.Num());
+
+	const FNetProfilerGameInstanceInternal& GameInstance = GameInstances[GameInstanceIndex];
+
+	return GameInstance.Connections.Num();
+}
+
+void FNetProfilerProvider::ReadConnections(uint32 GameInstanceIndex, TFunctionRef<void(const FNetProfilerConnection&)> Callback) const
+{
+	Session.ReadAccessCheck();
+
+	check(GameInstanceIndex < (uint32)GameInstances.Num());
+
+	const FNetProfilerGameInstanceInternal& GameInstance = GameInstances[GameInstanceIndex];
+
+	for (uint32 ConnectionIndex : MakeArrayView(GameInstance.Connections.GetData(), GameInstance.Connections.Num()))
+	{
+		Callback(Connections[ConnectionIndex].Connection);
+	}
+}
+
+void FNetProfilerProvider::ReadConnection(uint32 ConnectionIndex, TFunctionRef<void(const FNetProfilerConnection&)> Callback) const
+{
+	Session.ReadAccessCheck();
+
+	check(ConnectionIndex < Connections.Num());
+
+	Callback(Connections[ConnectionIndex].Connection);
+}
+
+uint32 FNetProfilerProvider::GetObjectCount(uint32 GameInstanceIndex) const
+{
+	Session.ReadAccessCheck();
+
+	check(GameInstanceIndex < (uint32)GameInstances.Num());
+
+	const FNetProfilerGameInstanceInternal& GameInstance = GameInstances[GameInstanceIndex];
+
+	return GameInstance.Objects->Num();
+}
+
+void FNetProfilerProvider::ReadObjects(uint32 GameInstanceIndex, TFunctionRef<void(const FNetProfilerObjectInstance&)> Callback) const
+{
+	Session.ReadAccessCheck();
+
+	check(GameInstanceIndex < (uint32)GameInstances.Num());
+
+	const FNetProfilerGameInstanceInternal& GameInstance = GameInstances[GameInstanceIndex];
+	const auto& Objects = *GameInstance.Objects;
+
+	for (uint32 ObjectsIt = 0, ObjectsEndIt = GameInstance.Objects->Num(); ObjectsIt < ObjectsEndIt; ++ObjectsIt)
+	{
+		Callback(Objects[ObjectsIt]);
+	}
+}
+
+void FNetProfilerProvider::ReadObject(uint32 GameInstanceIndex, uint32 ObjectIndex, TFunctionRef<void(const FNetProfilerObjectInstance&)> Callback) const
+{
+	Session.ReadAccessCheck();
+
+	check(GameInstanceIndex < (uint32)GameInstances.Num());
+
+	const FNetProfilerGameInstanceInternal& GameInstance = GameInstances[GameInstanceIndex];
+
+	check(ObjectIndex < GameInstance.Objects->Num());
+
+	Callback((*GameInstance.Objects)[ObjectIndex]);
+}
+
+uint32 FNetProfilerProvider::GeObjectsChangeCount(uint32 GameInstanceIndex) const
+{
+	Session.ReadAccessCheck();
+
+	check(GameInstanceIndex < (uint32)GameInstances.Num());
+
+	const FNetProfilerGameInstanceInternal& GameInstance = GameInstances[GameInstanceIndex];
+
+	return GameInstance.ObjectsChangeCount;
+}
+
+const INetProfilerProvider& ReadNetProfilerProvider(const IAnalysisSession& Session)
+{
+	Session.ReadAccessCheck();
+
+	return *Session.ReadProvider<INetProfilerProvider>(FNetProfilerProvider::ProviderName);
+}
+
+uint32 FNetProfilerProvider::GetPacketCount(uint32 ConnectionIndex, ENetProfilerConnectionMode Mode) const
+{
+	Session.ReadAccessCheck();
+
+	check(ConnectionIndex < Connections.Num());
+
+	const auto& Packets = Connections[ConnectionIndex].Data[Mode]->Packets;
+	
+	return Packets.Num();
+}
+
+void FNetProfilerProvider::EnumeratePackets(uint32 ConnectionIndex, ENetProfilerConnectionMode Mode, uint32 PacketIndexIntervalStart, uint32 PacketIndexIntervalEnd, TFunctionRef<void(const FNetProfilerPacket&)> Callback) const
+{
+	Session.ReadAccessCheck();
+
+	check(ConnectionIndex < Connections.Num());
+
+	const auto& Packets = Connections[ConnectionIndex].Data[Mode]->Packets;
+	
+	const uint32 PacketCount = Packets.Num();
+
+	if (PacketCount == 0 || PacketIndexIntervalStart > PacketIndexIntervalEnd)
+	{
+		return;
+	}
+
+	for (uint32 PacketIt = PacketIndexIntervalStart, PacketEndIt = FMath::Min(PacketIndexIntervalEnd, PacketCount - 1u); PacketIt <= PacketEndIt; ++PacketIt)
+	{
+		Callback(Packets[PacketIt]);
+	}
+}
+
+// Enumerate packet content events by range
+void FNetProfilerProvider::EnumeratePacketContentEventsByIndex(uint32 ConnectionIndex, ENetProfilerConnectionMode Mode, uint32 StartEventIndex, uint32 EndEventIndex, TFunctionRef<void(const FNetProfilerContentEvent&)> Callback) const
+{
+	Session.ReadAccessCheck();
+
+	check(ConnectionIndex < Connections.Num());
+
+	const auto& ContentEvents = Connections[ConnectionIndex].Data[Mode]->ContentEvents;
+	
+	const uint32 EventCount = ContentEvents.Num();
+
+	if (EventCount == 0 || StartEventIndex > EndEventIndex)
+	{
+		return;
+	}
+
+	for (uint32 EventIt = StartEventIndex, EventEndIt = FMath::Min(EndEventIndex, EventCount - 1u); EventIt <= EventEndIt; ++EventIt)
+	{
+		Callback(ContentEvents[EventIt]);
+	}
+}
+
+void FNetProfilerProvider::EnumeratePacketContentEventsByPosition(uint32 ConnectionIndex, ENetProfilerConnectionMode Mode, uint32 PacketIndex, uint32 StartPos, uint32 EndPos, TFunctionRef<void(const FNetProfilerContentEvent&)> Callback) const
+{
+	Session.ReadAccessCheck();
+
+	check(ConnectionIndex < Connections.Num());
+
+	 const FNetProfilerConnectionData& ConnectionData = * Connections[ConnectionIndex].Data[Mode];
+
+	const FNetProfilerPacket& Packet = ConnectionData.Packets[PacketIndex];
+	if (Packet.EventCount == 0u)
+	{
+		return;
+	}
+
+	const uint32 StartEventIndex = Packet.StartEventIndex;
+	const uint32 EndEventIndex = StartEventIndex + Packet.EventCount - 1u;
+
+	const auto& ContentEvents = ConnectionData.ContentEvents;
+	
+	uint32 EventIt = StartEventIndex; 
+	// Skip all Events outside of the scope
+	while (EventIt <= EndEventIndex &&	ContentEvents[EventIt].EndPos < StartPos)
+	{
+		++EventIt;
+	}
+
+	// Execute callback for all found events
+	while (EventIt <= EndEventIndex && ContentEvents[EventIt].StartPos < EndPos)
+	{
+		Callback(ContentEvents[EventIt]);
+		++EventIt;
+	}
+}
+
+uint32 FNetProfilerProvider::GetPacketChangeCount(uint32 ConnectionIndex, ENetProfilerConnectionMode Mode) const
+{
+	Session.ReadAccessCheck();
+	check(ConnectionIndex < Connections.Num());
+
+	return Connections[ConnectionIndex].Data[Mode]->PacketChangeCount;
+}
+
+uint32 FNetProfilerProvider::GetPacketContentEventChangeCount(uint32 ConnectionIndex, ENetProfilerConnectionMode Mode) const
+{
+	Session.ReadAccessCheck();
+	check(ConnectionIndex < Connections.Num());
+
+	return Connections[ConnectionIndex].Data[Mode]->ContentEventChangeCount;
+}
+
+}
diff --git a/Engine/Source/Developer/TraceServices/Private/Model/NetProfilerProvider.h b/Engine/Source/Developer/TraceServices/Private/Model/NetProfilerProvider.h
new file mode 100644
--- /dev/null
+++ b/Engine/Source/Developer/TraceServices/Private/Model/NetProfilerProvider.h
@@ -0,0 +1,117 @@
+// Copyright 1998-2019 Epic Games, Inc. All Rights Reserved.
+
+#pragma once
+
+#include "TraceServices/Model/NetProfiler.h"
+#include "Common/SlabAllocator.h"
+#include "Model/Tables.h"
+
+namespace Trace
+{
+
+class FAnalysisSessionLock;
+class FStringStore;
+
+struct FNetProfilerGameInstanceInternal
+{
+	FNetProfilerGameInstance Instance;
+
+	TPagedArray<Trace::FNetProfilerObjectInstance>* Objects;
+	TArray<uint32, TInlineAllocator<128>> Connections;
+	uint32 ObjectsChangeCount;
+};
+
+struct FNetProfilerConnectionData
+{
+	FNetProfilerConnectionData(ILinearAllocator& Allocator)
+		: Packets(Allocator, 1024)
+		, ContentEvents(Allocator, 8192)
+		, PacketChangeCount(0u)
+		, ContentEventChangeCount(0u)
+	{}
+
+	TPagedArray<FNetProfilerPacket> Packets;
+	TPagedArray<FNetProfilerContentEvent> ContentEvents;
+
+	uint32 PacketChangeCount;
+	uint32 ContentEventChangeCount;
+};
+
+struct FNetProfilerConnectionInternal
+{
+	FNetProfilerConnection Connection;
+
+	FNetProfilerConnectionData* Data[ENetProfilerConnectionMode::Count];
+};
+
+class FNetProfilerProvider
+	: public INetProfilerProvider
+{
+public:
+	static const FName ProviderName;
+
+	explicit FNetProfilerProvider(IAnalysisSession& InSession);
+	virtual ~FNetProfilerProvider();
+
+	uint32 AddNetProfilerName(TCHAR* Name) ;
+	const FNetProfilerName* GetNetProfilerName(uint32 ProfilerNameId) const;
+
+	Trace::FNetProfilerGameInstanceInternal& CreateGameInstance();
+	Trace::FNetProfilerGameInstanceInternal* EditGameInstance(uint32 GameInstanceIndex);
+
+	Trace::FNetProfilerConnectionInternal& CreateConnection(uint32 GameInstanceIndex);
+	Trace::FNetProfilerConnectionInternal* EditConnection(uint32 ConnectionIndex);
+	Trace::FNetProfilerConnectionData& EditConnectionData(uint32 ConnectionIndex, ENetProfilerConnectionMode Mode);
+
+	void EditPacketDeliveryStatus(uint32 ConnectionIndex, ENetProfilerConnectionMode Mode, uint32 SequenceNumber, ENetProfilerDeliveryStatus DeliveryStatus);
+
+	Trace::FNetProfilerObjectInstance& CreateObject(uint32 GameInstanceIndex);
+	Trace::FNetProfilerObjectInstance* EditObject(uint32 GameInstanceIndex, uint32 ObjectIndex);
+
+	// INetProfilerProvider Interface
+
+	// Access Names
+	virtual uint32 GetNameCount() const override { return Names.Num(); }
+	virtual void ReadNames(TFunctionRef<void(const FNetProfilerName*, uint64)> Callback) const override;
+	virtual void ReadName(uint32 NameIndex, TFunctionRef<void(const FNetProfilerName&)> Callback) const override;
+
+	// Access GameInstances
+	virtual uint32 GetGameInstanceCount() const override { return GameInstances.Num(); }
+	virtual void ReadGameInstances(TFunctionRef<void(const FNetProfilerGameInstance&)> Callback) const override;
+
+	// Access Connections
+	virtual uint32 GetConnectionCount(uint32 GameInstanceIndex) const override;
+	virtual void ReadConnections(uint32 GameInstanceIndex, TFunctionRef<void(const FNetProfilerConnection&)> Callback) const override;
+	virtual void ReadConnection(uint32 ConnectionIndex, TFunctionRef<void(const FNetProfilerConnection&)> Callback) const override;
+	virtual uint32 GetConnectionChangeCount() const override { return ConnectionChangeCount; }
+
+	// Access Object Instances
+	virtual uint32 GetObjectCount(uint32 GameInstanceIndex) const override;
+	virtual void ReadObjects(uint32 GameInstanceIndex, TFunctionRef<void(const FNetProfilerObjectInstance&)> Callback) const override;
+	virtual void ReadObject(uint32 GameInstanceIndex, uint32 ObjectIndex, TFunctionRef<void(const FNetProfilerObjectInstance&)> Callback) const override;
+	virtual uint32 GeObjectsChangeCount(uint32 GameInstanceIndex) const override;
+
+	// Enumerate packets in the provided packet interval
+	virtual uint32 GetPacketCount(uint32 ConnectionIndex, ENetProfilerConnectionMode Mode) const override;
+	virtual void EnumeratePackets(uint32 ConnectionIndex, ENetProfilerConnectionMode Mode, uint32 PacketIndexIntervalStart, uint32 PacketIndexIntervalEnd, TFunctionRef<void(const FNetProfilerPacket&)> Callback) const override;
+	virtual uint32 GetPacketChangeCount(uint32 ConnectionIndex, ENetProfilerConnectionMode Mode) const override;
+
+	// Enumerate packet content events by range
+	virtual void EnumeratePacketContentEventsByIndex(uint32 ConnectionIndex, ENetProfilerConnectionMode Mode, uint32 StartEventIndex, uint32 EndEventIndex, TFunctionRef<void(const FNetProfilerContentEvent&)> Callback) const override;
+	virtual void EnumeratePacketContentEventsByPosition(uint32 ConnectionIndex, ENetProfilerConnectionMode Mode, uint32 PacketIndex, uint32 StartPos, uint32 EndPos, TFunctionRef<void(const FNetProfilerContentEvent&)> Callback) const override;
+	virtual uint32 GetPacketContentEventChangeCount(uint32 ConnectionIndex, ENetProfilerConnectionMode Mode) const override;
+	
+private:
+	IAnalysisSession& Session;
+
+	TArray<FNetProfilerName> Names;
+
+	// All GameInstances seen throughout the session
+	TArray<FNetProfilerGameInstanceInternal, TInlineAllocator<4>> GameInstances;
+	
+	// All connections we have seen throughout the session
+	TPagedArray<FNetProfilerConnectionInternal> Connections;
+	uint32 ConnectionChangeCount;
+};
+
+}
diff --git a/Engine/Source/Developer/TraceServices/Private/Modules/NetProfilerModule.cpp b/Engine/Source/Developer/TraceServices/Private/Modules/NetProfilerModule.cpp
new file mode 100644
--- /dev/null
+++ b/Engine/Source/Developer/TraceServices/Private/Modules/NetProfilerModule.cpp
@@ -0,0 +1,34 @@
+// Copyright 1998-2019 Epic Games, Inc. All Rights Reserved.
+
+#include "NetProfilerModule.h"
+#include "Analyzers/NetTraceAnalyzer.h"
+#include "AnalysisServicePrivate.h"
+#include "Model/NetProfilerProvider.h"
+
+namespace Trace
+{
+
+FName FNetProfilerModule::ModuleName("TraceModule_NetProfiler");
+static const FName NetProfilerProviderName("NetProfilerProvider");
+
+void FNetProfilerModule::GetModuleInfo(FModuleInfo& OutModuleInfo)
+{
+	OutModuleInfo.Name = ModuleName;
+	OutModuleInfo.DisplayName = TEXT("NetProfiler");
+}
+	
+void FNetProfilerModule::OnAnalysisBegin(IAnalysisSession& InSession)
+{
+	FAnalysisSession& Session = static_cast<FAnalysisSession&>(InSession);
+
+	FNetProfilerProvider* NetProfilerProvider = new FNetProfilerProvider(Session);
+	Session.AddProvider(NetProfilerProviderName, NetProfilerProvider);
+	InSession.AddAnalyzer(new FNetTraceAnalyzer(InSession, *NetProfilerProvider));
+}
+
+void FNetProfilerModule::GetLoggers(TArray<const TCHAR *>& OutLoggers)
+{
+	OutLoggers.Add(TEXT("NetProfiler"));
+}
+
+}
diff --git a/Engine/Source/Developer/TraceServices/Private/Modules/NetProfilerModule.h b/Engine/Source/Developer/TraceServices/Private/Modules/NetProfilerModule.h
new file mode 100644
--- /dev/null
+++ b/Engine/Source/Developer/TraceServices/Private/Modules/NetProfilerModule.h
@@ -0,0 +1,24 @@
+// Copyright 1998-2019 Epic Games, Inc. All Rights Reserved.
+
+#pragma once
+
+#include "TraceServices/ModuleService.h"
+#include "UObject/NameTypes.h"
+
+namespace Trace
+{
+
+class FNetProfilerModule
+	: public IModule
+{
+public:
+	virtual void GetModuleInfo(FModuleInfo& OutModuleInfo) override;
+	virtual void OnAnalysisBegin(IAnalysisSession& Session) override;
+	virtual void GetLoggers(TArray<const TCHAR*>& OutLoggers) override;
+	virtual void GenerateReports(const IAnalysisSession& Session, const TCHAR* CmdLine, const TCHAR* OutputDirectory) override {}
+
+private:
+	static FName ModuleName;
+};
+
+}
diff --git a/Engine/Source/Developer/TraceServices/Private/TraceServicesModule.cpp b/Engine/Source/Developer/TraceServices/Private/TraceServicesModule.cpp
--- a/Engine/Source/Developer/TraceServices/Private/TraceServicesModule.cpp
+++ b/Engine/Source/Developer/TraceServices/Private/TraceServicesModule.cpp
@@ -11,6 +11,7 @@
 #include "Modules/StatsModule.h"
 #include "Modules/CsvProfilerModule.h"
 #include "Modules/CountersModule.h"
+#include "Modules/NetProfilerModule.h"
 
 class FTraceServicesModule
 	: public ITraceServicesModule
@@ -33,6 +34,7 @@ private:
 	Trace::FStatsModule StatsModule;
 	Trace::FCsvProfilerModule CsvProfilerModule;
 	Trace::FCountersModule CountersModule;
+	Trace::FNetProfilerModule NetProfilerModule;
 };
 
 TSharedPtr<Trace::ISessionService> FTraceServicesModule::GetSessionService()
@@ -73,10 +75,12 @@ void FTraceServicesModule::StartupModule()
 #endif
 	IModularFeatures::Get().RegisterModularFeature(Trace::ModuleFeatureName, &CsvProfilerModule);
 	IModularFeatures::Get().RegisterModularFeature(Trace::ModuleFeatureName, &CountersModule);
+	IModularFeatures::Get().RegisterModularFeature(Trace::ModuleFeatureName, &NetProfilerModule);
 }
 
 void FTraceServicesModule::ShutdownModule()
 {
+	IModularFeatures::Get().UnregisterModularFeature(Trace::ModuleFeatureName, &NetProfilerModule);
 	IModularFeatures::Get().UnregisterModularFeature(Trace::ModuleFeatureName, &CountersModule);
 	IModularFeatures::Get().UnregisterModularFeature(Trace::ModuleFeatureName, &CsvProfilerModule);
 #if EXPERIMENTAL_STATSTRACE_ENABLED
diff --git a/Engine/Source/Developer/TraceServices/Public/TraceServices/Model/NetProfiler.h b/Engine/Source/Developer/TraceServices/Public/TraceServices/Model/NetProfiler.h
new file mode 100644
--- /dev/null
+++ b/Engine/Source/Developer/TraceServices/Public/TraceServices/Model/NetProfiler.h
@@ -0,0 +1,132 @@
+#pragma once
+
+#include "TraceServices/Model/AnalysisSession.h"
+
+namespace Trace
+{
+
+enum class ENetProfilerDeliveryStatus : uint8
+{
+	Unknown,
+	Delivered,
+	Dropped
+};
+
+enum ENetProfilerConnectionMode : uint8
+{
+	Outgoing = 0,
+	Incoming = 1,
+	Count
+};
+
+struct FNetProfilerName
+{
+	const TCHAR* Name;	// Name
+	uint32 NameIndex;	// Index in the type array?
+};
+
+typedef double FNetProfilerTimeStamp;
+
+struct FNetProfilerLifeTime
+{
+	FNetProfilerTimeStamp Begin = 0;
+	FNetProfilerTimeStamp End = 0;
+};
+
+//struct FNetProfilerProtocol
+//{
+//	const TCHAR* Name;
+//	uint32 ProtocolIndex;
+//	uint64 ProtocolIdentifier;
+//};
+
+struct FNetProfilerObjectInstance
+{
+	uint32 ObjectIndex;				// Index in the object array
+	uint16 NameIndex;				// Index in the Name array
+	uint32 TypeId;					// ProtocolIdentifier
+	uint32 NetId;					// NetHandleIndex or NetGUID
+	FNetProfilerLifeTime LifeTime;	// Lifetime of this instance
+};
+
+struct FNetProfilerContentEvent
+{
+	uint32 NameIndex ;			// identify the name / type, should we store the actual Name as well
+	uint32 ObjectInstanceIndex;	// object instance, non zero if this is a NetObject
+
+	uint64 StartPos : 16;		// Start position in the packet
+	uint64 EndPos : 16;			// End position in the packet
+	uint64 Level : 4;			// Level
+	uint64 ParentIndex : 28;	// Parent to be able to build a tree of nested events?
+};
+
+struct FNetProfilerPacket
+{
+	FNetProfilerTimeStamp TimeStamp;
+	uint32 SequenceNumber;
+	uint32 ContentSizeInBits;						// This is the part that is tracked by the PacketContents
+	uint32 TotalPacketSizeInBytes;					// This is the actual size of the packet sent on the socket
+	ENetProfilerDeliveryStatus DeliveryStatus;		// Indicates if the packet was delivered or not, updated as soon as we know
+
+	// Index into Events
+	uint32 StartEventIndex;
+	uint32 EventCount;
+};
+
+struct FNetProfilerConnection
+{
+	const TCHAR* Name;
+	FNetProfilerLifeTime LifeTime;
+	uint32 GameIntanceIndex;
+	uint32 ConnectionIndex : 16;
+	uint32 ConnectionId : 14;
+	uint32 bHasIncomingData : 1;
+	uint32 bHasOutgoingData : 1;
+};
+
+struct FNetProfilerGameInstance
+{
+	FNetProfilerLifeTime LifeTime;
+	uint32 GameInstanceIndex;
+	uint32 GameInstanceId;
+};
+
+// What queries do we need?
+class INetProfilerProvider : public IProvider
+{
+public:
+	// Access Names
+	virtual uint32 GetNameCount() const = 0;
+	virtual void ReadNames(TFunctionRef<void(const FNetProfilerName*, uint64)> Callback) const = 0;
+	virtual void ReadName(uint32 NameIndex, TFunctionRef<void(const FNetProfilerName&)> Callback) const = 0;
+
+	// Access GameInstances
+	virtual uint32 GetGameInstanceCount() const = 0;
+	virtual void ReadGameInstances(TFunctionRef<void(const FNetProfilerGameInstance&)> Callback) const = 0;
+
+	// Access Connections
+	virtual uint32 GetConnectionCount(uint32 GameInstanceIndex) const = 0;
+	virtual void ReadConnections(uint32 GameInstanceIndex, TFunctionRef<void(const FNetProfilerConnection&)> Callback) const = 0;
+	virtual void ReadConnection(uint32 ConnectionIndex, TFunctionRef<void(const FNetProfilerConnection&)> Callback) const = 0;
+	virtual uint32 GetConnectionChangeCount() const = 0;
+
+	// Access Object Instances
+	virtual uint32 GetObjectCount(uint32 GameInstanceIndex) const = 0;	
+	virtual void ReadObjects(uint32 GameInstanceIndex, TFunctionRef<void(const FNetProfilerObjectInstance&)> Callback) const = 0;
+	virtual void ReadObject(uint32 GameInstanceIndex, uint32 ObjectIndex, TFunctionRef<void(const FNetProfilerObjectInstance&)> Callback) const = 0;
+	virtual uint32 GeObjectsChangeCount(uint32 GameInstanceIndex) const = 0;
+
+	// Enumerate packets in the provided packet interval
+	virtual uint32 GetPacketCount(uint32 ConnectionIndex, ENetProfilerConnectionMode Mode) const = 0;
+	virtual void EnumeratePackets(uint32 ConnectionIndex, ENetProfilerConnectionMode Mode, uint32 PacketIndexIntervalStart, uint32 PacketIndexIntervalEnd, TFunctionRef<void(const FNetProfilerPacket&)> Callback) const = 0;
+	virtual uint32 GetPacketChangeCount(uint32 ConnectionIndex, ENetProfilerConnectionMode Mode) const = 0;
+
+	// Enumerate packet content events by range
+	virtual void EnumeratePacketContentEventsByIndex(uint32 ConnectionIndex, ENetProfilerConnectionMode Mode, uint32 StartEventIndex, uint32 EndEventIndex, TFunctionRef<void(const FNetProfilerContentEvent&)> Callback) const = 0;
+	virtual void EnumeratePacketContentEventsByPosition(uint32 ConnectionIndex, ENetProfilerConnectionMode Mode, uint32 PacketIndex, uint32 StartPosition, uint32 EndPosition, TFunctionRef<void(const FNetProfilerContentEvent&)> Callback) const = 0;
+	virtual uint32 GetPacketContentEventChangeCount(uint32 ConnectionIndex, ENetProfilerConnectionMode Mode) const = 0;
+};
+
+TRACESERVICES_API const INetProfilerProvider& ReadNetProfilerProvider(const IAnalysisSession& Session);
+
+}
