Fixed another nopch build error

#rnx

diff --git a/Engine/Source/Developer/TraceServices/Public/TraceServices/Model/NetProfiler.h b/Engine/Source/Developer/TraceServices/Public/TraceServices/Model/NetProfiler.h
--- a/Engine/Source/Developer/TraceServices/Public/TraceServices/Model/NetProfiler.h
+++ b/Engine/Source/Developer/TraceServices/Public/TraceServices/Model/NetProfiler.h
@@ -1,6 +1,7 @@
 #pragma once
 
 #include "TraceServices/Model/AnalysisSession.h"
+#include "Templates/Function.h"
 
 namespace Trace
 {
