Unreal Insights:
Added scope butterfly model to ITimingProfilerProvider

diff --git a/Engine/Source/Developer/TraceServices/Private/Model/TimingProfiler.cpp b/Engine/Source/Developer/TraceServices/Private/Model/TimingProfiler.cpp
--- a/Engine/Source/Developer/TraceServices/Private/Model/TimingProfiler.cpp
+++ b/Engine/Source/Developer/TraceServices/Private/Model/TimingProfiler.cpp
@@ -5,10 +5,49 @@
 #include "AnalysisServicePrivate.h"
 #include "Common/StringStore.h"
 #include "Common/TimelineStatistics.h"
+#include "Templates/TypeHash.h"
 
 namespace Trace
 {
 
+struct FTimingProfilerCallstackKey
+{
+	bool operator==(const FTimingProfilerCallstackKey& Other) const
+	{
+		return Other.TimerStack == TimerStack;
+	}
+
+	friend uint32 GetTypeHash(const FTimingProfilerCallstackKey& Key)
+	{
+		return Key.Hash;
+	}
+
+	TArray<uint32> TimerStack;
+	uint32 Hash;
+};
+
+class FTimingProfilerButterfly
+	: public ITimingProfilerButterfly
+{
+public:
+	FTimingProfilerButterfly();
+	virtual ~FTimingProfilerButterfly() = default;
+	virtual const FTimingProfilerButterflyNode& GenerateCallersTree(uint32 TimerId) override;
+	virtual const FTimingProfilerButterflyNode& GenerateCalleesTree(uint32 TimerId) override;
+
+private:
+	void CreateCallersTreeRecursive(const FTimingProfilerButterflyNode* TimerNode, const FTimingProfilerButterflyNode* RootNode, FTimingProfilerButterflyNode* OutputParent);
+	void CreateCalleesTreeRecursive(const FTimingProfilerButterflyNode* TimerNode, FTimingProfilerButterflyNode* OutputParent);
+
+	FSlabAllocator Allocator;
+	TPagedArray<FTimingProfilerButterflyNode> Nodes;
+	TArray<TArray<FTimingProfilerButterflyNode*>> TimerCallstacksMap;
+	TMap<uint32, FTimingProfilerButterflyNode*> CachedCallerTrees;
+	TMap<uint32, FTimingProfilerButterflyNode*> CachedCalleeTrees;
+
+	friend class FTimingProfilerProvider;
+};
+
 FTimingProfilerProvider::FTimingProfilerProvider(IAnalysisSession& InSession)
 	: Session(InSession)
 {
@@ -176,4 +215,213 @@ ITable<FTimingProfilerAggregatedStats>* FTimingProfilerProvider::CreateAggregati
 	return Table;
 }
 
+FTimingProfilerButterfly::FTimingProfilerButterfly()
+	: Allocator(2 << 20)
+	, Nodes(Allocator, 1024)
+{
+
+}
+
+void FTimingProfilerButterfly::CreateCallersTreeRecursive(const FTimingProfilerButterflyNode* TimerNode, const FTimingProfilerButterflyNode* RootNode, FTimingProfilerButterflyNode* OutputParent)
+{
+	if (!TimerNode)
+	{
+		return;
+	}
+	FTimingProfilerButterflyNode* AggregatedChildNode = nullptr;
+	for (FTimingProfilerButterflyNode* Candidate : OutputParent->Children)
+	{
+		if (Candidate->Timer == TimerNode->Timer)
+		{
+			AggregatedChildNode = Candidate;
+			break;
+		}
+	}
+	if (!AggregatedChildNode)
+	{
+		AggregatedChildNode = &Nodes.PushBack();
+		AggregatedChildNode->Timer = TimerNode->Timer;
+		OutputParent->Children.Add(AggregatedChildNode);
+		AggregatedChildNode->Parent = OutputParent;
+	}
+
+	AggregatedChildNode->InclusiveTime += RootNode->InclusiveTime;
+	AggregatedChildNode->ExclusiveTime += RootNode->ExclusiveTime;
+	AggregatedChildNode->Count += RootNode->Count;
+
+	CreateCallersTreeRecursive(TimerNode->Parent, RootNode, AggregatedChildNode);
+}
+
+const FTimingProfilerButterflyNode& FTimingProfilerButterfly::GenerateCallersTree(uint32 TimerId)
+{
+	FTimingProfilerButterflyNode** Cached = CachedCallerTrees.Find(TimerId);
+	if (Cached)
+	{
+		return **Cached;
+	}
+
+	FTimingProfilerButterflyNode* Root = &Nodes.PushBack();
+	for (FTimingProfilerButterflyNode* TimerNode : TimerCallstacksMap[TimerId])
+	{
+		Root->Timer = TimerNode->Timer;
+		Root->InclusiveTime += TimerNode->InclusiveTime;
+		Root->ExclusiveTime += TimerNode->ExclusiveTime;
+		Root->Count += TimerNode->Count;
+
+		CreateCallersTreeRecursive(TimerNode->Parent, TimerNode, Root);
+	}
+	CachedCallerTrees.Add(TimerId, Root);
+	return *Root;
+}
+
+void FTimingProfilerButterfly::CreateCalleesTreeRecursive(const FTimingProfilerButterflyNode* TimerNode, FTimingProfilerButterflyNode* OutputParent)
+{
+	for (const FTimingProfilerButterflyNode* ChildNode : TimerNode->Children)
+	{
+		FTimingProfilerButterflyNode* AggregatedChildNode = nullptr;
+		for (FTimingProfilerButterflyNode* Candidate : OutputParent->Children)
+		{
+			if (Candidate->Timer == ChildNode->Timer)
+			{
+				AggregatedChildNode = Candidate;
+				break;
+			}
+		}
+		if (!AggregatedChildNode)
+		{
+			AggregatedChildNode = &Nodes.PushBack();
+			AggregatedChildNode->Timer = ChildNode->Timer;
+			OutputParent->Children.Add(AggregatedChildNode);
+			AggregatedChildNode->Parent = OutputParent;
+		}
+		AggregatedChildNode->InclusiveTime += ChildNode->InclusiveTime;
+		AggregatedChildNode->ExclusiveTime += ChildNode->ExclusiveTime;
+		AggregatedChildNode->Count += ChildNode->Count;
+
+		CreateCalleesTreeRecursive(ChildNode, AggregatedChildNode);
+	}
+}
+
+const FTimingProfilerButterflyNode& FTimingProfilerButterfly::GenerateCalleesTree(uint32 TimerId)
+{
+	FTimingProfilerButterflyNode** Cached = CachedCalleeTrees.Find(TimerId);
+	if (Cached)
+	{
+		return **Cached;
+	}
+
+	FTimingProfilerButterflyNode* Root = &Nodes.PushBack();
+	for (FTimingProfilerButterflyNode* TimerNode : TimerCallstacksMap[TimerId])
+	{
+		Root->Timer = TimerNode->Timer;
+		Root->InclusiveTime += TimerNode->InclusiveTime;
+		Root->ExclusiveTime += TimerNode->ExclusiveTime;
+		Root->Count += TimerNode->Count;
+
+		CreateCalleesTreeRecursive(TimerNode, Root);
+	}
+	CachedCalleeTrees.Add(TimerId, Root);
+	return *Root;
+}
+
+ITimingProfilerButterfly* FTimingProfilerProvider::CreateButterfly(double IntervalStart, double IntervalEnd, TFunctionRef<bool(uint32)> CpuThreadFilter, bool IncludeGpu) const
+{
+	FTimingProfilerButterfly* Butterfly = new FTimingProfilerButterfly();
+	Butterfly->TimerCallstacksMap.AddDefaulted(Timers.Num());
+
+	TArray<const TimelineInternal*> IncludedTimelines;
+	if (IncludeGpu)
+	{
+		IncludedTimelines.Add(&Timelines[GpuTimelineIndex].Get());
+	}
+	for (const auto& KV : CpuThreadTimelineIndexMap)
+	{
+		if (CpuThreadFilter(KV.Key))
+		{
+			IncludedTimelines.Add(&Timelines[KV.Value].Get());
+		}
+	}
+
+	FTimingProfilerCallstackKey CurrentCallstackKey;
+	CurrentCallstackKey.TimerStack.Reserve(1024);
+
+	struct FLocalStackEntry
+	{
+		FTimingProfilerButterflyNode* Node;
+		double StartTime;
+		double ExclusiveTime;
+		uint32 CurrentCallstackHash;
+	};
+
+	TArray<FLocalStackEntry> CurrentCallstack;
+	CurrentCallstack.Reserve(1024);
+
+	TMap<FTimingProfilerCallstackKey, FTimingProfilerButterflyNode*> CallstackNodeMap;
+
+	double LastTime = 0.0;
+	for (const TimelineInternal* Timeline : IncludedTimelines)
+	{
+		Timeline->EnumerateEvents(IntervalStart, IntervalEnd, [this, IntervalStart, IntervalEnd, &CurrentCallstackKey, &CurrentCallstack, &CallstackNodeMap, &LastTime, Butterfly](bool IsEnter, double Time, const FTimingProfilerEvent& Event)
+		{
+			Time = FMath::Clamp(Time, IntervalStart, IntervalEnd);
+			FTimingProfilerButterflyNode* ParentNode = nullptr;
+			uint32 ParentCallstackHash = 17;
+			if (CurrentCallstack.Num())
+			{
+				FLocalStackEntry& Stackentry = CurrentCallstack.Top();
+				ParentNode = Stackentry.Node;
+				ParentCallstackHash = Stackentry.CurrentCallstackHash;
+				Stackentry.ExclusiveTime += Time - LastTime;
+			}
+			LastTime = Time;
+			if (IsEnter)
+			{
+				FLocalStackEntry& StackEntry = CurrentCallstack.AddDefaulted_GetRef();
+				StackEntry.StartTime = Time;
+				StackEntry.ExclusiveTime = 0.0;
+				StackEntry.CurrentCallstackHash = ParentCallstackHash * 17 + Event.TimerIndex;
+
+				CurrentCallstackKey.TimerStack.Push(Event.TimerIndex);
+				CurrentCallstackKey.Hash = StackEntry.CurrentCallstackHash;
+
+				FTimingProfilerButterflyNode** FindIt = CallstackNodeMap.Find(CurrentCallstackKey);
+				if (FindIt)
+				{
+					StackEntry.Node = *FindIt;
+				}
+				else
+				{
+					StackEntry.Node = &Butterfly->Nodes.PushBack();
+					CallstackNodeMap.Add(CurrentCallstackKey, StackEntry.Node);
+					Butterfly->TimerCallstacksMap[Event.TimerIndex].Add(StackEntry.Node);
+					StackEntry.Node->InclusiveTime = 0.0;
+					StackEntry.Node->ExclusiveTime = 0.0;
+					StackEntry.Node->Count = 0;
+					StackEntry.Node->Timer = &this->Timers[Event.TimerIndex];
+					StackEntry.Node->Parent = ParentNode;
+					if (ParentNode)
+					{
+						ParentNode->Children.Add(StackEntry.Node);
+					}
+				}
+			}
+			else
+			{
+				FLocalStackEntry& StackEntry = CurrentCallstack.Top();
+				double InclusiveTime = Time - StackEntry.StartTime;
+				check(InclusiveTime >= 0.0);
+				check(StackEntry.ExclusiveTime >= 0.0 && StackEntry.ExclusiveTime <= InclusiveTime);
+				
+				StackEntry.Node->InclusiveTime += InclusiveTime;
+				StackEntry.Node->ExclusiveTime += StackEntry.ExclusiveTime;
+				++StackEntry.Node->Count;
+
+				CurrentCallstack.Pop(false);
+				CurrentCallstackKey.TimerStack.Pop(false);
+			}
+		});
+	}
+	return Butterfly;
+}
+
 }
diff --git a/Engine/Source/Developer/TraceServices/Private/Model/TimingProfilerPrivate.h b/Engine/Source/Developer/TraceServices/Private/Model/TimingProfilerPrivate.h
--- a/Engine/Source/Developer/TraceServices/Private/Model/TimingProfilerPrivate.h
+++ b/Engine/Source/Developer/TraceServices/Private/Model/TimingProfilerPrivate.h
@@ -33,6 +33,7 @@ public:
 	virtual void EnumerateTimelines(TFunctionRef<void(const Timeline&)> Callback) const override;
 	virtual void ReadTimers(TFunctionRef<void(const FTimingProfilerTimer*, uint64)> Callback) const override;
 	virtual ITable<FTimingProfilerAggregatedStats>* CreateAggregation(double IntervalStart, double IntervalEnd, TFunctionRef<bool(uint32)> CpuThreadFilter, bool IncludeGpu) const override;
+	virtual ITimingProfilerButterfly* CreateButterfly(double IntervalStart, double IntervalEnd, TFunctionRef<bool(uint32)> CpuThreadFilter, bool IncludeGpu) const override;
 
 private:
 	FTimingProfilerTimer& AddTimerInternal(const TCHAR* Name, bool IsGpuEvent);
diff --git a/Engine/Source/Developer/TraceServices/Public/TraceServices/Model/TimingProfiler.h b/Engine/Source/Developer/TraceServices/Public/TraceServices/Model/TimingProfiler.h
--- a/Engine/Source/Developer/TraceServices/Public/TraceServices/Model/TimingProfiler.h
+++ b/Engine/Source/Developer/TraceServices/Public/TraceServices/Model/TimingProfiler.h
@@ -5,6 +5,7 @@
 #include "TraceServices/Model/AnalysisSession.h"
 #include "TraceServices/Containers/Timelines.h"
 #include "TraceServices/Containers/Tables.h"
+#include "Containers/Array.h"
 
 namespace Trace
 {
@@ -28,6 +29,23 @@ struct FTimingProfilerAggregatedStats
 	const FTimingProfilerTimer* Timer = nullptr;
 };
 
+struct FTimingProfilerButterflyNode
+{
+	const FTimingProfilerTimer* Timer = nullptr;
+	uint64 Count = 0;
+	double InclusiveTime = 0.0;
+	double ExclusiveTime = 0.0;
+	const FTimingProfilerButterflyNode* Parent = nullptr;
+	TArray<FTimingProfilerButterflyNode*> Children;
+};
+
+class ITimingProfilerButterfly
+{
+public:
+	virtual const FTimingProfilerButterflyNode& GenerateCallersTree(uint32 TimerId) = 0;
+	virtual const FTimingProfilerButterflyNode& GenerateCalleesTree(uint32 TimerId) = 0;
+};
+
 class ITimingProfilerProvider
 	: public IProvider
 {
@@ -42,6 +60,7 @@ public:
 	virtual void EnumerateTimelines(TFunctionRef<void(const Timeline&)> Callback) const = 0;
 	virtual void ReadTimers(TFunctionRef<void(const FTimingProfilerTimer*, uint64)> Callback) const = 0;
 	virtual ITable<FTimingProfilerAggregatedStats>* CreateAggregation(double IntervalStart, double IntervalEnd, TFunctionRef<bool(uint32)> CpuThreadFilter, bool IncludeGpu) const = 0;
+	virtual ITimingProfilerButterfly* CreateButterfly(double IntervalStart, double IntervalEnd, TFunctionRef<bool(uint32)> CpuThreadFilter, bool IncludeGpu) const = 0;
 };
 
 TRACESERVICES_API const ITimingProfilerProvider* ReadTimingProfilerProvider(const IAnalysisSession& Session);
