Add missing include

diff --git a/Engine/Source/Developer/TraceServices/Private/Model/NetProfilerProvider.cpp b/Engine/Source/Developer/TraceServices/Private/Model/NetProfilerProvider.cpp
--- a/Engine/Source/Developer/TraceServices/Private/Model/NetProfilerProvider.cpp
+++ b/Engine/Source/Developer/TraceServices/Private/Model/NetProfilerProvider.cpp
@@ -2,6 +2,7 @@
 
 #include "NetProfilerProvider.h"
 #include "AnalysisServicePrivate.h"
+#include "Containers/ArrayView.h"
 #include "Common/StringStore.h"
 
 namespace Trace
