UnrealInsights:
Fixed handling of ANSICHAR strings in log trace

diff --git a/Engine/Source/Developer/TraceServices/Private/Analyzers/LoadTimeTraceAnalysis.cpp b/Engine/Source/Developer/TraceServices/Private/Analyzers/LoadTimeTraceAnalysis.cpp
--- a/Engine/Source/Developer/TraceServices/Private/Analyzers/LoadTimeTraceAnalysis.cpp
+++ b/Engine/Source/Developer/TraceServices/Private/Analyzers/LoadTimeTraceAnalysis.cpp
@@ -259,7 +259,7 @@ bool FAsyncLoadingTraceAnalyzer::OnEvent(uint16 RouteId, const FOnEventContext&
 		TSharedRef<FRequestGroupState> GroupState = MakeShared<FRequestGroupState>();
 		const TCHAR* FormatString = reinterpret_cast<const TCHAR*>(EventData.GetAttachment());
 		const uint8* FormatArgs = EventData.GetAttachment() + (FCString::Strlen(FormatString) + 1) * sizeof(TCHAR);
-		Trace::FFormatArgsHelper::Format(FormatBuffer, FormatBufferSize - 1, FormatString, FormatArgs);
+		Trace::FFormatArgsHelper::Format(FormatBuffer, FormatBufferSize - 1, TempBuffer, FormatBufferSize - 1, FormatString, FormatArgs);
 		GroupState->Name = Session.StoreString(FormatBuffer);
 		uint32 ThreadId = EventData.GetValue<uint32>("ThreadId");
 		TSharedRef<FThreadState> ThreadState = GetThreadState(ThreadId);
diff --git a/Engine/Source/Developer/TraceServices/Private/Analyzers/LoadTimeTraceAnalysis.h b/Engine/Source/Developer/TraceServices/Private/Analyzers/LoadTimeTraceAnalysis.h
--- a/Engine/Source/Developer/TraceServices/Private/Analyzers/LoadTimeTraceAnalysis.h
+++ b/Engine/Source/Developer/TraceServices/Private/Analyzers/LoadTimeTraceAnalysis.h
@@ -121,6 +121,7 @@ private:
 		FormatBufferSize = 65536
 	};
 	TCHAR FormatBuffer[FormatBufferSize];
+	TCHAR TempBuffer[FormatBufferSize];
 
 	Trace::IAnalysisSession& Session;
 	Trace::FLoadTimeProfilerProvider& LoadTimeProfilerProvider;
diff --git a/Engine/Source/Developer/TraceServices/Private/Common/FormatArgs.cpp b/Engine/Source/Developer/TraceServices/Private/Common/FormatArgs.cpp
--- a/Engine/Source/Developer/TraceServices/Private/Common/FormatArgs.cpp
+++ b/Engine/Source/Developer/TraceServices/Private/Common/FormatArgs.cpp
@@ -185,10 +185,28 @@ bool FFormatArgsHelper::AdvanceArgumentStream(FFormatArgsStreamContext& Context)
 	}
 	if (Context.ArgumentTypeCategory == FFormatArgsTrace::FormatArgTypeCode_CategoryString)
 	{
-		check(Context.ArgumentTypeSize == 2);
-		const TCHAR* StringPtr = reinterpret_cast<const TCHAR*>(Context.PayloadPtr);
-		while (*StringPtr++);
-		Context.PayloadPtr = reinterpret_cast<const uint8*>(StringPtr);
+		if (Context.ArgumentTypeSize == 1)
+		{
+			const uint8* StringPtr = reinterpret_cast<const uint8*>(Context.PayloadPtr);
+			while (*StringPtr++);
+			Context.PayloadPtr = StringPtr;
+		}
+		else if (Context.ArgumentTypeSize == 2)
+		{
+			const uint16* StringPtr = reinterpret_cast<const uint16*>(Context.PayloadPtr);
+			while (*StringPtr++);
+			Context.PayloadPtr = reinterpret_cast<const uint8*>(StringPtr);
+		}
+		else if (Context.ArgumentTypeSize == 4)
+		{
+			const uint32* StringPtr = reinterpret_cast<const uint32*>(Context.PayloadPtr);
+			while (*StringPtr++);
+			Context.PayloadPtr = reinterpret_cast<const uint8*>(StringPtr);
+		}
+		else
+		{
+			check(false);
+		}
 	}
 	else
 	{
@@ -231,20 +249,34 @@ double FFormatArgsHelper::ExtractFloatingPointArgument(FFormatArgsStreamContext&
 	return Result;
 }
 
-const TCHAR* FFormatArgsHelper::ExtractStringArgument(FFormatArgsStreamContext& ArgStream)
+const TCHAR* FFormatArgsHelper::ExtractStringArgument(FFormatArgsStreamContext& ArgStream, TCHAR* Temp, uint64 MaxTemp)
 {
 	static TCHAR Empty[] = TEXT("");
 	const TCHAR* Result = Empty;
 	if (ArgStream.ArgumentTypeCategory == FFormatArgsTrace::FormatArgTypeCode_CategoryString)
 	{
-		check(ArgStream.ArgumentTypeSize == 2);
-		Result = reinterpret_cast<const TCHAR*>(ArgStream.PayloadPtr);
+		if (ArgStream.ArgumentTypeSize == sizeof(TCHAR))
+		{
+			Result = reinterpret_cast<const TCHAR*>(ArgStream.PayloadPtr);
+		}
+		else if (ArgStream.ArgumentTypeSize == sizeof(ANSICHAR))
+		{
+			const ANSICHAR* AnsiString = reinterpret_cast<const ANSICHAR*>(ArgStream.PayloadPtr);
+			int32 SourceLength = TCString<ANSICHAR>::Strlen(AnsiString) + 1;
+			TStringConvert<ANSICHAR, TCHAR> StringConvert;
+			int32 ConvertedLength = StringConvert.ConvertedLength(AnsiString, SourceLength);
+			if (ConvertedLength < MaxTemp)
+			{
+				StringConvert.Convert(Temp, MaxTemp, AnsiString, SourceLength);
+				Result = Temp;
+			}
+		}
 	}
 	AdvanceArgumentStream(ArgStream);
 	return Result;
 }
 
-int32 FFormatArgsHelper::FormatArgument(TCHAR* Out, uint64 MaxOut, const FFormatArgSpec& ArgSpec, FFormatArgsStreamContext& ArgStream)
+int32 FFormatArgsHelper::FormatArgument(TCHAR* Out, uint64 MaxOut, TCHAR* Temp, uint64 MaxTemp, const FFormatArgSpec& ArgSpec, FFormatArgsStreamContext& ArgStream)
 {
 	check(ArgSpec.AdditionalIntegerArgumentCount <= 2);
 	switch (ArgSpec.ExpectedTypeCategory)
@@ -291,15 +323,15 @@ int32 FFormatArgsHelper::FormatArgument(TCHAR* Out, uint64 MaxOut, const FFormat
 	case FFormatArgsTrace::FormatArgTypeCode_CategoryString:
 		if (ArgSpec.AdditionalIntegerArgumentCount == 2)
 		{
-			return FCString::Snprintf(Out, MaxOut, ArgSpec.FormatString, ExtractIntegerArgument(ArgStream), ExtractIntegerArgument(ArgStream), ExtractStringArgument(ArgStream));
+			return FCString::Snprintf(Out, MaxOut, ArgSpec.FormatString, ExtractIntegerArgument(ArgStream), ExtractIntegerArgument(ArgStream), ExtractStringArgument(ArgStream, Temp, MaxTemp));
 		}
 		else if (ArgSpec.AdditionalIntegerArgumentCount == 1)
 		{
-			return FCString::Snprintf(Out, MaxOut, ArgSpec.FormatString, ExtractIntegerArgument(ArgStream), ExtractStringArgument(ArgStream));
+			return FCString::Snprintf(Out, MaxOut, ArgSpec.FormatString, ExtractIntegerArgument(ArgStream), ExtractStringArgument(ArgStream, Temp, MaxTemp));
 		}
 		else
 		{
-			return FCString::Snprintf(Out, MaxOut, ArgSpec.FormatString, ExtractStringArgument(ArgStream));
+			return FCString::Snprintf(Out, MaxOut, ArgSpec.FormatString, ExtractStringArgument(ArgStream, Temp, MaxTemp));
 		}
 		break;
 	default:
@@ -308,7 +340,7 @@ int32 FFormatArgsHelper::FormatArgument(TCHAR* Out, uint64 MaxOut, const FFormat
 	}
 }
 
-void FFormatArgsHelper::Format(TCHAR* Out, uint64 MaxOut, const TCHAR* FormatString, const uint8* FormatArgs)
+void FFormatArgsHelper::Format(TCHAR* Out, uint64 MaxOut, TCHAR* Temp, uint64 MaxTemp, const TCHAR* FormatString, const uint8* FormatArgs)
 {
 	FFormatArgsStreamContext ArgumentStream;
 	InitArgumentStream(ArgumentStream, FormatArgs);
@@ -333,7 +365,7 @@ void FFormatArgsHelper::Format(TCHAR* Out, uint64 MaxOut, const TCHAR* FormatStr
 		}
 		if (Spec.Valid)
 		{
-			int32 Length = FormatArgument(Dst, DstEnd - Dst, Spec, ArgumentStream);
+			int32 Length = FormatArgument(Dst, DstEnd - Dst, Temp, MaxTemp, Spec, ArgumentStream);
 			if (Length < 0)
 			{
 				break;
diff --git a/Engine/Source/Developer/TraceServices/Private/Common/FormatArgs.h b/Engine/Source/Developer/TraceServices/Private/Common/FormatArgs.h
--- a/Engine/Source/Developer/TraceServices/Private/Common/FormatArgs.h
+++ b/Engine/Source/Developer/TraceServices/Private/Common/FormatArgs.h
@@ -9,7 +9,7 @@ namespace Trace
 
 struct FFormatArgsHelper
 {
-	static void Format(TCHAR* Out, uint64 MaxOut, const TCHAR* FormatString, const uint8* FormatArgs);
+	static void Format(TCHAR* Out, uint64 MaxOut, TCHAR* Temp, uint64 MaxTemp, const TCHAR* FormatString, const uint8* FormatArgs);
 
 private:
 	struct FFormatArgSpec
@@ -36,8 +36,8 @@ private:
 	static bool AdvanceArgumentStream(FFormatArgsStreamContext& Context);
 	static uint64 ExtractIntegerArgument(FFormatArgsStreamContext& ArgStream);
 	static double ExtractFloatingPointArgument(FFormatArgsStreamContext& ArgStream);
-	static const TCHAR* ExtractStringArgument(FFormatArgsStreamContext& ArgStream);
-	static int32 FormatArgument(TCHAR* Out, uint64 MaxOut, const FFormatArgSpec& ArgSpec, FFormatArgsStreamContext& ArgStream);
+	static const TCHAR* ExtractStringArgument(FFormatArgsStreamContext& ArgStream, TCHAR* Temp, uint64 MaxTemp);
+	static int32 FormatArgument(TCHAR* Out, uint64 MaxOut, TCHAR* Temp, uint64 MaxTemp, const FFormatArgSpec& ArgSpec, FFormatArgsStreamContext& ArgStream);
 };
 
 }
\ No newline at end of file
diff --git a/Engine/Source/Developer/TraceServices/Private/Model/Bookmarks.cpp b/Engine/Source/Developer/TraceServices/Private/Model/Bookmarks.cpp
--- a/Engine/Source/Developer/TraceServices/Private/Model/Bookmarks.cpp
+++ b/Engine/Source/Developer/TraceServices/Private/Model/Bookmarks.cpp
@@ -37,7 +37,7 @@ void FBookmarkProvider::AppendBookmark(double Time, uint64 BookmarkPoint, const
 	FBookmarkSpec& Spec = GetSpec(BookmarkPoint);
 	TSharedRef<FBookmarkInternal> Bookmark = MakeShared<FBookmarkInternal>();
 	Bookmark->Time = Time;
-	FFormatArgsHelper::Format(FormatBuffer, FormatBufferSize - 1, Spec.FormatString, FormatArgs);
+	FFormatArgsHelper::Format(FormatBuffer, FormatBufferSize - 1, TempBuffer, FormatBufferSize - 1, Spec.FormatString, FormatArgs);
 	Bookmark->Text = Session.StoreString(FormatBuffer);
 	Bookmarks.Add(Bookmark);
 	Session.UpdateDurationSeconds(Time);
diff --git a/Engine/Source/Developer/TraceServices/Private/Model/BookmarksPrivate.h b/Engine/Source/Developer/TraceServices/Private/Model/BookmarksPrivate.h
--- a/Engine/Source/Developer/TraceServices/Private/Model/BookmarksPrivate.h
+++ b/Engine/Source/Developer/TraceServices/Private/Model/BookmarksPrivate.h
@@ -47,6 +47,7 @@ private:
 	TMap<uint64, TSharedPtr<FBookmarkSpec>> SpecMap;
 	TArray<TSharedRef<FBookmarkInternal>> Bookmarks;
 	TCHAR FormatBuffer[FormatBufferSize];
+	TCHAR TempBuffer[FormatBufferSize];
 };
 
 }
\ No newline at end of file
diff --git a/Engine/Source/Developer/TraceServices/Private/Model/Log.cpp b/Engine/Source/Developer/TraceServices/Private/Model/Log.cpp
--- a/Engine/Source/Developer/TraceServices/Private/Model/Log.cpp
+++ b/Engine/Source/Developer/TraceServices/Private/Model/Log.cpp
@@ -59,7 +59,7 @@ void FLogProvider::AppendMessage(uint64 LogPoint, double Time, const uint8* Form
 		FLogMessageInternal& InternalMessage = Messages.PushBack();
 		InternalMessage.Time = Time;
 		InternalMessage.Spec = *FindSpec;
-		FFormatArgsHelper::Format(FormatBuffer, FormatBufferSize - 1, InternalMessage.Spec->FormatString, FormatArgs);
+		FFormatArgsHelper::Format(FormatBuffer, FormatBufferSize - 1, TempBuffer, FormatBufferSize - 1, InternalMessage.Spec->FormatString, FormatArgs);
 		InternalMessage.Message = Session.StoreString(FormatBuffer);
 		Session.UpdateDurationSeconds(Time);
 	}
diff --git a/Engine/Source/Developer/TraceServices/Private/Model/LogPrivate.h b/Engine/Source/Developer/TraceServices/Private/Model/LogPrivate.h
--- a/Engine/Source/Developer/TraceServices/Private/Model/LogPrivate.h
+++ b/Engine/Source/Developer/TraceServices/Private/Model/LogPrivate.h
@@ -79,6 +79,7 @@ private:
 	TPagedArray<FLogMessageSpec> MessageSpecs;
 	TPagedArray<FLogMessageInternal> Messages;
 	TCHAR FormatBuffer[FormatBufferSize];
+	TCHAR TempBuffer[FormatBufferSize];
 	TTableView<FMessagesTableLayout> MessagesTable;
 };
 
