[Insights]
  - Networking: Fixed bad indexing when analyzing destroyed instance

diff --git a/Engine/Source/Developer/TraceServices/Private/Analyzers/NetTraceAnalyzer.cpp b/Engine/Source/Developer/TraceServices/Private/Analyzers/NetTraceAnalyzer.cpp
--- a/Engine/Source/Developer/TraceServices/Private/Analyzers/NetTraceAnalyzer.cpp
+++ b/Engine/Source/Developer/TraceServices/Private/Analyzers/NetTraceAnalyzer.cpp
@@ -316,7 +316,7 @@ bool FNetTraceAnalyzer::OnEvent(uint16 RouteId, const FOnEventContext& Context)
 			FNetTraceActiveObjectState DestroyedObjectState;
 			if (GameInstanceState->ActiveObjects.RemoveAndCopyValue(ObjectId, DestroyedObjectState))
 			{
-				if (Trace::FNetProfilerObjectInstance* ObjectInstance = NetProfilerProvider.EditObject(GameInstanceId, DestroyedObjectState.ObjectIndex))
+				if (Trace::FNetProfilerObjectInstance* ObjectInstance = NetProfilerProvider.EditObject(GameInstanceState->GameInstanceIndex, DestroyedObjectState.ObjectIndex))
 				{
 					// Update object data
 					ObjectInstance->LifeTime.End = GetLastTimestamp();
