Emit a trace event with some contextual information about the session being traced.

#rnx

diff --git a/Engine/Source/Developer/TraceServices/Private/SessionService.cpp b/Engine/Source/Developer/TraceServices/Private/SessionService.cpp
--- a/Engine/Source/Developer/TraceServices/Private/SessionService.cpp
+++ b/Engine/Source/Developer/TraceServices/Private/SessionService.cpp
@@ -6,6 +6,8 @@
 #include "Misc/Paths.h"
 #include "Misc/OutputDeviceRedirector.h"
 #include "SocketSubsystem.h"
+#include "Trace/Analysis.h"
+#include "Trace/Analyzer.h"
 #include "Trace/ControlClient.h"
 #include "IPAddress.h"
 #include "AddressInfoTypes.h"
@@ -19,6 +21,49 @@
 namespace Trace
 {
 
+struct FDiagnosticsSessionAnalyzer
+	: public Trace::IAnalyzer
+{
+	virtual void OnAnalysisBegin(const FOnAnalysisContext& Context) override
+	{
+		Context.InterfaceBuilder.RouteEvent(0, "Diagnostics", "Session");
+	}
+
+	virtual bool OnEvent(uint16, const FOnEventContext& Context) override
+	{
+		const FEventData& EventData = Context.EventData;
+		const uint8* Attachment = EventData.GetAttachment();
+		if (Attachment == nullptr)
+		{
+			return false;
+		}
+
+		uint8 CommandLineOffset = EventData.GetValue<uint8>("CommandLineOffset");
+		uint8 AppNameOffset = EventData.GetValue<uint8>("AppNameOffset");
+
+		Platform = FString(CommandLineOffset, (const ANSICHAR*)Attachment);
+
+		Attachment += CommandLineOffset;
+		int32 CommandLineLength = AppNameOffset - CommandLineOffset;
+		CommandLine = FString(CommandLineLength, (const ANSICHAR*)Attachment);
+
+		Attachment += CommandLineLength;
+		int32 AppNameLength = EventData.GetAttachmentSize() - AppNameOffset;
+		AppName = FString(AppNameLength, (const ANSICHAR*)Attachment);
+
+		ConfigurationType = EventData.GetValue<int8>("ConfigurationType");
+		TargetType = EventData.GetValue<int8>("TargetType");
+
+		return false;
+	}
+
+	FString Platform;
+	FString AppName;
+	FString CommandLine;
+	int8 ConfigurationType = 0;
+	int8 TargetType = 0;
+};
+
 FSessionService::FSessionService(FModuleService& InModuleService)
 	: ModuleService(InModuleService)
 {
@@ -106,6 +151,11 @@ bool FSessionService::GetSessionInfo(FSessionHandle SessionHandle, FSessionInfo&
 	OutSessionInfo.TimeStamp = FindIt->TimeStamp;
 	OutSessionInfo.Size = FindIt->Size;
 	OutSessionInfo.bIsLive = FindIt->bIsLive;
+	OutSessionInfo.Platform = *FindIt->Platform;
+	OutSessionInfo.AppName = *FindIt->AppName;
+	OutSessionInfo.CommandLine = *FindIt->CommandLine;
+	OutSessionInfo.ConfigurationType = EBuildConfiguration(FindIt->ConfigurationType);
+	OutSessionInfo.TargetType = EBuildTargetType(FindIt->TargetType);
 	return true;
 }
 
@@ -228,6 +278,53 @@ bool FSessionService::Tick(float DeltaTime)
 	return true;
 }
 
+void FSessionService::UpdateSessionContext(FStoreSessionHandle StoreHandle, FSessionInfoInternal& Info)
+{
+	if (Info.Platform.Len())
+	{
+		return;
+	}
+
+	IInDataStream* Stream = TraceStore->OpenSessionStream(StoreHandle);
+
+	struct FDataStream
+		: public IInDataStream
+	{
+		virtual int32 Read(void* Data, uint32 Size) override
+		{
+			int32 InnerBytesRead = Inner->Read(Data, Size);
+			BytesRead += InnerBytesRead;
+			return (BytesRead < (2 << 20)) ? InnerBytesRead : 0;
+		}
+
+		int32 BytesRead = 0;
+		IInDataStream* Inner;
+	};
+
+	FDataStream DataStream;
+	DataStream.Inner = Stream;
+
+	FDiagnosticsSessionAnalyzer Analyzer;
+	FAnalysisContext Context;
+	Context.AddAnalyzer(Analyzer);
+	Context.Process(DataStream).Wait();
+
+	delete Stream;
+
+	if (Analyzer.Platform.Len() == 0)
+	{
+		Info.Platform = TEXT("Unknown");
+	}
+	else
+	{
+		Info.Platform = MoveTemp(Analyzer.Platform);
+		Info.AppName = MoveTemp(Analyzer.AppName);
+		Info.CommandLine = MoveTemp(Analyzer.CommandLine);
+		Info.ConfigurationType = Analyzer.ConfigurationType;
+		Info.TargetType = Analyzer.TargetType;
+	}
+}
+
 void FSessionService::UpdateSessions()
 {
 	TArray<FStoreSessionInfo> StoreSessions;
@@ -267,6 +364,11 @@ void FSessionService::UpdateSessions()
 			FindIt->RecorderSessionHandle = RecorderSession.Handle;
 		}
 	}
+
+	for (auto& KeyValue : Sessions)
+	{
+		UpdateSessionContext(KeyValue.Key, KeyValue.Value);
+	}
 }
 
 }
diff --git a/Engine/Source/Developer/TraceServices/Private/SessionServicePrivate.h b/Engine/Source/Developer/TraceServices/Private/SessionServicePrivate.h
--- a/Engine/Source/Developer/TraceServices/Private/SessionServicePrivate.h
+++ b/Engine/Source/Developer/TraceServices/Private/SessionServicePrivate.h
@@ -32,20 +32,26 @@ public:
 	virtual bool ConnectSession(const TCHAR* ControlClientAddress) override;
 
 private:
-	bool Tick(float DeltaTime);
-	void UpdateSessions();
-
 	struct FSessionInfoInternal
 	{
 		FRecorderSessionHandle RecorderSessionHandle;
 		const TCHAR* Uri;
 		const TCHAR* Name;
+		FString Platform;
+		FString AppName;
+		FString CommandLine;
 		TMap<FName, TSet<FString>> EnabledModuleLoggersMap;
 		FDateTime TimeStamp;
 		uint64 Size;
 		bool bIsLive;
+		int8 ConfigurationType;
+		int8 TargetType;
 	};
 
+	bool Tick(float DeltaTime);
+	void UpdateSessions();
+	void UpdateSessionContext(FStoreSessionHandle StoreHandle, FSessionInfoInternal& Info);
+
 	FModuleService& ModuleService;
 	void* RecorderEvent = nullptr;
 	FString LocalSessionDirectory;
diff --git a/Engine/Source/Developer/TraceServices/Public/TraceServices/SessionService.h b/Engine/Source/Developer/TraceServices/Public/TraceServices/SessionService.h
--- a/Engine/Source/Developer/TraceServices/Public/TraceServices/SessionService.h
+++ b/Engine/Source/Developer/TraceServices/Public/TraceServices/SessionService.h
@@ -3,6 +3,7 @@
 #pragma once
 
 #include "CoreMinimal.h"
+#include "HAL/PlatformMisc.h"
 #include "Templates/SharedPointer.h"
 #include "Trace/Store.h"
 #include "Trace/DataStream.h"
@@ -16,9 +17,14 @@ struct FSessionInfo
 {
 	const TCHAR* Uri = nullptr;
 	const TCHAR* Name = nullptr;
+	const TCHAR* Platform = nullptr;
+	const TCHAR* AppName = nullptr;
+	const TCHAR* CommandLine = nullptr;
 	FDateTime TimeStamp;
 	uint64 Size = 0;
 	bool bIsLive = false;
+	EBuildConfiguration ConfigurationType = EBuildConfiguration::Unknown;
+	EBuildTargetType TargetType = EBuildTargetType::Unknown;
 };
 
 class ISessionService
diff --git a/Engine/Source/Runtime/Launch/Private/LaunchEngineLoop.cpp b/Engine/Source/Runtime/Launch/Private/LaunchEngineLoop.cpp
--- a/Engine/Source/Runtime/Launch/Private/LaunchEngineLoop.cpp
+++ b/Engine/Source/Runtime/Launch/Private/LaunchEngineLoop.cpp
@@ -1209,6 +1209,42 @@ int32 FEngineLoop::PreInit(const TCHAR* CmdLine)
 
 	SCOPED_BOOT_TIMING("FEngineLoop::PreInit");
 
+	// Trace out information about this session
+	{
+		uint8 Payload[1024];
+		int32 PayloadSize = 0;
+
+		auto AddToPayload = [&] (const TCHAR* String) -> uint8
+		{
+			int32 Length = FCString::Strlen(String);
+			Length = FMath::Min<int32>(Length, sizeof(Payload) - PayloadSize - 1);
+			for (int32 i = 0, n = Length; i < n; ++i)
+			{
+				Payload[PayloadSize] = uint8(String[i] & 0x7f);
+				++PayloadSize;
+			}
+			return uint8(PayloadSize - Length);
+		};
+
+		AddToPayload(FGenericPlatformMisc::GetUBTPlatform());
+		uint8 CommandLineOffset = AddToPayload(CmdLine);
+		uint8 AppNameOffset = AddToPayload(TEXT(UE_APP_NAME));
+
+		UE_TRACE_EVENT_BEGIN(Diagnostics, Session, Important|Always)
+			UE_TRACE_EVENT_FIELD(uint8, CommandLineOffset)
+			UE_TRACE_EVENT_FIELD(uint8, AppNameOffset)
+			UE_TRACE_EVENT_FIELD(uint8, ConfigurationType)
+			UE_TRACE_EVENT_FIELD(uint8, TargetType)
+		UE_TRACE_EVENT_END()
+
+		UE_TRACE_LOG(Diagnostics, Session, PayloadSize)
+			<< Session.CommandLineOffset(CommandLineOffset)
+			<< Session.AppNameOffset(AppNameOffset)
+			<< Session.ConfigurationType(uint8(FApp::GetBuildConfiguration()))
+			<< Session.TargetType(uint8(FApp::GetBuildTargetType()))
+			<< Session.Attachment(Payload, PayloadSize);
+	}
+
 #if PLATFORM_WINDOWS
 	// Register a handler for Ctrl-C so we've effective signal handling from the outset.
 	FWindowsPlatformMisc::SetGracefulTerminationHandler();
