[Insights] Fixed parsing for AppName and CommandLine. + Minor adjustment of columns width in Start Page.

diff --git a/Engine/Source/Developer/TraceInsights/Private/Insights/Widgets/SStartPageWindow.cpp b/Engine/Source/Developer/TraceInsights/Private/Insights/Widgets/SStartPageWindow.cpp
--- a/Engine/Source/Developer/TraceInsights/Private/Insights/Widgets/SStartPageWindow.cpp
+++ b/Engine/Source/Developer/TraceInsights/Private/Insights/Widgets/SStartPageWindow.cpp
@@ -522,7 +522,7 @@ TSharedRef<SWidget> SStartPageWindow::ConstructSessionsPanel()
 				SNew(SHeaderRow)
 
 				+ SHeaderRow::Column(FName(TEXT("Name")))
-				.FillWidth(0.4f)
+				.FillWidth(0.25f)
 				.DefaultLabel(LOCTEXT("NameColumn", "Name"))
 
 				+ SHeaderRow::Column(FName(TEXT("Platform")))
diff --git a/Engine/Source/Developer/TraceServices/Private/SessionService.cpp b/Engine/Source/Developer/TraceServices/Private/SessionService.cpp
--- a/Engine/Source/Developer/TraceServices/Private/SessionService.cpp
+++ b/Engine/Source/Developer/TraceServices/Private/SessionService.cpp
@@ -38,19 +38,19 @@ struct FDiagnosticsSessionAnalyzer
 			return false;
 		}
 
-		uint8 CommandLineOffset = EventData.GetValue<uint8>("CommandLineOffset");
 		uint8 AppNameOffset = EventData.GetValue<uint8>("AppNameOffset");
+		uint8 CommandLineOffset = EventData.GetValue<uint8>("CommandLineOffset");
 
-		Platform = FString(CommandLineOffset, (const ANSICHAR*)Attachment);
-
-		Attachment += CommandLineOffset;
-		int32 CommandLineLength = AppNameOffset - CommandLineOffset;
-		CommandLine = FString(CommandLineLength, (const ANSICHAR*)Attachment);
+		Platform = FString(AppNameOffset, (const ANSICHAR*)Attachment);
 
-		Attachment += CommandLineLength;
-		int32 AppNameLength = EventData.GetAttachmentSize() - AppNameOffset;
+		Attachment += AppNameOffset;
+		int32 AppNameLength = CommandLineOffset - AppNameOffset;
 		AppName = FString(AppNameLength, (const ANSICHAR*)Attachment);
 
+		Attachment += AppNameLength;
+		int32 CommandLineLength = EventData.GetAttachmentSize() - CommandLineOffset;
+		CommandLine = FString(CommandLineLength, (const ANSICHAR*)Attachment);
+
 		ConfigurationType = EventData.GetValue<int8>("ConfigurationType");
 		TargetType = EventData.GetValue<int8>("TargetType");
 
