[Insights] Enabled UE_TRACE_TABLE_COLUMN macros also for Mac. Fixes UE-80839.

diff --git a/Engine/Source/Developer/TraceServices/Private/Model/Tables.h b/Engine/Source/Developer/TraceServices/Private/Model/Tables.h
--- a/Engine/Source/Developer/TraceServices/Private/Model/Tables.h
+++ b/Engine/Source/Developer/TraceServices/Private/Model/Tables.h
@@ -23,7 +23,7 @@
 		::Trace::FColumnValueContainer GetColumnValueInternal(const RowType& Row) const { return 0; } \
 		::Trace::TColumnDeclaration<0, RowType
 
-#if PLATFORM_WINDOWS
+#if PLATFORM_WINDOWS || PLATFORM_MAC
 
 #define UE_TRACE_TABLE_COLUMN(Name, DisplayName) \
 		> PREPROCESSOR_JOIN(__Column__, __LINE__); \
