Expose setting up custom session directory in Trace services.

diff --git a/Engine/Source/Developer/TraceServices/Private/SessionService.cpp b/Engine/Source/Developer/TraceServices/Private/SessionService.cpp
--- a/Engine/Source/Developer/TraceServices/Private/SessionService.cpp
+++ b/Engine/Source/Developer/TraceServices/Private/SessionService.cpp
@@ -65,9 +65,22 @@ struct FDiagnosticsSessionAnalyzer
 };
 
 FSessionService::FSessionService(FModuleService& InModuleService)
+	: FSessionService(InModuleService, nullptr)
+{
+
+}
+
+FSessionService::FSessionService(FModuleService& InModuleService, const TCHAR* OverrideSessionDirectory)
 	: ModuleService(InModuleService)
 {
-	LocalSessionDirectory = FPaths::ProjectSavedDir() / TEXT("TraceSessions");
+	if (OverrideSessionDirectory)
+	{
+		LocalSessionDirectory = FString(OverrideSessionDirectory);
+	}
+	else
+	{
+		LocalSessionDirectory = FPaths::ProjectSavedDir() / TEXT("TraceSessions");
+	}
 	TraceStore = Store_Create(*LocalSessionDirectory);
 	TraceRecorder = Recorder_Create(TraceStore.ToSharedRef());
 
diff --git a/Engine/Source/Developer/TraceServices/Private/SessionServicePrivate.h b/Engine/Source/Developer/TraceServices/Private/SessionServicePrivate.h
--- a/Engine/Source/Developer/TraceServices/Private/SessionServicePrivate.h
+++ b/Engine/Source/Developer/TraceServices/Private/SessionServicePrivate.h
@@ -17,6 +17,7 @@ class FSessionService
 {
 public:
 	FSessionService(FModuleService& ModuleService);
+	FSessionService(FModuleService& ModuleService, const TCHAR* OverrideSessionsDirectory);
 	virtual ~FSessionService();
 	virtual bool StartRecorderServer() override;
 	virtual bool IsRecorderServerRunning() const override;
diff --git a/Engine/Source/Developer/TraceServices/Private/TraceServicesModule.cpp b/Engine/Source/Developer/TraceServices/Private/TraceServicesModule.cpp
--- a/Engine/Source/Developer/TraceServices/Private/TraceServicesModule.cpp
+++ b/Engine/Source/Developer/TraceServices/Private/TraceServicesModule.cpp
@@ -20,6 +20,9 @@ public:
 	virtual TSharedPtr<Trace::ISessionService> GetSessionService() override;
 	virtual TSharedPtr<Trace::IAnalysisService> GetAnalysisService() override;
 	virtual TSharedPtr<Trace::IModuleService> GetModuleService() override;
+	virtual TSharedPtr<Trace::ISessionService> CreateSessionService(const TCHAR*) override;
+	virtual TSharedPtr<Trace::IAnalysisService> CreateAnalysisService() override;
+	virtual TSharedPtr<Trace::IModuleService> CreateModuleService() override;
 
 	virtual void StartupModule() override;
 	virtual void ShutdownModule() override;
@@ -66,6 +69,29 @@ TSharedPtr<Trace::IModuleService> FTraceServicesModule::GetModuleService()
 	return ModuleService;
 }
 
+TSharedPtr<Trace::ISessionService> FTraceServicesModule::CreateSessionService(const TCHAR* SessionDirectory)
+{
+	checkf(!SessionService.IsValid(), TEXT("A SessionService already exists."));
+	GetModuleService();
+	SessionService = MakeShared<Trace::FSessionService>(*ModuleService.Get(), SessionDirectory);
+	return SessionService;
+}
+
+TSharedPtr<Trace::IAnalysisService> FTraceServicesModule::CreateAnalysisService()
+{
+	checkf(!AnalysisService.IsValid(), TEXT("A AnalysisService already exists."));
+	GetModuleService();
+	AnalysisService = MakeShared<Trace::FAnalysisService>(*ModuleService.Get());
+	return AnalysisService;
+}
+
+TSharedPtr<Trace::IModuleService> FTraceServicesModule::CreateModuleService()
+{
+	checkf(!ModuleService.IsValid(), TEXT("A ModuleService already exists."));
+	ModuleService = MakeShared<Trace::FModuleService>();
+	return ModuleService;
+}
+
 void FTraceServicesModule::StartupModule()
 {
 	IModularFeatures::Get().RegisterModularFeature(Trace::ModuleFeatureName, &TimingProfilerModule);
@@ -88,6 +114,10 @@ void FTraceServicesModule::ShutdownModule()
 #endif
 	IModularFeatures::Get().UnregisterModularFeature(Trace::ModuleFeatureName, &LoadTimeProfilerModule);
 	IModularFeatures::Get().UnregisterModularFeature(Trace::ModuleFeatureName, &TimingProfilerModule);
+
+	SessionService.Reset();
+	AnalysisService.Reset();
+	ModuleService.Reset();
 }
 
 IMPLEMENT_MODULE(FTraceServicesModule, TraceServices)
diff --git a/Engine/Source/Developer/TraceServices/Public/TraceServices/ITraceServicesModule.h b/Engine/Source/Developer/TraceServices/Public/TraceServices/ITraceServicesModule.h
--- a/Engine/Source/Developer/TraceServices/Public/TraceServices/ITraceServicesModule.h
+++ b/Engine/Source/Developer/TraceServices/Public/TraceServices/ITraceServicesModule.h
@@ -22,6 +22,9 @@ public:
 	virtual TSharedPtr<Trace::ISessionService> GetSessionService() = 0;
 	virtual TSharedPtr<Trace::IAnalysisService> GetAnalysisService() = 0;
 	virtual TSharedPtr<Trace::IModuleService> GetModuleService() = 0;
+	virtual TSharedPtr<Trace::ISessionService> CreateSessionService(const TCHAR* SessionDirectory) = 0;
+	virtual TSharedPtr<Trace::IAnalysisService> CreateAnalysisService() = 0;
+	virtual TSharedPtr<Trace::IModuleService> CreateModuleService() = 0;
 
 	virtual ~ITraceServicesModule() = default;
 };
