UnrealInsights:
Fixed inclusive/exclusive timer stat mixup

diff --git a/Engine/Source/Developer/TraceServices/Private/Model/TimingProfiler.cpp b/Engine/Source/Developer/TraceServices/Private/Model/TimingProfiler.cpp
--- a/Engine/Source/Developer/TraceServices/Private/Model/TimingProfiler.cpp
+++ b/Engine/Source/Developer/TraceServices/Private/Model/TimingProfiler.cpp
@@ -230,10 +230,10 @@ ITable<FTimingProfilerAggregatedStats>* FTimingProfilerProvider::CreateAggregati
 		Row.Timer = KV.Key;
 		const FAggregatedTimingStats& Stats = KV.Value;
 		Row.InstanceCount = Stats.InstanceCount;
-		Row.TotalInclusiveTime = Stats.TotalExclusiveTime;
-		Row.MinInclusiveTime = Stats.MinExclusiveTime;
-		Row.MaxInclusiveTime = Stats.MaxExclusiveTime;
-		Row.AverageInclusiveTime = Stats.AverageExclusiveTime;
+		Row.TotalInclusiveTime = Stats.TotalInclusiveTime;
+		Row.MinInclusiveTime = Stats.MinInclusiveTime;
+		Row.MaxInclusiveTime = Stats.MaxInclusiveTime;
+		Row.AverageInclusiveTime = Stats.AverageInclusiveTime;
 		Row.MedianInclusiveTime = Stats.MedianInclusiveTime;
 		Row.TotalExclusiveTime = Stats.TotalExclusiveTime;
 		Row.MinExclusiveTime = Stats.MinExclusiveTime;
