CI fixes

diff --git a/Engine/Source/Developer/TraceServices/Public/Model/IntervalTimeline.h b/Engine/Source/Developer/TraceServices/Public/Model/IntervalTimeline.h
--- a/Engine/Source/Developer/TraceServices/Public/Model/IntervalTimeline.h
+++ b/Engine/Source/Developer/TraceServices/Public/Model/IntervalTimeline.h
@@ -34,9 +34,9 @@ public:
 	virtual const EventType& GetEvent(uint64 InIndex) const override { return Events[InIndex].Event; }
 	virtual double GetStartTime() const override { return Events.Num() > 0 ? Events[0].StartTime : 0.0; }
 	virtual double GetEndTime() const override { return Events.Num() > 0 ? Events[Events.Num() - 1].EndTime : 0.0; }
-	virtual void EnumerateEventsDownSampled(double IntervalStart, double IntervalEnd, double Resolution, ITimeline::EventCallback Callback) const override { check(false); }
-	virtual void EnumerateEventsDownSampled(double IntervalStart, double IntervalEnd, double Resolution, ITimeline::EventRangeCallback Callback) const override { check(false); }
-	virtual void EnumerateEvents(double IntervalStart, double IntervalEnd, ITimeline::EventCallback Callback) const override
+	virtual void EnumerateEventsDownSampled(double IntervalStart, double IntervalEnd, double Resolution, typename ITimeline<EventType>::EventCallback Callback) const override { check(false); }
+	virtual void EnumerateEventsDownSampled(double IntervalStart, double IntervalEnd, double Resolution, typename ITimeline<EventType>::EventRangeCallback Callback) const override { check(false); }
+	virtual void EnumerateEvents(double IntervalStart, double IntervalEnd, typename  ITimeline<EventType>::EventCallback Callback) const override
 	{
 		EnumerateEvents(IntervalStart, IntervalEnd, [&Callback](double InStartTime, double InEndTime, uint32 InDepth, const EventType& InEvent)
 		{
@@ -45,7 +45,7 @@ public:
 		});
 	}
 	
-	virtual void EnumerateEvents(double IntervalStart, double IntervalEnd, ITimeline::EventRangeCallback Callback) const override
+	virtual void EnumerateEvents(double IntervalStart, double IntervalEnd, typename ITimeline<EventType>::EventRangeCallback Callback) const override
 	{
 		auto EventIterator = Events.GetIteratorFromPage(0);
 		const FEventInternal* Event = EventIterator.GetCurrentItem();
diff --git a/Engine/Source/Developer/TraceServices/Public/Model/MonotonicTimeline.h b/Engine/Source/Developer/TraceServices/Public/Model/MonotonicTimeline.h
--- a/Engine/Source/Developer/TraceServices/Public/Model/MonotonicTimeline.h
+++ b/Engine/Source/Developer/TraceServices/Public/Model/MonotonicTimeline.h
@@ -140,7 +140,7 @@ public:
 		return NumScopeEntries > 0 ? FMath::Abs(DetailLevels[0].ScopeEntries[NumScopeEntries - 1].Time) : 0.0;
 	}
 
-	virtual void EnumerateEventsDownSampled(double IntervalStart, double IntervalEnd, double Resolution, ITimeline::EventCallback Callback) const override
+	virtual void EnumerateEventsDownSampled(double IntervalStart, double IntervalEnd, double Resolution, typename ITimeline<EventType>::EventCallback Callback) const override
 	{
 		int32 DetailLevelIndex = SettingsType::DetailLevelsCount - 1;
 		for (; DetailLevelIndex > 0; --DetailLevelIndex)
@@ -265,7 +265,7 @@ public:
 		}
 	}
 
-	virtual void EnumerateEventsDownSampled(double IntervalStart, double IntervalEnd, double Resolution, ITimeline::EventRangeCallback Callback) const override
+	virtual void EnumerateEventsDownSampled(double IntervalStart, double IntervalEnd, double Resolution, typename ITimeline<EventType>::EventRangeCallback Callback) const override
 	{
 		struct FStackEntry
 		{
@@ -315,12 +315,12 @@ public:
 		});
 	}
 
-	virtual void EnumerateEvents(double IntervalStart, double IntervalEnd, ITimeline::EventCallback Callback) const override
+	virtual void EnumerateEvents(double IntervalStart, double IntervalEnd, typename ITimeline<EventType>::EventCallback Callback) const override
 	{
 		EnumerateEventsDownSampled(IntervalStart, IntervalEnd, 0.0, Callback);
 	}
 
-	virtual void EnumerateEvents(double IntervalStart, double IntervalEnd, ITimeline::EventRangeCallback Callback) const override
+	virtual void EnumerateEvents(double IntervalStart, double IntervalEnd, typename ITimeline<EventType>::EventRangeCallback Callback) const override
 	{
 		EnumerateEventsDownSampled(IntervalStart, IntervalEnd, 0.0, Callback);
 	}
diff --git a/Engine/Source/Developer/TraceServices/Public/Model/PointTimeline.h b/Engine/Source/Developer/TraceServices/Public/Model/PointTimeline.h
--- a/Engine/Source/Developer/TraceServices/Public/Model/PointTimeline.h
+++ b/Engine/Source/Developer/TraceServices/Public/Model/PointTimeline.h
@@ -34,9 +34,9 @@ public:
 	virtual const EventType& GetEvent(uint64 InIndex) const override { return Events[InIndex].Event; }
 	virtual double GetStartTime() const override { return Events.Num() > 0 ? Events[0].Time : 0.0; }
 	virtual double GetEndTime() const override { return Events.Num() > 0 ? Events[Events.Num() - 1].Time : 0.0; }
-	virtual void EnumerateEventsDownSampled(double IntervalStart, double IntervalEnd, double Resolution, ITimeline::EventCallback Callback) const override { check(false); }
-	virtual void EnumerateEventsDownSampled(double IntervalStart, double IntervalEnd, double Resolution, ITimeline::EventRangeCallback Callback) const override { check(false); }
-	virtual void EnumerateEvents(double IntervalStart, double IntervalEnd, ITimeline::EventCallback Callback) const override
+	virtual void EnumerateEventsDownSampled(double IntervalStart, double IntervalEnd, double Resolution, typename ITimeline<EventType>::EventCallback Callback) const override { check(false); }
+	virtual void EnumerateEventsDownSampled(double IntervalStart, double IntervalEnd, double Resolution, typename ITimeline<EventType>::EventRangeCallback Callback) const override { check(false); }
+	virtual void EnumerateEvents(double IntervalStart, double IntervalEnd, typename ITimeline<EventType>::EventCallback Callback) const override
 	{
 		EnumerateEvents(IntervalStart, IntervalEnd, [&Callback](double InStartTime, double InEndTime, uint32 InDepth, const EventType& InEvent)
 		{
@@ -45,7 +45,7 @@ public:
 		});
 	}
 	
-	virtual void EnumerateEvents(double IntervalStart, double IntervalEnd, ITimeline::EventRangeCallback Callback) const override
+	virtual void EnumerateEvents(double IntervalStart, double IntervalEnd, typename ITimeline<EventType>::EventRangeCallback Callback) const override
 	{
 		auto EventIterator = Events.GetIteratorFromPage(0);
 		const FEventInternal* Event = EventIterator.GetCurrentItem();
