UnrealInsights:
Cleanup of load time profiling instrumentation, mainly removed tracking and reporting of internal EDL events
Added support for having multiple async loading threads
Added I/O dispatcher and AsyncLoading throughput graphs

diff --git a/Engine/Source/Developer/TraceInsights/Private/Insights/ViewModels/LoadingTimingTrack.cpp b/Engine/Source/Developer/TraceInsights/Private/Insights/ViewModels/LoadingTimingTrack.cpp
--- a/Engine/Source/Developer/TraceInsights/Private/Insights/ViewModels/LoadingTimingTrack.cpp
+++ b/Engine/Source/Developer/TraceInsights/Private/Insights/ViewModels/LoadingTimingTrack.cpp
@@ -29,53 +29,9 @@ void FLoadingSharedState::Reset()
 
 ////////////////////////////////////////////////////////////////////////////////////////////////////
 
-const TCHAR* GetName(ELoadTimeProfilerPackageEventType Type)
-{
-	switch (Type)
-	{
-	case LoadTimeProfilerPackageEventType_CreateLinker:				return TEXT("CreateLinker");
-	case LoadTimeProfilerPackageEventType_FinishLinker:				return TEXT("FinishLinker");
-	case LoadTimeProfilerPackageEventType_StartImportPackages:		return TEXT("StartImportPackages");
-	case LoadTimeProfilerPackageEventType_SetupImports:				return TEXT("SetupImports");
-	case LoadTimeProfilerPackageEventType_SetupExports:				return TEXT("SetupExports");
-	case LoadTimeProfilerPackageEventType_ProcessImportsAndExports:	return TEXT("ProcessImportsAndExports");
-	case LoadTimeProfilerPackageEventType_ExportsDone:				return TEXT("ExportsDone");
-	case LoadTimeProfilerPackageEventType_PostLoadWait:				return TEXT("PostLoadWait");
-	case LoadTimeProfilerPackageEventType_StartPostLoad:			return TEXT("StartPostLoad");
-	case LoadTimeProfilerPackageEventType_Tick:						return TEXT("Tick");
-	case LoadTimeProfilerPackageEventType_Finish:					return TEXT("Finish");
-	case LoadTimeProfilerPackageEventType_DeferredPostLoad:			return TEXT("DeferredPostLoad");
-	case LoadTimeProfilerPackageEventType_None:						return TEXT("None");
-	default:														return TEXT("");
-	}
-}
-
-////////////////////////////////////////////////////////////////////////////////////////////////////
-
-const TCHAR* GetName(ELoadTimeProfilerObjectEventType Type)
-{
-	switch (Type)
-	{
-	case LoadTimeProfilerObjectEventType_Create:	return TEXT("Create");
-	case LoadTimeProfilerObjectEventType_Serialize:	return TEXT("Serialize");
-	case LoadTimeProfilerObjectEventType_PostLoad:	return TEXT("PostLoad");
-	case LoadTimeProfilerObjectEventType_None:		return TEXT("None");
-	default:										return TEXT("");
-	};
-}
-
-////////////////////////////////////////////////////////////////////////////////////////////////////
-
-const TCHAR* FLoadingSharedState::GetEventNameByPackageEventType(uint32 Depth, const Trace::FLoadTimeProfilerCpuEvent& Event) const
-{
-	return GetName(Event.PackageEventType);
-}
-
-////////////////////////////////////////////////////////////////////////////////////////////////////
-
 const TCHAR* FLoadingSharedState::GetEventNameByExportEventType(uint32 Depth, const Trace::FLoadTimeProfilerCpuEvent& Event) const
 {
-	return GetName(Event.ExportEventType);
+	return Trace::GetLoadTimeProfilerObjectEventTypeString(Event.EventType);
 }
 
 ////////////////////////////////////////////////////////////////////////////////////////////////////
@@ -125,11 +81,10 @@ void FLoadingSharedState::SetColorSchema(int32 Schema)
 {
 	switch (Schema)
 	{
-		case 0: LoadingGetEventNameFn = FLoadingTrackGetEventNameDelegate::CreateRaw(this, &FLoadingSharedState::GetEventNameByPackageEventType); break;
-		case 1: LoadingGetEventNameFn = FLoadingTrackGetEventNameDelegate::CreateRaw(this, &FLoadingSharedState::GetEventNameByExportEventType); break;
-		case 2: LoadingGetEventNameFn = FLoadingTrackGetEventNameDelegate::CreateRaw(this, &FLoadingSharedState::GetEventNameByPackageName); break;
-		case 3: LoadingGetEventNameFn = FLoadingTrackGetEventNameDelegate::CreateRaw(this, &FLoadingSharedState::GetEventNameByExportClassName); break;
-		case 4: LoadingGetEventNameFn = FLoadingTrackGetEventNameDelegate::CreateRaw(this, &FLoadingSharedState::GetEventName); break;
+		case 0: LoadingGetEventNameFn = FLoadingTrackGetEventNameDelegate::CreateRaw(this, &FLoadingSharedState::GetEventNameByExportEventType); break;
+		case 1: LoadingGetEventNameFn = FLoadingTrackGetEventNameDelegate::CreateRaw(this, &FLoadingSharedState::GetEventNameByPackageName); break;
+		case 2: LoadingGetEventNameFn = FLoadingTrackGetEventNameDelegate::CreateRaw(this, &FLoadingSharedState::GetEventNameByExportClassName); break;
+		case 3: LoadingGetEventNameFn = FLoadingTrackGetEventNameDelegate::CreateRaw(this, &FLoadingSharedState::GetEventName); break;
 	};
 }
 
@@ -148,21 +103,19 @@ void FLoadingTimingTrack::InitTooltip(FTooltipDrawState& Tooltip, const FTimingE
 
 	Tooltip.AddNameValueTextLine(TEXT("Duration:"), TimeUtils::FormatTimeAuto(HoveredTimingEvent.Duration()));
 	Tooltip.AddNameValueTextLine(TEXT("Depth:"), FString::Printf(TEXT("%d"), HoveredTimingEvent.Depth));
-	Tooltip.AddNameValueTextLine(TEXT("Package Event:"), ::GetName(HoveredTimingEvent.LoadingInfo.PackageEventType));
 
 	if (Package)
 	{
 		Tooltip.AddNameValueTextLine(TEXT("Package Name:"), Package->Name);
 		Tooltip.AddNameValueTextLine(TEXT("Header Size:"), FString::Printf(TEXT("%s bytes"), *FText::AsNumber(Package->Summary.TotalHeaderSize).ToString()));
-		Tooltip.AddNameValueTextLine(TEXT("Package Summary:"), FString::Printf(TEXT("%d names, %d imports, %d exports"), Package->Summary.NameCount, Package->Summary.ImportCount, Package->Summary.ExportCount));
+		Tooltip.AddNameValueTextLine(TEXT("Package Summary:"), FString::Printf(TEXT("%d imports, %d exports"), Package->Summary.ImportCount, Package->Summary.ExportCount));
 	}
 
-	Tooltip.AddNameValueTextLine(TEXT("Export Event:"), FString::Printf(TEXT("%s%s"), ::GetName(HoveredTimingEvent.LoadingInfo.ExportEventType), Export && Export->IsAsset ? TEXT(" [asset]") : TEXT("")));
+	Tooltip.AddNameValueTextLine(TEXT("Export Event:"), FString::Printf(TEXT("%s"), Trace::GetLoadTimeProfilerObjectEventTypeString(HoveredTimingEvent.LoadingInfo.EventType)));
 
 	if (Export)
 	{
 		Tooltip.AddNameValueTextLine(TEXT("Export Class:"), Export->Class ? Export->Class->Name : TEXT("N/A"));
-		Tooltip.AddNameValueTextLine(TEXT("Serial Offset:"), FString::Printf(TEXT("%s bytes"), *FText::AsNumber(Export->SerialOffset).ToString()));
 		Tooltip.AddNameValueTextLine(TEXT("Serial Size:"), FString::Printf(TEXT("%s bytes"), *FText::AsNumber(Export->SerialSize).ToString()));
 	}
 
@@ -224,20 +177,20 @@ struct FSearchLoadTimeProfilerCpuEvent
 // FLoadingMainThreadTimingTrack
 ////////////////////////////////////////////////////////////////////////////////////////////////////
 
-void FLoadingMainThreadTimingTrack::Draw(ITimingViewDrawHelper& Helper) const
+void FLoadingTimingTrack::Draw(ITimingViewDrawHelper& Helper) const
 {
-	FTimingEventsTrack& Track = *const_cast<FLoadingMainThreadTimingTrack*>(this);
+	FTimingEventsTrack& Track = *const_cast<FLoadingTimingTrack*>(this);
 
 	if (Helper.BeginTimeline(Track))
 	{
 		TSharedPtr<const Trace::IAnalysisSession> Session = FInsightsManager::Get()->GetSession();
-		if (Session.IsValid() && Trace::ReadLoadTimeProfilerProvider(*Session.Get()))
+		if (Trace::ReadLoadTimeProfilerProvider(*Session.Get()))
 		{
 			Trace::FAnalysisSessionReadScope SessionReadScope(*Session.Get());
 
 			const Trace::ILoadTimeProfilerProvider& LoadTimeProfilerProvider = *Trace::ReadLoadTimeProfilerProvider(*Session.Get());
 
-			LoadTimeProfilerProvider.ReadMainThreadCpuTimeline([this, &Helper](const Trace::ILoadTimeProfilerProvider::CpuTimeline& Timeline)
+			LoadTimeProfilerProvider.ReadTimeline(TimelineIndex, [this, &Helper](const Trace::ILoadTimeProfilerProvider::CpuTimeline& Timeline)
 			{
 				if (FTimingEventsTrack::bUseDownSampling)
 				{
@@ -265,112 +218,31 @@ void FLoadingMainThreadTimingTrack::Draw(ITimingViewDrawHelper& Helper) const
 
 ////////////////////////////////////////////////////////////////////////////////////////////////////
 
-bool FLoadingMainThreadTimingTrack::SearchTimingEvent(const double InStartTime,
-	const double InEndTime,
-	TFunctionRef<bool(double, double, uint32)> InPredicate,
-	FTimingEvent& InOutTimingEvent,
-	bool bInStopAtFirstMatch,
-	bool bInSearchForLargestEvent) const
+bool FLoadingTimingTrack::SearchTimingEvent(const double InStartTime,
+													  const double InEndTime,
+													  TFunctionRef<bool(double, double, uint32)> InPredicate,
+													  FTimingEvent& InOutTimingEvent,
+													  bool bInStopAtFirstMatch,
+													  bool bInSearchForLargestEvent) const
 {
 	FSearchLoadTimeProfilerCpuEvent Ctx(InStartTime, InEndTime, InPredicate, InOutTimingEvent, bInStopAtFirstMatch, bInSearchForLargestEvent);
 
 	TSharedPtr<const Trace::IAnalysisSession> Session = FInsightsManager::Get()->GetSession();
-	if (Session.IsValid())
-	{
-		Trace::FAnalysisSessionReadScope SessionReadScope(*Session.Get());
-
-		const FTimingEventsTrack* Track = Ctx.TimingEvent.Track;
-
-		if (Trace::ReadLoadTimeProfilerProvider(*Session.Get()))
-		{
-			const Trace::ILoadTimeProfilerProvider& LoadTimeProfilerProvider = *Trace::ReadLoadTimeProfilerProvider(*Session.Get());
-
-			LoadTimeProfilerProvider.ReadMainThreadCpuTimeline([&Ctx](const Trace::ILoadTimeProfilerProvider::CpuTimeline& Timeline)
-			{
-				Timeline.EnumerateEvents(Ctx.StartTime, Ctx.EndTime, [&Ctx](double EventStartTime, double EventEndTime, uint32 EventDepth, const Trace::FLoadTimeProfilerCpuEvent& Event)
-				{
-					Ctx.CheckEvent(EventStartTime, EventEndTime, EventDepth, Event);
-				});
-			});
-		}
-	}
-
-	return Ctx.bFound;
-}
-
-////////////////////////////////////////////////////////////////////////////////////////////////////
-// FLoadingAsyncThreadTimingTrack
-////////////////////////////////////////////////////////////////////////////////////////////////////
-
-void FLoadingAsyncThreadTimingTrack::Draw(ITimingViewDrawHelper& Helper) const
-{
-	FTimingEventsTrack& Track = *const_cast<FLoadingAsyncThreadTimingTrack*>(this);
-
-	if (Helper.BeginTimeline(Track))
-	{
-		TSharedPtr<const Trace::IAnalysisSession> Session = FInsightsManager::Get()->GetSession();
-		if (Session.IsValid() && Trace::ReadLoadTimeProfilerProvider(*Session.Get()))
-		{
-			Trace::FAnalysisSessionReadScope SessionReadScope(*Session.Get());
-
-			const Trace::ILoadTimeProfilerProvider& LoadTimeProfilerProvider = *Trace::ReadLoadTimeProfilerProvider(*Session.Get());
-
-			LoadTimeProfilerProvider.ReadAsyncLoadingThreadCpuTimeline([this, &Helper](const Trace::ILoadTimeProfilerProvider::CpuTimeline& Timeline)
-			{
-				if (FTimingEventsTrack::bUseDownSampling)
-				{
-					const double SecondsPerPixel = 1.0 / Helper.GetViewport().GetScaleX();
-					Timeline.EnumerateEventsDownSampled(Helper.GetViewport().GetStartTime(), Helper.GetViewport().GetEndTime(), SecondsPerPixel, [this, &Helper](double StartTime, double EndTime, uint32 Depth, const Trace::FLoadTimeProfilerCpuEvent& Event)
-					{
-						const TCHAR* Name = State->GetEventNameEx(Depth, Event);
-						Helper.AddEvent(StartTime, EndTime, Depth, Name);
-					});
-				}
-				else
-				{
-					Timeline.EnumerateEvents(Helper.GetViewport().GetStartTime(), Helper.GetViewport().GetEndTime(), [this, &Helper](double StartTime, double EndTime, uint32 Depth, const Trace::FLoadTimeProfilerCpuEvent& Event)
-					{
-						const TCHAR* Name = State->GetEventNameEx(Depth, Event);
-						Helper.AddEvent(StartTime, EndTime, Depth, Name);
-					});
-				}
-			});
-		}
+	Trace::FAnalysisSessionReadScope SessionReadScope(*Session.Get());
 
-		Helper.EndTimeline(Track);
-	}
-}
-
-////////////////////////////////////////////////////////////////////////////////////////////////////
+	const FTimingEventsTrack* Track = Ctx.TimingEvent.Track;
 
-bool FLoadingAsyncThreadTimingTrack::SearchTimingEvent(const double InStartTime,
-													   const double InEndTime,
-													   TFunctionRef<bool(double, double, uint32)> InPredicate,
-													   FTimingEvent& InOutTimingEvent,
-													   bool bInStopAtFirstMatch,
-													   bool bInSearchForLargestEvent) const
-{
-	FSearchLoadTimeProfilerCpuEvent Ctx(InStartTime, InEndTime, InPredicate, InOutTimingEvent, bInStopAtFirstMatch, bInSearchForLargestEvent);
-
-	TSharedPtr<const Trace::IAnalysisSession> Session = FInsightsManager::Get()->GetSession();
-	if (Session.IsValid())
+	if (Trace::ReadLoadTimeProfilerProvider(*Session.Get()))
 	{
-		Trace::FAnalysisSessionReadScope SessionReadScope(*Session.Get());
-
-		const FTimingEventsTrack* Track = Ctx.TimingEvent.Track;
+		const Trace::ILoadTimeProfilerProvider& LoadTimeProfilerProvider = *Trace::ReadLoadTimeProfilerProvider(*Session.Get());
 
-		if (Trace::ReadLoadTimeProfilerProvider(*Session.Get()))
+		LoadTimeProfilerProvider.ReadTimeline(TimelineIndex, [&Ctx](const Trace::ILoadTimeProfilerProvider::CpuTimeline& Timeline)
 		{
-			const Trace::ILoadTimeProfilerProvider& LoadTimeProfilerProvider = *Trace::ReadLoadTimeProfilerProvider(*Session.Get());
-
-			LoadTimeProfilerProvider.ReadAsyncLoadingThreadCpuTimeline([&Ctx](const Trace::ILoadTimeProfilerProvider::CpuTimeline& Timeline)
+			Timeline.EnumerateEvents(Ctx.StartTime, Ctx.EndTime, [&Ctx](double EventStartTime, double EventEndTime, uint32 EventDepth, const Trace::FLoadTimeProfilerCpuEvent& Event)
 			{
-				Timeline.EnumerateEvents(Ctx.StartTime, Ctx.EndTime, [&Ctx](double EventStartTime, double EventEndTime, uint32 EventDepth, const Trace::FLoadTimeProfilerCpuEvent& Event)
-				{
-					Ctx.CheckEvent(EventStartTime, EventEndTime, EventDepth, Event);
-				});
+				Ctx.CheckEvent(EventStartTime, EventEndTime, EventDepth, Event);
 			});
-		}
+		});
 	}
 
 	return Ctx.bFound;
diff --git a/Engine/Source/Developer/TraceInsights/Private/Insights/ViewModels/LoadingTimingTrack.h b/Engine/Source/Developer/TraceInsights/Private/Insights/ViewModels/LoadingTimingTrack.h
--- a/Engine/Source/Developer/TraceInsights/Private/Insights/ViewModels/LoadingTimingTrack.h
+++ b/Engine/Source/Developer/TraceInsights/Private/Insights/ViewModels/LoadingTimingTrack.h
@@ -18,7 +18,6 @@ class FLoadingSharedState
 public:
 	void Reset();
 
-	const TCHAR* GetEventNameByPackageEventType(uint32 Depth, const Trace::FLoadTimeProfilerCpuEvent& Event) const;
 	const TCHAR* GetEventNameByExportEventType(uint32 Depth, const Trace::FLoadTimeProfilerCpuEvent& Event) const;
 	const TCHAR* GetEventNameByPackageName(uint32 Depth, const Trace::FLoadTimeProfilerCpuEvent& Event) const;
 	const TCHAR* GetEventNameByExportClassName(uint32 Depth, const Trace::FLoadTimeProfilerCpuEvent& Event) const;
@@ -37,45 +36,22 @@ private:
 class FLoadingTimingTrack : public FTimingEventsTrack
 {
 public:
-	explicit FLoadingTimingTrack(uint64 InTrackId, const FName& InSubType, const FString& InName, TSharedPtr<FLoadingSharedState> InState)
+	explicit FLoadingTimingTrack(uint64 InTrackId, uint32 InTimelineIndex, const FName& InSubType, const FString& InName, TSharedPtr<FLoadingSharedState> InState)
 		: FTimingEventsTrack(InTrackId, FName(TEXT("Loading")), InSubType, InName)
 		, State(InState)
+		, TimelineIndex(InTimelineIndex)
 	{
 	}
 	virtual ~FLoadingTimingTrack() {}
 
 	virtual void InitTooltip(FTooltipDrawState& Tooltip, const FTimingEvent& HoveredTimingEvent) const override;
 
-protected:
-	TSharedPtr<FLoadingSharedState> State;
-};
-
-////////////////////////////////////////////////////////////////////////////////////////////////////
-
-class FLoadingMainThreadTimingTrack : public FLoadingTimingTrack
-{
-public:
-	explicit FLoadingMainThreadTimingTrack(uint64 InTrackId, TSharedPtr<FLoadingSharedState> InState)
-		: FLoadingTimingTrack(InTrackId, FName(TEXT("MainThread")), TEXT("Loading - Main Thread"), InState)
-	{
-	}
-
 	virtual void Draw(ITimingViewDrawHelper& Helper) const override;
 	virtual bool SearchTimingEvent(const double InStartTime, const double InEndTime, TFunctionRef<bool(double, double, uint32)> InPredicate, FTimingEvent& InOutTimingEvent, bool bInStopAtFirstMatch, bool bInSearchForLargestEvent) const override;
-};
 
-////////////////////////////////////////////////////////////////////////////////////////////////////
-
-class FLoadingAsyncThreadTimingTrack : public FLoadingTimingTrack
-{
-public:
-	explicit FLoadingAsyncThreadTimingTrack(uint64 InTrackId, TSharedPtr<FLoadingSharedState> InState)
-		: FLoadingTimingTrack(InTrackId, FName(TEXT("AsyncThread")), TEXT("Loading - Async Thread"), InState)
-	{
-	}
-
-	virtual void Draw(ITimingViewDrawHelper& Helper) const override;
-	virtual bool SearchTimingEvent(const double InStartTime, const double InEndTime, TFunctionRef<bool(double, double, uint32)> InPredicate, FTimingEvent& InOutTimingEvent, bool bInStopAtFirstMatch, bool bInSearchForLargestEvent) const override;
+protected:
+	TSharedPtr<FLoadingSharedState> State;
+	uint32 TimelineIndex;
 };
 
 ////////////////////////////////////////////////////////////////////////////////////////////////////
diff --git a/Engine/Source/Developer/TraceInsights/Private/Insights/Widgets/STimingView.cpp b/Engine/Source/Developer/TraceInsights/Private/Insights/Widgets/STimingView.cpp
--- a/Engine/Source/Developer/TraceInsights/Private/Insights/Widgets/STimingView.cpp
+++ b/Engine/Source/Developer/TraceInsights/Private/Insights/Widgets/STimingView.cpp
@@ -61,9 +61,6 @@
 // start auto generated ids from a big number (MSB set to 1) to avoid collisions with ids for gpu/cpu tracks based on 32bit timeline index
 uint64 FBaseTimingTrack::IdGenerator = (1ULL << 63);
 
-const TCHAR* GetName(ELoadTimeProfilerPackageEventType Type);
-const TCHAR* GetName(ELoadTimeProfilerObjectEventType Type);
-
 const TCHAR* GetFileActivityTypeName(Trace::EFileActivityType Type);
 uint32 GetFileActivityTypeColor(Trace::EFileActivityType Type);
 
@@ -237,11 +234,7 @@ void STimingView::Reset()
 
 	LoadingSharedState->Reset();
 
-	LoadingMainThreadTrack = nullptr;
-	LoadingAsyncThreadTrack = nullptr;
-
-	LoadingMainThreadId = 0;
-	LoadingAsyncThreadId = 0;
+	LoadingTracks.Reset();
 
 	//////////////////////////////////////////////////
 
@@ -462,32 +455,6 @@ void STimingView::Tick(const FGeometry& AllottedGeometry, const double InCurrent
 		// After this point we can add new tracks, so reset the dirty flag
 		bTimingEventsTracksChanged = false;
 
-		if (Trace::ReadLoadTimeProfilerProvider(*Session.Get()))
-		{
-			const Trace::ILoadTimeProfilerProvider& LoadTimeProfilerProvider = *Trace::ReadLoadTimeProfilerProvider(*Session.Get());
-
-			LoadingMainThreadId = LoadTimeProfilerProvider.GetMainThreadId();
-			LoadingAsyncThreadId = LoadTimeProfilerProvider.GetAsyncLoadingThreadId();
-
-			if (LoadingMainThreadTrack == nullptr)
-			{
-				uint64 TrackId = FBaseTimingTrack::GenerateId();
-				LoadingMainThreadTrack = new FLoadingMainThreadTimingTrack(TrackId, LoadingSharedState);
-				LoadingMainThreadTrack->SetOrder(-3);
-				LoadingMainThreadTrack->SetVisibilityFlag(bShowHideAllLoadingTracks);
-				AddTimingEventsTrack(LoadingMainThreadTrack);
-			}
-
-			if (LoadingAsyncThreadTrack == nullptr)
-			{
-				uint64 TrackId = FBaseTimingTrack::GenerateId();
-				LoadingAsyncThreadTrack = new FLoadingAsyncThreadTimingTrack(TrackId, LoadingSharedState);
-				LoadingMainThreadTrack->SetOrder(-2);
-				LoadingAsyncThreadTrack->SetVisibilityFlag(bShowHideAllLoadingTracks);
-				AddTimingEventsTrack(LoadingAsyncThreadTrack);
-			}
-		}
-
 		if (Trace::ReadFileActivityProvider(*Session.Get()))
 		{
 			if (IoOverviewTrack == nullptr)
@@ -503,13 +470,14 @@ void STimingView::Tick(const FGeometry& AllottedGeometry, const double InCurrent
 
 		int32 Order = 1;
 
-		if (Trace::ReadTimingProfilerProvider(*Session.Get()))
-		{
-			const Trace::ITimingProfilerProvider& TimingProfilerProvider = *Trace::ReadTimingProfilerProvider(*Session.Get());
+		const Trace::ITimingProfilerProvider* TimingProfilerProvider = Trace::ReadTimingProfilerProvider(*Session.Get());
+		const Trace::ILoadTimeProfilerProvider* LoadTimeProfilerProvider = Trace::ReadLoadTimeProfilerProvider(*Session.Get());
 
+		if (TimingProfilerProvider || LoadTimeProfilerProvider)
+		{
 			// Check if we have a GPU track.
 			uint32 GpuTimelineIndex;
-			if (TimingProfilerProvider.GetGpuTimelineIndex(GpuTimelineIndex))
+			if (TimingProfilerProvider && TimingProfilerProvider->GetGpuTimelineIndex(GpuTimelineIndex))
 			{
 				const uint64 TrackId = static_cast<uint64>(GpuTimelineIndex);
 				if (!CachedTimelines.Contains(TrackId))
@@ -523,7 +491,7 @@ void STimingView::Tick(const FGeometry& AllottedGeometry, const double InCurrent
 
 			// Check available CPU tracks.
 			const Trace::IThreadProvider& ThreadProvider = Trace::ReadThreadProvider(*Session.Get());
-			ThreadProvider.EnumerateThreads([this, &Order, &TimingProfilerProvider](const Trace::FThreadInfo& ThreadInfo)
+			ThreadProvider.EnumerateThreads([this, &Order, TimingProfilerProvider, LoadTimeProfilerProvider](const Trace::FThreadInfo& ThreadInfo)
 			{
 				const TCHAR* GroupName = ThreadInfo.GroupName;
 				if (GroupName == nullptr)
@@ -547,8 +515,26 @@ void STimingView::Tick(const FGeometry& AllottedGeometry, const double InCurrent
 					}
 				}
 
+				bool bIsLoadingThread = false;
+				uint32 LoadingTimelineIndex;
+				if (LoadTimeProfilerProvider && LoadTimeProfilerProvider->GetCpuThreadTimelineIndex(ThreadInfo.Id, LoadingTimelineIndex))
+				{
+					bIsLoadingThread = true;
+					if (!LoadingTracks.Contains(LoadingTimelineIndex))
+					{
+						FString TrackName(ThreadInfo.Name && *ThreadInfo.Name ? FString::Printf(TEXT("Loading - %s"), ThreadInfo.Name) : FString::Printf(TEXT("Loading - Thread %u"), ThreadInfo.Id));
+
+						uint64 TrackId = FBaseTimingTrack::GenerateId();
+						FLoadingTimingTrack* LoadingThreadTrack = new FLoadingTimingTrack(TrackId, LoadingTimelineIndex, GroupName, TrackName, LoadingSharedState);
+						LoadingThreadTrack->SetOrder(-100 + LoadingTracks.Num());
+						LoadingThreadTrack->SetVisibilityFlag(bShowHideAllLoadingTracks);
+						AddTimingEventsTrack(LoadingThreadTrack);
+						LoadingTracks.Add(LoadingTimelineIndex, LoadingThreadTrack);
+					}
+				}
+
 				uint32 CpuTimelineIndex;
-				if (TimingProfilerProvider.GetCpuThreadTimelineIndex(ThreadInfo.Id, CpuTimelineIndex))
+				if (TimingProfilerProvider && TimingProfilerProvider->GetCpuThreadTimelineIndex(ThreadInfo.Id, CpuTimelineIndex))
 				{
 					FTimingEventsTrack* Track = nullptr;
 					const uint64 TrackId = static_cast<uint64>(CpuTimelineIndex);
@@ -566,7 +552,7 @@ void STimingView::Tick(const FGeometry& AllottedGeometry, const double InCurrent
 						FThreadGroup& ThreadGroup = ThreadGroups[GroupName];
 						ThreadGroup.NumTimelines++;
 
-						if (bAssetLoadingMode && (ThreadInfo.Id == LoadingMainThreadId || ThreadInfo.Id == LoadingAsyncThreadId))
+						if (bAssetLoadingMode && bIsLoadingThread)
 						{
 							Track->SetVisibilityFlag(true);
 							ThreadGroup.bIsVisible = true;
@@ -1819,11 +1805,6 @@ FReply STimingView::OnKeyDown(const FGeometry& MyGeometry, const FKeyEvent& InKe
 		LoadingSharedState->SetColorSchema(3);
 		return FReply::Handled();
 	}
-	else if (InKeyEvent.GetKey() == EKeys::Five)
-	{
-		LoadingSharedState->SetColorSchema(4);
-		return FReply::Handled();
-	}
 
 	return SCompoundWidget::OnKeyDown(MyGeometry, InKeyEvent);
 }
diff --git a/Engine/Source/Developer/TraceInsights/Private/Insights/Widgets/STimingView.h b/Engine/Source/Developer/TraceInsights/Private/Insights/Widgets/STimingView.h
--- a/Engine/Source/Developer/TraceInsights/Private/Insights/Widgets/STimingView.h
+++ b/Engine/Source/Developer/TraceInsights/Private/Insights/Widgets/STimingView.h
@@ -405,11 +405,7 @@ protected:
 	// Asset Loading
 
 	TSharedPtr<FLoadingSharedState> LoadingSharedState;
-	FTimingEventsTrack* LoadingMainThreadTrack;
-	FTimingEventsTrack* LoadingAsyncThreadTrack;
-
-	uint32 LoadingMainThreadId;
-	uint32 LoadingAsyncThreadId;
+	TMap<uint32, FTimingEventsTrack*> LoadingTracks;
 
 	bool bAssetLoadingMode;
 
diff --git a/Engine/Source/Developer/TraceInsights/Public/Insights/ViewModels/TimingEvent.h b/Engine/Source/Developer/TraceInsights/Public/Insights/ViewModels/TimingEvent.h
--- a/Engine/Source/Developer/TraceInsights/Public/Insights/ViewModels/TimingEvent.h
+++ b/Engine/Source/Developer/TraceInsights/Public/Insights/ViewModels/TimingEvent.h
@@ -57,8 +57,7 @@ struct FTimingEvent
 
 		LoadingInfo.Package = nullptr;
 		LoadingInfo.Export = nullptr;
-		LoadingInfo.PackageEventType = LoadTimeProfilerPackageEventType_None;
-		LoadingInfo.ExportEventType = LoadTimeProfilerObjectEventType_None;
+		LoadingInfo.EventType = Trace::LoadTimeProfilerObjectEventType_None;
 
 		Offset = 0;
 		Size = 0;
diff --git a/Engine/Source/Developer/TraceServices/Private/Analyzers/LoadTimeTraceAnalysis.cpp b/Engine/Source/Developer/TraceServices/Private/Analyzers/LoadTimeTraceAnalysis.cpp
--- a/Engine/Source/Developer/TraceServices/Private/Analyzers/LoadTimeTraceAnalysis.cpp
+++ b/Engine/Source/Developer/TraceServices/Private/Analyzers/LoadTimeTraceAnalysis.cpp
@@ -1,7 +1,6 @@
 // Copyright 1998-2019 Epic Games, Inc. All Rights Reserved.
 #include "LoadTimeTraceAnalysis.h"
 
-#include "HAL/FileManager.h"
 #include "Serialization/LoadTimeTrace.h"
 #include "Trace/Trace.h"
 #include "Model/LoadTimeProfilerPrivate.h"
@@ -17,35 +16,30 @@ FAsyncLoadingTraceAnalyzer::FAsyncLoadingTraceAnalyzer(Trace::IAnalysisSession&
 
 FAsyncLoadingTraceAnalyzer::~FAsyncLoadingTraceAnalyzer()
 {
-
+	for (const auto& KV : ThreadStatesMap)
+	{
+		delete KV.Value;
+	}
+	for (const auto& KV : ActiveAsyncPackagesMap)
+	{
+		delete KV.Value;
+	}
+	for (const auto& KV : ActiveRequestsMap)
+	{
+		delete KV.Value;
+	}
 }
 
-TSharedRef<FAsyncLoadingTraceAnalyzer::FThreadState> FAsyncLoadingTraceAnalyzer::GetThreadState(uint32 ThreadId)
+FAsyncLoadingTraceAnalyzer::FThreadState& FAsyncLoadingTraceAnalyzer::GetThreadState(uint32 ThreadId)
 {
-	if (!ThreadStatesMap.Contains(ThreadId))
+	FThreadState* ThreadState = ThreadStatesMap.FindRef(ThreadId);
+	if (!ThreadState)
 	{
-		TSharedRef<FThreadState> ThreadState = MakeShared<FThreadState>();
+		ThreadState = new FThreadState();
 		ThreadStatesMap.Add(ThreadId, ThreadState);
-		if (MainThreadId == -1)
-		{
-			MainThreadId = ThreadId;
-			LoadTimeProfilerProvider.SetMainThreadId(MainThreadId);
-			ThreadState->CpuTimeline = &LoadTimeProfilerProvider.EditMainThreadCpuTimeline();
-		}
-		else if (ThreadId == AsyncLoadingThreadId)
-		{
-			ThreadState->CpuTimeline = &LoadTimeProfilerProvider.EditAsyncLoadingThreadCpuTimeline();
-		}
-		else
-		{
-			ThreadState->CpuTimeline = &LoadTimeProfilerProvider.EditAdditionalCpuTimeline(ThreadId);
-		}
-		return ThreadState;
-	}
-	else
-	{
-		return ThreadStatesMap[ThreadId];
+		ThreadState->CpuTimeline = &LoadTimeProfilerProvider.EditCpuTimeline(ThreadId);
 	}
+	return *ThreadState;
 }
 
 const Trace::FClassInfo* FAsyncLoadingTraceAnalyzer::GetClassInfo(uint64 ClassPtr) const
@@ -68,25 +62,26 @@ void FAsyncLoadingTraceAnalyzer::OnAnalysisBegin(const FOnAnalysisContext& Conte
 	Builder.RouteEvent(RouteId_StartAsyncLoading, "LoadTime", "StartAsyncLoading");
 	Builder.RouteEvent(RouteId_SuspendAsyncLoading, "LoadTime", "SuspendAsyncLoading");
 	Builder.RouteEvent(RouteId_ResumeAsyncLoading, "LoadTime", "ResumeAsyncLoading");
-	Builder.RouteEvent(RouteId_NewLinker, "LoadTime", "NewLinker");
-	Builder.RouteEvent(RouteId_DestroyLinker, "LoadTime", "DestroyLinker");
 	Builder.RouteEvent(RouteId_PackageSummary, "LoadTime", "PackageSummary");
 	Builder.RouteEvent(RouteId_BeginCreateExport, "LoadTime", "BeginCreateExport");
 	Builder.RouteEvent(RouteId_EndCreateExport, "LoadTime", "EndCreateExport");
-	Builder.RouteEvent(RouteId_BeginObjectScope, "LoadTime", "BeginObjectScope");
-	Builder.RouteEvent(RouteId_EndObjectScope, "LoadTime", "EndObjectScope");
+	Builder.RouteEvent(RouteId_BeginSerializeExport, "LoadTime", "BeginSerializeExport");
+	Builder.RouteEvent(RouteId_EndSerializeExport, "LoadTime", "EndSerializeExport");
+	Builder.RouteEvent(RouteId_BeginPostLoadExport, "LoadTime", "BeginPostLoadExport");
+	Builder.RouteEvent(RouteId_EndPostLoadExport, "LoadTime", "EndPostLoadExport");
 	Builder.RouteEvent(RouteId_NewAsyncPackage, "LoadTime", "NewAsyncPackage");
+	Builder.RouteEvent(RouteId_BeginLoadAsyncPackage, "LoadTime", "BeginLoadAsyncPackage");
+	Builder.RouteEvent(RouteId_EndLoadAsyncPackage, "LoadTime", "EndLoadAsyncPackage");
 	Builder.RouteEvent(RouteId_DestroyAsyncPackage, "LoadTime", "DestroyAsyncPackage");
 	Builder.RouteEvent(RouteId_BeginRequest, "LoadTime", "BeginRequest");
 	Builder.RouteEvent(RouteId_EndRequest, "LoadTime", "EndRequest");
 	Builder.RouteEvent(RouteId_BeginRequestGroup, "LoadTime", "BeginRequestGroup");
 	Builder.RouteEvent(RouteId_EndRequestGroup, "LoadTime", "EndRequestGroup");
 	Builder.RouteEvent(RouteId_AsyncPackageRequestAssociation, "LoadTime", "AsyncPackageRequestAssociation");
-	Builder.RouteEvent(RouteId_AsyncPackageLinkerAssociation, "LoadTime", "AsyncPackageLinkerAssociation");
 	Builder.RouteEvent(RouteId_AsyncPackageImportDependency, "LoadTime", "AsyncPackageImportDependency");
-	Builder.RouteEvent(RouteId_BeginAsyncPackageScope, "LoadTime", "BeginAsyncPackageScope");
-	Builder.RouteEvent(RouteId_EndAsyncPackageScope, "LoadTime", "EndAsyncPackageScope");
 	Builder.RouteEvent(RouteId_ClassInfo, "LoadTime", "ClassInfo");
+	Builder.RouteEvent(RouteId_BatchIssued, "IoDispatcher", "BatchIssued");
+	Builder.RouteEvent(RouteId_BatchResolved, "IoDispatcher", "BatchResolved");
 }
 
 bool FAsyncLoadingTraceAnalyzer::OnEvent(uint16 RouteId, const FOnEventContext& Context)
@@ -94,122 +89,106 @@ bool FAsyncLoadingTraceAnalyzer::OnEvent(uint16 RouteId, const FOnEventContext&
 	const auto& EventData = Context.EventData;
 	switch (RouteId)
 	{
-	case RouteId_StartAsyncLoading:
-	{
-		Trace::FAnalysisSessionEditScope _(Session);
-		AsyncLoadingThreadId = EventData.GetValue<uint32>("ThreadId");
-		LoadTimeProfilerProvider.SetAsyncLoadingThreadId(AsyncLoadingThreadId);
-		break;
-	}
-	case RouteId_NewLinker:
-	{
-		uint64 LinkerPtr = EventData.GetValue<uint64>("Linker");
-		//check(!ActiveLinkersMap.Contains(LinkerPtr));
-		TSharedRef<FLinkerState> LinkerState = MakeShared<FLinkerState>();
-		ActiveLinkersMap.Add(LinkerPtr, LinkerState);
-		break;
-	}
-	case RouteId_DestroyLinker:
-	{
-		uint64 Ptr = EventData.GetValue<uint64>("Linker");
-		//check(ActiveLinkersMap.Contains(Ptr));
-		ActiveLinkersMap.Remove(Ptr);
-		break;
-	}
 	case RouteId_PackageSummary:
 	{
-		uint64 LinkerPtr = EventData.GetValue<uint64>("Linker");
-		TSharedRef<FLinkerState>* LinkerState = ActiveLinkersMap.Find(LinkerPtr);
-		if (LinkerState && (*LinkerState)->PackageInfo)
+		uint64 AsyncPackagePtr = EventData.GetValue<uint64>("AsyncPackage");
+		FAsyncPackageState* AsyncPackage = ActiveAsyncPackagesMap.FindRef(AsyncPackagePtr);
+		if (AsyncPackage)
 		{
 			Trace::FAnalysisSessionEditScope _(Session);
-			Trace::FPackageSummaryInfo& Summary = (*LinkerState)->PackageInfo->Summary;
+			Trace::FPackageSummaryInfo& Summary = AsyncPackage->PackageInfo->Summary;
 			Summary.TotalHeaderSize = EventData.GetValue<uint32>("TotalHeaderSize");
-			Summary.NameCount = EventData.GetValue<uint32>("NameCount");
 			Summary.ImportCount = EventData.GetValue<uint32>("ImportCount");
 			Summary.ExportCount = EventData.GetValue<uint32>("ExportCount");
 		}
 		break;
 	}
-	case RouteId_BeginAsyncPackageScope:
+	case RouteId_BeginCreateExport:
 	{
 		Trace::FAnalysisSessionEditScope _(Session);
 
+		Trace::FPackageExportInfo& Export = LoadTimeProfilerProvider.CreateExport();
 		uint64 AsyncPackagePtr = EventData.GetValue<uint64>("AsyncPackage");
-		TSharedRef<FAsyncPackageState>* AsyncPackageState = ActiveAsyncPackagesMap.Find(AsyncPackagePtr);
+		FAsyncPackageState* AsyncPackage = ActiveAsyncPackagesMap.FindRef(AsyncPackagePtr);
+		if (AsyncPackage)
+		{
+			AsyncPackage->PackageInfo->Exports.Add(&Export);
+			Export.Package = AsyncPackage->PackageInfo;
+		}
 		uint32 ThreadId = EventData.GetValue<uint32>("ThreadId");
-		TSharedRef<FThreadState> ThreadState = GetThreadState(ThreadId);
-		ELoadTimeProfilerPackageEventType EventType = static_cast<ELoadTimeProfilerPackageEventType>(EventData.GetValue<uint8>("EventType"));
-		ThreadState->EnterPackageScope(Context.SessionContext.TimestampFromCycle(EventData.GetValue<uint64>("Cycle")), AsyncPackageState ? (*AsyncPackageState)->PackageInfo : nullptr, EventType);
+		FThreadState& ThreadState = GetThreadState(ThreadId);
+		ThreadState.EnterExportScope(Context.SessionContext.TimestampFromCycle(EventData.GetValue<uint64>("Cycle")), &Export, Trace::LoadTimeProfilerObjectEventType_Create);
 		break;
 	}
-	case RouteId_EndAsyncPackageScope:
+	case RouteId_EndCreateExport:
 	{
 		Trace::FAnalysisSessionEditScope _(Session);
 
 		uint32 ThreadId = EventData.GetValue<uint32>("ThreadId");
-		TSharedRef<FThreadState> ThreadState = GetThreadState(ThreadId);
-		uint64 Cycle = EventData.GetValue<uint64>("Cycle");
-		ThreadState->LeaveScope(Context.SessionContext.TimestampFromCycle(EventData.GetValue<uint64>("Cycle")));
+		FThreadState& ThreadState = GetThreadState(ThreadId);
+		uint64 ObjectPtr = EventData.GetValue<uint64>("Object");
+		Trace::FPackageExportInfo* Export = ThreadState.GetCurrentExportScope();
+		if (Export)
+		{
+			ExportsMap.Add(ObjectPtr, Export);
+			const Trace::FClassInfo* ObjectClass = GetClassInfo(EventData.GetValue<uint64>("Class"));
+			Export->Class = ObjectClass;
+		}
+		ThreadState.LeaveExportScope(Context.SessionContext.TimestampFromCycle(EventData.GetValue<uint64>("Cycle")));
 		break;
 	}
-	case RouteId_BeginCreateExport:
+	case RouteId_BeginSerializeExport:
 	{
 		Trace::FAnalysisSessionEditScope _(Session);
 
-		Trace::FPackageExportInfo& Export = LoadTimeProfilerProvider.CreateExport();
-		Export.SerialOffset = EventData.GetValue<uint64>("SerialOffset");
-		Export.SerialSize = EventData.GetValue<uint64>("SerialSize");
-		Export.IsAsset = EventData.GetValue<bool>("IsAsset");
+		uint64 ObjectPtr = EventData.GetValue<uint64>("Object");
+		uint64 SerialSize = EventData.GetValue<uint64>("SerialSize");
 		uint32 ThreadId = EventData.GetValue<uint32>("ThreadId");
-		TSharedRef<FThreadState> ThreadState = GetThreadState(ThreadId);
-		ThreadState->EnterExportScope(Context.SessionContext.TimestampFromCycle(EventData.GetValue<uint64>("Cycle")), &Export, LoadTimeProfilerObjectEventType_Create);
-		uint64 LinkerPtr = EventData.GetValue<uint64>("Linker");
-		TSharedRef<FLinkerState>* LinkerState = ActiveLinkersMap.Find(LinkerPtr);
-		if (LinkerState && (*LinkerState)->PackageInfo)
+
+		Trace::FPackageExportInfo* Export = ExportsMap.FindRef(ObjectPtr);
+		check(Export);
+
+		if (Export)
 		{
-			(*LinkerState)->PackageInfo->Exports.Add(&Export);
-			Export.Package = (*LinkerState)->PackageInfo;
+			Export->SerialSize = SerialSize;
+			if (Export->Package)
+			{
+				const_cast<Trace::FPackageInfo*>(Export->Package)->TotalExportsSerialSize += SerialSize;
+			}
 		}
+		FThreadState& ThreadState = GetThreadState(ThreadId);
+		ThreadState.EnterExportScope(Context.SessionContext.TimestampFromCycle(EventData.GetValue<uint64>("Cycle")), Export, Trace::LoadTimeProfilerObjectEventType_Serialize);
 		break;
 	}
-	case RouteId_EndCreateExport:
+	case RouteId_EndSerializeExport:
 	{
-		Trace::FAnalysisSessionEditScope _(Session);
-
 		uint32 ThreadId = EventData.GetValue<uint32>("ThreadId");
-		TSharedRef<FThreadState> ThreadState = GetThreadState(ThreadId);
-		uint64 ObjectPtr = EventData.GetValue<uint64>("Object");
-		Trace::FPackageExportInfo* Export = ThreadState->GetCurrentExportScope();
-		if (Export)
+		FThreadState& ThreadState = GetThreadState(ThreadId);
 		{
-			ExportsMap.Add(ObjectPtr, Export);
-			const Trace::FClassInfo* ObjectClass = GetClassInfo(EventData.GetValue<uint64>("Class"));
-			Export->Class = ObjectClass;
+			Trace::FAnalysisSessionEditScope _(Session);
+			ThreadState.LeaveExportScope(Context.SessionContext.TimestampFromCycle(EventData.GetValue<uint64>("Cycle")));
 		}
-		ThreadState->LeaveScope(Context.SessionContext.TimestampFromCycle(EventData.GetValue<uint64>("Cycle")));
 		break;
 	}
-	case RouteId_BeginObjectScope:
+	case RouteId_BeginPostLoadExport:
 	{
 		uint64 ObjectPtr = EventData.GetValue<uint64>("Object");
-		Trace::FPackageExportInfo** Export = ExportsMap.Find(ObjectPtr);
+		Trace::FPackageExportInfo* Export = ExportsMap.FindRef(ObjectPtr);
 		uint32 ThreadId = EventData.GetValue<uint32>("ThreadId");
-		TSharedRef<FThreadState> ThreadState = GetThreadState(ThreadId);
-		ELoadTimeProfilerObjectEventType EventType = static_cast<ELoadTimeProfilerObjectEventType>(EventData.GetValue<uint8>("EventType"));
+		FThreadState& ThreadState = GetThreadState(ThreadId);
 		{
 			Trace::FAnalysisSessionEditScope _(Session);
-			ThreadState->EnterExportScope(Context.SessionContext.TimestampFromCycle(EventData.GetValue<uint64>("Cycle")), Export ? *Export : nullptr, EventType);
+			ThreadState.EnterExportScope(Context.SessionContext.TimestampFromCycle(EventData.GetValue<uint64>("Cycle")), Export, Trace::LoadTimeProfilerObjectEventType_PostLoad);
 		}
 		break;
 	}
-	case RouteId_EndObjectScope:
+	case RouteId_EndPostLoadExport:
 	{
 		uint32 ThreadId = EventData.GetValue<uint32>("ThreadId");
-		TSharedRef<FThreadState> ThreadState = GetThreadState(ThreadId);
+		FThreadState& ThreadState = GetThreadState(ThreadId);
 		{
 			Trace::FAnalysisSessionEditScope _(Session);
-			ThreadState->LeaveScope(Context.SessionContext.TimestampFromCycle(EventData.GetValue<uint64>("Cycle")));
+			ThreadState.LeaveExportScope(Context.SessionContext.TimestampFromCycle(EventData.GetValue<uint64>("Cycle")));
 		}
 		break;
 	}
@@ -218,13 +197,13 @@ bool FAsyncLoadingTraceAnalyzer::OnEvent(uint16 RouteId, const FOnEventContext&
 		uint64 RequestId = EventData.GetValue<uint64>("RequestId");
 		uint32 ThreadId = EventData.GetValue<uint32>("ThreadId");
 		//check(!ActiveRequestsMap.Contains(RequestId));
-		TSharedRef<FRequestState> RequestState = MakeShared<FRequestState>();
+		FRequestState* RequestState = new FRequestState();
 		RequestState->WallTimeStartCycle = EventData.GetValue<uint64>("Cycle");
 		RequestState->WallTimeEndCycle = 0;
 		RequestState->ThreadId = ThreadId;
 		ActiveRequestsMap.Add(RequestId, RequestState);
-		TSharedRef<FThreadState> ThreadState = GetThreadState(ThreadId);
-		TSharedPtr<FRequestGroupState> RequestGroup = ThreadState->RequestGroupStack.Num() ? ThreadState->RequestGroupStack.Top() : nullptr;
+		FThreadState& ThreadState = GetThreadState(ThreadId);
+		TSharedPtr<FRequestGroupState> RequestGroup = ThreadState.RequestGroupStack.Num() ? ThreadState.RequestGroupStack.Top() : nullptr;
 		if (!RequestGroup)
 		{
 			RequestGroup = MakeShared<FRequestGroupState>();
@@ -239,10 +218,9 @@ bool FAsyncLoadingTraceAnalyzer::OnEvent(uint16 RouteId, const FOnEventContext&
 	case RouteId_EndRequest:
 	{
 		uint64 RequestId = EventData.GetValue<uint64>("RequestId");
-		TSharedRef<FRequestState>* FindRequestState = ActiveRequestsMap.Find(RequestId);
-		if (FindRequestState)
+		FRequestState* RequestState = ActiveRequestsMap.FindRef(RequestId);
+		if (RequestState)
 		{
-			TSharedRef<FRequestState> RequestState = *FindRequestState;
 			RequestState->WallTimeEndCycle = EventData.GetValue<uint64>("Cycle");
 			RequestState->Group->LatestEndCycle = FMath::Max(RequestState->Group->LatestEndCycle, RequestState->WallTimeEndCycle);
 			--RequestState->Group->ActiveRequestsCount;
@@ -262,17 +240,17 @@ bool FAsyncLoadingTraceAnalyzer::OnEvent(uint16 RouteId, const FOnEventContext&
 		Trace::FFormatArgsHelper::Format(FormatBuffer, FormatBufferSize - 1, TempBuffer, FormatBufferSize - 1, FormatString, FormatArgs);
 		GroupState->Name = Session.StoreString(FormatBuffer);
 		uint32 ThreadId = EventData.GetValue<uint32>("ThreadId");
-		TSharedRef<FThreadState> ThreadState = GetThreadState(ThreadId);
-		ThreadState->RequestGroupStack.Push(GroupState);
+		FThreadState& ThreadState = GetThreadState(ThreadId);
+		ThreadState.RequestGroupStack.Push(GroupState);
 		break;
 	}
 	case RouteId_EndRequestGroup:
 	{
 		uint32 ThreadId = EventData.GetValue<uint32>("ThreadId");
-		TSharedRef<FThreadState> ThreadState = GetThreadState(ThreadId);
-		if (ThreadState->RequestGroupStack.Num())
+		FThreadState& ThreadState = GetThreadState(ThreadId);
+		if (ThreadState.RequestGroupStack.Num())
 		{
-			TSharedPtr<FRequestGroupState> GroupState = ThreadState->RequestGroupStack.Pop(false);
+			TSharedPtr<FRequestGroupState> GroupState = ThreadState.RequestGroupStack.Pop(false);
 			GroupState->bIsClosed = true;
 			if (GroupState->LoadRequest && GroupState->ActiveRequestsCount == 0)
 			{
@@ -287,11 +265,41 @@ bool FAsyncLoadingTraceAnalyzer::OnEvent(uint16 RouteId, const FOnEventContext&
 
 		uint64 AsyncPackagePtr = EventData.GetValue<uint64>("AsyncPackage");
 		//check(!ActivePackagesMap.Contains(AsyncPackagePtr));
-		TSharedRef<FAsyncPackageState> AsyncPackageState = MakeShared<FAsyncPackageState>();
-		AsyncPackageState->PackageInfo = &LoadTimeProfilerProvider.CreatePackage(reinterpret_cast<const TCHAR*>(EventData.GetAttachment()));
+		FAsyncPackageState* AsyncPackageState = new FAsyncPackageState();
+		AsyncPackageState->PackageInfo = &LoadTimeProfilerProvider.EditPackageInfo(reinterpret_cast<const TCHAR*>(EventData.GetAttachment()));
 		ActiveAsyncPackagesMap.Add(AsyncPackagePtr, AsyncPackageState);
 		break;
 	}
+	case RouteId_BeginLoadAsyncPackage:
+	{
+		uint64 AsyncPackagePtr = EventData.GetValue<uint64>("AsyncPackage");
+		FAsyncPackageState* AsyncPackage = ActiveAsyncPackagesMap.FindRef(AsyncPackagePtr);
+		if (AsyncPackage)
+		{
+			AsyncPackage->LoadStartCycle = EventData.GetValue<uint64>("Cycle");
+			if (AsyncPackage->PackageInfo)
+			{
+				Trace::FAnalysisSessionEditScope _(Session);
+				AsyncPackage->LoadHandle = LoadTimeProfilerProvider.BeginLoadPackage(*AsyncPackage->PackageInfo, Context.SessionContext.TimestampFromCycle(AsyncPackage->LoadStartCycle));
+			}
+		}
+		break;
+	}
+	case RouteId_EndLoadAsyncPackage:
+	{
+		uint64 AsyncPackagePtr = EventData.GetValue<uint64>("AsyncPackage");
+		FAsyncPackageState* AsyncPackage = ActiveAsyncPackagesMap.FindRef(AsyncPackagePtr);
+		if (AsyncPackage)
+		{
+			AsyncPackage->LoadEndCycle = EventData.GetValue<uint64>("Cycle");
+			if (AsyncPackage->PackageInfo && AsyncPackage->LoadHandle != uint64(-1))
+			{
+				Trace::FAnalysisSessionEditScope _(Session);
+				LoadTimeProfilerProvider.EndLoadPackage(AsyncPackage->LoadHandle, Context.SessionContext.TimestampFromCycle(AsyncPackage->LoadEndCycle));
+			}
+		}
+		break;
+	}
 	case RouteId_DestroyAsyncPackage:
 	{
 		uint64 AsyncPackagePtr = EventData.GetValue<uint64>("AsyncPackage");
@@ -302,41 +310,25 @@ bool FAsyncLoadingTraceAnalyzer::OnEvent(uint16 RouteId, const FOnEventContext&
 	case RouteId_AsyncPackageImportDependency:
 	{
 		uint64 AsyncPackagePtr = EventData.GetValue<uint64>("AsyncPackage");
-		TSharedRef<FAsyncPackageState>* FindAsyncPackage = ActiveAsyncPackagesMap.Find(AsyncPackagePtr);
+		FAsyncPackageState* AsyncPackage = ActiveAsyncPackagesMap.FindRef(AsyncPackagePtr);
 		uint64 ImportedAsyncPackagePtr = EventData.GetValue<uint64>("ImportedAsyncPackage");
-		TSharedRef<FAsyncPackageState>* FindImportedAsyncPackage = ActiveAsyncPackagesMap.Find(ImportedAsyncPackagePtr);
-		if (FindAsyncPackage && FindImportedAsyncPackage)
+		FAsyncPackageState* ImportedAsyncPackage = ActiveAsyncPackagesMap.FindRef(ImportedAsyncPackagePtr);
+		if (AsyncPackage && ImportedAsyncPackage)
 		{
-			TSharedRef<FAsyncPackageState> AsyncPackageState = *FindAsyncPackage;
-			TSharedRef<FAsyncPackageState> ImportedAsyncPackageState = *FindImportedAsyncPackage;
-			check(AsyncPackageState->Request);
-			PackageRequestAssociation(Context, ImportedAsyncPackageState, AsyncPackageState->Request.ToSharedRef());
+			check(AsyncPackage->Request);
+			PackageRequestAssociation(Context, ImportedAsyncPackage, AsyncPackage->Request);
 		}
 		break;
 	}
 	case RouteId_AsyncPackageRequestAssociation:
 	{
 		uint64 AsyncPackagePtr = EventData.GetValue<uint64>("AsyncPackage");
-		TSharedRef<FAsyncPackageState>* FindAsyncPackageState = ActiveAsyncPackagesMap.Find(AsyncPackagePtr);
+		FAsyncPackageState* AsyncPackageState = ActiveAsyncPackagesMap.FindRef(AsyncPackagePtr);
 		uint64 RequestId = EventData.GetValue<uint64>("RequestId");
-		TSharedRef<FRequestState>* FindRequestState = ActiveRequestsMap.Find(RequestId);
-		if (FindAsyncPackageState && FindRequestState)
+		FRequestState* RequestState = ActiveRequestsMap.FindRef(RequestId);
+		if (AsyncPackageState && RequestState)
 		{
-			PackageRequestAssociation(Context, *FindAsyncPackageState, *FindRequestState);
-		}
-		break;
-	}
-	case RouteId_AsyncPackageLinkerAssociation:
-	{
-		uint64 LinkerPtr = EventData.GetValue<uint64>("Linker");
-		TSharedRef<FLinkerState>* LinkerState = ActiveLinkersMap.Find(LinkerPtr);
-		uint64 AsyncPackagePtr = EventData.GetValue<uint64>("AsyncPackage");
-		TSharedRef<FAsyncPackageState>* AsyncPackageState = ActiveAsyncPackagesMap.Find(AsyncPackagePtr);
-		if (LinkerState && AsyncPackageState)
-		{
-			(*AsyncPackageState)->LinkerState = *LinkerState;
-			(*LinkerState)->AsyncPackageState = *AsyncPackageState;
-			(*LinkerState)->PackageInfo = (*AsyncPackageState)->PackageInfo;
+			PackageRequestAssociation(Context, AsyncPackageState, RequestState);
 		}
 		break;
 	}
@@ -349,38 +341,49 @@ bool FAsyncLoadingTraceAnalyzer::OnEvent(uint16 RouteId, const FOnEventContext&
 		ClassInfosMap.Add(ClassPtr, &ClassInfo);
 		break;
 	}
+	case RouteId_BatchIssued:
+	{
+		Trace::FAnalysisSessionEditScope _(Session);
+
+		uint64 Cycle = EventData.GetValue<uint64>("Cycle");
+		uint64 BatchId = EventData.GetValue<uint64>("BatchId");
+
+		uint64 BatchHandle = LoadTimeProfilerProvider.BeginIoDispatcherBatch(BatchId, Context.SessionContext.TimestampFromCycle(Cycle));
+		ActiveBatchesMap.Add(BatchId, BatchHandle);
+
+		break;
 	}
+	case RouteId_BatchResolved:
+	{
+		uint64 Cycle = EventData.GetValue<uint64>("Cycle");
+		uint64 BatchId = EventData.GetValue<uint64>("BatchId");
+		uint64 TotalSize = EventData.GetValue<uint64>("TotalSize");
 
-	return true;
-}
+		uint64* FindBatchHandle = ActiveBatchesMap.Find(BatchId);
+		if (FindBatchHandle)
+		{
+			Trace::FAnalysisSessionEditScope _(Session);
+			LoadTimeProfilerProvider.EndIoDispatcherBatch(*FindBatchHandle, Context.SessionContext.TimestampFromCycle(Cycle), TotalSize);
+		}
 
-void FAsyncLoadingTraceAnalyzer::FThreadState::EnterPackageScope(double Time, const Trace::FPackageInfo* PackageInfo, ELoadTimeProfilerPackageEventType EventType)
-{
-	FScopeStackEntry& StackEntry = CpuScopeStack[CpuScopeStackDepth++];
-	StackEntry.Event.Export = nullptr;
-	StackEntry.Event.ExportEventType = LoadTimeProfilerObjectEventType_None;
-	StackEntry.Event.Package = PackageInfo;
-	StackEntry.Event.PackageEventType = EventType;
-	CurrentEvent = StackEntry.Event;
-	CpuTimeline->AppendBeginEvent(Time, StackEntry.Event);
+		break;
+	}
+	}
+
+	return true;
 }
 
-void FAsyncLoadingTraceAnalyzer::FThreadState::EnterExportScope(double Time, const Trace::FPackageExportInfo* ExportInfo, ELoadTimeProfilerObjectEventType EventType)
+void FAsyncLoadingTraceAnalyzer::FThreadState::EnterExportScope(double Time, const Trace::FPackageExportInfo* ExportInfo, Trace::ELoadTimeProfilerObjectEventType EventType)
 {
 	FScopeStackEntry& StackEntry = CpuScopeStack[CpuScopeStackDepth++];
 	StackEntry.Event.Export = ExportInfo;
-	StackEntry.Event.ExportEventType = EventType;
+	StackEntry.Event.EventType = EventType;
 	StackEntry.Event.Package = ExportInfo ? ExportInfo->Package : nullptr;
-	StackEntry.Event.PackageEventType = CurrentEvent.PackageEventType;
-	if (EventType == LoadTimeProfilerObjectEventType_PostLoad && StackEntry.Event.PackageEventType == LoadTimeProfilerPackageEventType_None)
-	{
-		StackEntry.Event.PackageEventType = LoadTimeProfilerPackageEventType_DeferredPostLoad;
-	}
 	CurrentEvent = StackEntry.Event;
 	CpuTimeline->AppendBeginEvent(Time, StackEntry.Event);
 }
 
-void FAsyncLoadingTraceAnalyzer::FThreadState::LeaveScope(double Time)
+void FAsyncLoadingTraceAnalyzer::FThreadState::LeaveExportScope(double Time)
 {
 	FScopeStackEntry& StackEntry = CpuScopeStack[--CpuScopeStackDepth];
 	CpuTimeline->AppendEndEvent(Time);
@@ -407,7 +410,7 @@ Trace::FPackageExportInfo* FAsyncLoadingTraceAnalyzer::FThreadState::GetCurrentE
 	}
 }
 
-void FAsyncLoadingTraceAnalyzer::PackageRequestAssociation(const FOnEventContext& Context, TSharedRef<FAsyncPackageState> AsyncPackageState, TSharedRef<FRequestState> RequestState)
+void FAsyncLoadingTraceAnalyzer::PackageRequestAssociation(const FOnEventContext& Context, FAsyncPackageState* AsyncPackageState, FRequestState* RequestState)
 {
 	if (!AsyncPackageState->Request)
 	{
diff --git a/Engine/Source/Developer/TraceServices/Private/Analyzers/LoadTimeTraceAnalysis.h b/Engine/Source/Developer/TraceServices/Private/Analyzers/LoadTimeTraceAnalysis.h
--- a/Engine/Source/Developer/TraceServices/Private/Analyzers/LoadTimeTraceAnalysis.h
+++ b/Engine/Source/Developer/TraceServices/Private/Analyzers/LoadTimeTraceAnalysis.h
@@ -16,8 +16,7 @@ inline bool operator!=(const Trace::FLoadTimeProfilerCpuEvent& Lhs, const Trace:
 {
 	return Lhs.Package != Rhs.Package ||
 		Lhs.Export != Rhs.Export ||
-		Lhs.PackageEventType != Rhs.PackageEventType ||
-		Lhs.ExportEventType != Rhs.ExportEventType;
+		Lhs.EventType != Rhs.EventType;
 }
 
 class FAsyncLoadingTraceAnalyzer : public Trace::IAnalyzer
@@ -36,7 +35,7 @@ private:
 	struct FRequestGroupState
 	{
 		FString Name;
-		TArray<TSharedRef<FRequestState>> Requests;
+		TArray<FRequestState*> Requests;
 		Trace::FLoadRequest* LoadRequest = nullptr;
 		uint64 LatestEndCycle = 0;
 		uint64 ActiveRequestsCount = 0;
@@ -49,20 +48,16 @@ private:
 		uint64 WallTimeEndCycle;
 		uint32 ThreadId;
 		TSharedPtr<FRequestGroupState> Group;
-		TArray<TSharedRef<FAsyncPackageState>> AsyncPackages;
-	};
-
-	struct FLinkerState
-	{
-		TSharedPtr<FAsyncPackageState> AsyncPackageState;
-		Trace::FPackageInfo* PackageInfo = nullptr;
+		TArray<FAsyncPackageState*> AsyncPackages;
 	};
 
 	struct FAsyncPackageState
 	{
 		Trace::FPackageInfo* PackageInfo = nullptr;
-		TSharedPtr<FLinkerState> LinkerState;
-		TSharedPtr<FRequestState> Request;
+		FRequestState* Request = nullptr;
+		uint64 LoadHandle = uint64(-1);
+		uint64 LoadStartCycle = 0;
+		uint64 LoadEndCycle = 0;
 	};
 
 	struct FScopeStackEntry
@@ -80,14 +75,13 @@ private:
 		
 		Trace::FLoadTimeProfilerProvider::CpuTimelineInternal* CpuTimeline;
 
-		void EnterPackageScope(double Time, const Trace::FPackageInfo* PackageInfo, ELoadTimeProfilerPackageEventType EventType);
-		void EnterExportScope(double Time, const Trace::FPackageExportInfo* ExportInfo, ELoadTimeProfilerObjectEventType EventType);
-		void LeaveScope(double Time);
+		void EnterExportScope(double Time, const Trace::FPackageExportInfo* ExportInfo, Trace::ELoadTimeProfilerObjectEventType EventType);
+		void LeaveExportScope(double Time);
 		Trace::FPackageExportInfo* GetCurrentExportScope();
 	};
 
-	void PackageRequestAssociation(const FOnEventContext& Context, TSharedRef<FAsyncPackageState> AsyncPackageState, TSharedRef<FRequestState> RequestState);
-	TSharedRef<FThreadState> GetThreadState(uint32 ThreadId);
+	void PackageRequestAssociation(const FOnEventContext& Context, FAsyncPackageState* AsyncPackageState, FRequestState* RequestState);
+	FThreadState& GetThreadState(uint32 ThreadId);
 	const Trace::FClassInfo* GetClassInfo(uint64 ClassPtr) const;
 
 	enum : uint16
@@ -95,9 +89,9 @@ private:
 		RouteId_StartAsyncLoading,
 		RouteId_SuspendAsyncLoading,
 		RouteId_ResumeAsyncLoading,
-		RouteId_NewLinker,
-		RouteId_DestroyLinker,
 		RouteId_NewAsyncPackage,
+		RouteId_BeginLoadAsyncPackage,
+		RouteId_EndLoadAsyncPackage,
 		RouteId_DestroyAsyncPackage,
 		RouteId_BeginRequest,
 		RouteId_EndRequest,
@@ -105,15 +99,16 @@ private:
 		RouteId_EndRequestGroup,
 		RouteId_PackageSummary,
 		RouteId_AsyncPackageRequestAssociation,
-		RouteId_AsyncPackageLinkerAssociation,
 		RouteId_AsyncPackageImportDependency,
-		RouteId_BeginAsyncPackageScope,
-		RouteId_EndAsyncPackageScope,
 		RouteId_BeginCreateExport,
 		RouteId_EndCreateExport,
-		RouteId_BeginObjectScope,
-		RouteId_EndObjectScope,
+		RouteId_BeginSerializeExport,
+		RouteId_EndSerializeExport,
+		RouteId_BeginPostLoadExport,
+		RouteId_EndPostLoadExport,
 		RouteId_ClassInfo,
+		RouteId_BatchIssued,
+		RouteId_BatchResolved,
 	};
 
 	enum
@@ -126,15 +121,38 @@ private:
 	Trace::IAnalysisSession& Session;
 	Trace::FLoadTimeProfilerProvider& LoadTimeProfilerProvider;
 
-	//TArray<TSharedRef<FRequestState>> Requests;
+	template<typename ValueType>
+	struct FPointerMapKeyFuncs
+	{
+		typedef uint64 KeyType;
+		typedef uint64 KeyInitType;
+		typedef const TPairInitializer<uint64, ValueType>& ElementInitType;
+
+		enum { bAllowDuplicateKeys = false };
+
+		static FORCEINLINE bool Matches(uint64 A, uint64 B)
+		{
+			return A == B;
+		}
+
+		static FORCEINLINE uint32 GetKeyHash(uint64 Key)
+		{
+			return PointerHash((const void*)Key);
+		}
+
+		static FORCEINLINE KeyInitType GetSetKey(ElementInitType Element)
+		{
+			return Element.Key;
+		}
+	};
 
-	TMap<uint64, TSharedRef<FLinkerState>> ActiveLinkersMap;
-	TMap<uint64, TSharedRef<FAsyncPackageState>> ActiveAsyncPackagesMap;
-	TMap<uint64, Trace::FPackageExportInfo*> ExportsMap;
-	TMap<uint64, TSharedRef<FRequestState>> ActiveRequestsMap;
-	TMap<uint32, TSharedRef<FThreadState>> ThreadStatesMap;
-	TMap<uint64, const Trace::FClassInfo*> ClassInfosMap;
+	template<typename ValueType>
+	using TPointerMap = TMap<uint64, ValueType, FDefaultSetAllocator, FPointerMapKeyFuncs<ValueType>>;
 
-	int64 MainThreadId = -1;
-	int64 AsyncLoadingThreadId = -1;
+	TPointerMap<FAsyncPackageState*> ActiveAsyncPackagesMap;
+	TPointerMap<Trace::FPackageExportInfo*> ExportsMap;
+	TMap<uint64, FRequestState*> ActiveRequestsMap;
+	TPointerMap<uint64> ActiveBatchesMap;
+	TMap<uint32, FThreadState*> ThreadStatesMap;
+	TPointerMap<const Trace::FClassInfo*> ClassInfosMap;
 };
diff --git a/Engine/Source/Developer/TraceServices/Private/Model/LoadTimeProfiler.cpp b/Engine/Source/Developer/TraceServices/Private/Model/LoadTimeProfiler.cpp
--- a/Engine/Source/Developer/TraceServices/Private/Model/LoadTimeProfiler.cpp
+++ b/Engine/Source/Developer/TraceServices/Private/Model/LoadTimeProfiler.cpp
@@ -8,15 +8,96 @@
 namespace Trace
 {
 
-FLoadTimeProfilerProvider::FLoadTimeProfilerProvider(IAnalysisSession& InSession)
+FLoadTimeProfilerProvider::FLoaderFrameCounter::FLoaderFrameCounter(ELoaderFrameCounterType InType, const TPagedArray<FLoaderFrame>& InFrames)
+	: Frames(InFrames)
+	, Type(InType)
+{
+
+}
+
+const TCHAR* FLoadTimeProfilerProvider::FLoaderFrameCounter::GetName() const
+{
+	switch (Type)
+	{
+	case LoaderFrameCounterType_IoDispatcherThroughput:
+		return TEXT("AssetLoading/IoDispatcher/Throughput");
+	case LoaderFrameCounterType_LoaderThroughput:
+		return TEXT("AssetLoading/AsyncLoading/Throughput");
+	}
+	check(false);
+	return TEXT("");
+}
+
+const TCHAR* FLoadTimeProfilerProvider::FLoaderFrameCounter::GetDescription() const
+{
+	return TEXT("N/A");
+}
+
+bool FLoadTimeProfilerProvider::FLoaderFrameCounter::IsFloatingPoint() const
+{
+	switch (Type)
+	{
+	case LoaderFrameCounterType_IoDispatcherThroughput:
+	case LoaderFrameCounterType_LoaderThroughput:
+		return true;
+	}
+	return false;
+}
+
+ECounterDisplayHint FLoadTimeProfilerProvider::FLoaderFrameCounter::GetDisplayHint() const
+{
+	switch (Type)
+	{
+	case LoaderFrameCounterType_IoDispatcherThroughput:
+	case LoaderFrameCounterType_LoaderThroughput:
+		return CounterDisplayHint_Memory;
+	}
+	return CounterDisplayHint_None;
+}
+
+void FLoadTimeProfilerProvider::FLoaderFrameCounter::EnumerateValues(double IntervalStart, double IntervalEnd, bool bIncludeExternalBounds, TFunctionRef<void(double, int64)> Callback) const
+{
+	
+}
+
+void FLoadTimeProfilerProvider::FLoaderFrameCounter::EnumerateFloatValues(double IntervalStart, double IntervalEnd, bool bIncludeExternalBounds, TFunctionRef<void(double, double)> Callback) const
+{
+	if (Frames.Num() == 0)
+	{
+		return;
+	}
+	int32 FirstFrameIndex = FMath::Max(0, int32(FMath::RoundToZero(IntervalStart / FLoadTimeProfilerProvider::LoaderFrameLength)) - 1);
+	int32 LastFrameIndex = FMath::Min(int32(Frames.Num()) - 1, int32(FMath::RoundToZero(IntervalEnd / FLoadTimeProfilerProvider::LoaderFrameLength)) + 1);
+	for (int32 FrameIndex = FirstFrameIndex; FrameIndex <= LastFrameIndex; ++FrameIndex)
+	{
+		double Time = (FrameIndex + 0.5) * FLoadTimeProfilerProvider::LoaderFrameLength;
+		if (bIncludeExternalBounds || (IntervalStart <= Time && Time <= IntervalEnd))
+		{
+			const FLoaderFrame& Frame = Frames[FrameIndex];
+			switch (Type)
+			{
+			case LoaderFrameCounterType_LoaderThroughput:
+				Callback(Time, double(Frame.HeaderLoadedBytes + Frame.ExportLoadedBytes) / FLoadTimeProfilerProvider::LoaderFrameLength);
+				break;
+			case LoaderFrameCounterType_IoDispatcherThroughput:
+				Callback(Time, double(Frame.IoDispatcherReadBytes) / FLoadTimeProfilerProvider::LoaderFrameLength);
+				break;
+			}
+		}
+	}
+}
+
+FLoadTimeProfilerProvider::FLoadTimeProfilerProvider(IAnalysisSession& InSession, ICounterProvider& InCounterProvider)
 	: Session(InSession)
-	, ClassInfos(Session.GetLinearAllocator(), 4096)
-	, Requests(Session.GetLinearAllocator(), 4096)
-	, Packages(Session.GetLinearAllocator(), 4096)
-	, Exports(Session.GetLinearAllocator(), 4096)
-	, MainThreadCpuTimeline(MakeShared<CpuTimelineInternal>(Session.GetLinearAllocator()))
-	, AsyncLoadingThreadCpuTimeline(MakeShared<CpuTimelineInternal>(Session.GetLinearAllocator()))
+	, CounterProvider(InCounterProvider)
+	, ClassInfos(Session.GetLinearAllocator(), 16384)
+	, Requests(Session.GetLinearAllocator(), 16384)
+	, Packages(Session.GetLinearAllocator(), 16384)
+	, PackageLoads(Session.GetLinearAllocator(), 16384)
+	, IoDispatcherBatches(Session.GetLinearAllocator(), 16384)
+	, Exports(Session.GetLinearAllocator(), 16384)
 	, RequestsTable(Requests)
+	, Frames(Session.GetLinearAllocator(), 4096)
 {
 	RequestsTable.EditLayout().
 		AddColumn(&FLoadRequest::Name, TEXT("Name")).
@@ -52,11 +133,6 @@ FLoadTimeProfilerProvider::FLoadTimeProfilerProvider(IAnalysisSession& InSession
 				return Row.PackageInfo->Name;
 			},
 			TEXT("Package")).
-		AddColumn<const TCHAR*>([](const FPackagesTableRow& Row)
-			{
-				return GetLoadTimeProfilerPackageEventTypeString(Row.EventType);
-			},
-			TEXT("EventType")).
 		AddColumn(&FPackagesTableRow::SerializedHeaderSize, TEXT("SerializedHeaderSize")).
 		AddColumn(&FPackagesTableRow::SerializedExportsCount, TEXT("SerializedExportsCount")).
 		AddColumn(&FPackagesTableRow::SerializedExportsSize, TEXT("SerializedExportsSize")).
@@ -84,35 +160,50 @@ FLoadTimeProfilerProvider::FLoadTimeProfilerProvider(IAnalysisSession& InSession
 		AddColumn(&FExportsTableRow::AsyncLoadingThreadTime, TEXT("AsyncLoadingThreadTime"));
 }
 
-void FLoadTimeProfilerProvider::ReadMainThreadCpuTimeline(TFunctionRef<void(const CpuTimeline &)> Callback) const
+bool FLoadTimeProfilerProvider::GetCpuThreadTimelineIndex(uint32 ThreadId, uint32& OutTimelineIndex) const
 {
-	Session.ReadAccessCheck();
-	Callback(*MainThreadCpuTimeline);
+	const uint32* FindTimelineIndex = CpuTimelinesThreadMap.Find(ThreadId);
+	if (FindTimelineIndex)
+	{
+		OutTimelineIndex = *FindTimelineIndex;
+		return true;
+	}
+	else
+	{
+		return false;
+	}
 }
 
-void FLoadTimeProfilerProvider::ReadAsyncLoadingThreadCpuTimeline(TFunctionRef<void(const CpuTimeline &)> Callback) const
+bool FLoadTimeProfilerProvider::ReadTimeline(uint32 Index, TFunctionRef<void(const CpuTimeline &)> Callback) const
 {
 	Session.ReadAccessCheck();
-	Callback(*AsyncLoadingThreadCpuTimeline);
+	if (Index >= uint32(CpuTimelines.Num()))
+	{
+		return false;
+	}
+	Callback(*CpuTimelines[Index]);
+	return true;
 }
 
 ITable<FLoadTimeProfilerAggregatedStats>* FLoadTimeProfilerProvider::CreateEventAggregation(double IntervalStart, double IntervalEnd) const
 {
 	TArray<const CpuTimelineInternal*> Timelines;
-	Timelines.Add(&MainThreadCpuTimeline.Get());
-	Timelines.Add(&AsyncLoadingThreadCpuTimeline.Get());
+	for (const TSharedRef<CpuTimelineInternal>& Timeline : CpuTimelines)
+	{
+		Timelines.Add(&Timeline.Get());
+	}
 
 	auto BucketMapper = [](const FLoadTimeProfilerCpuEvent& Event)
 	{
-		return Event.PackageEventType;
+		return Event.Package;
 	};
-	TMap<ELoadTimeProfilerPackageEventType, FAggregatedTimingStats> Aggregation;
+	TMap<const FPackageInfo*, FAggregatedTimingStats> Aggregation;
 	FTimelineStatistics::CreateAggregation(Timelines, BucketMapper, IntervalStart, IntervalEnd, Aggregation);
 	TTable<FLoadTimeProfilerAggregatedStats>* Table = new TTable<FLoadTimeProfilerAggregatedStats>(AggregatedStatsTableLayout);
 	for (const auto& KV : Aggregation)
 	{
 		FLoadTimeProfilerAggregatedStats& Row = Table->AddRow();
-		Row.Name = GetLoadTimeProfilerPackageEventTypeString(KV.Key);
+		Row.Name = KV.Key ? KV.Key->Name : TEXT("[unknown]");
 		const FAggregatedTimingStats& Stats = KV.Value;
 		Row.Count = Stats.InstanceCount;
 		Row.Total = Stats.TotalExclusiveTime;
@@ -127,8 +218,10 @@ ITable<FLoadTimeProfilerAggregatedStats>* FLoadTimeProfilerProvider::CreateEvent
 ITable<FLoadTimeProfilerAggregatedStats>* FLoadTimeProfilerProvider::CreateObjectTypeAggregation(double IntervalStart, double IntervalEnd) const
 {
 	TArray<const CpuTimelineInternal*> Timelines;
-	Timelines.Add(&MainThreadCpuTimeline.Get());
-	Timelines.Add(&AsyncLoadingThreadCpuTimeline.Get());
+	for (const TSharedRef<CpuTimelineInternal>& Timeline : CpuTimelines)
+	{
+		Timelines.Add(&Timeline.Get());
+	}
 
 	auto BucketMapper = [](const FLoadTimeProfilerCpuEvent& Event) -> const FClassInfo*
 	{
@@ -160,37 +253,32 @@ ITable<FPackagesTableRow>* FLoadTimeProfilerProvider::CreatePackageDetailsTable(
 {
 	TTable<FPackagesTableRow>* Table = new TTable<FPackagesTableRow>(PackagesTableLayout);
 
-	TMap<TTuple<const FPackageInfo*, ELoadTimeProfilerPackageEventType>, FPackagesTableRow*> PackagesMap;
+	TMap<const FPackageInfo*, FPackagesTableRow*> PackagesMap;
 
 	auto FindRow = [Table, &PackagesMap](const FLoadTimeProfilerCpuEvent& Event) -> FPackagesTableRow*
 	{
 		if (Event.Package)
 		{
-			auto Key = MakeTuple(Event.Package, Event.PackageEventType);
-			FPackagesTableRow** FindIt = PackagesMap.Find(Key);
+			FPackagesTableRow** FindIt = PackagesMap.Find(Event.Package);
 			FPackagesTableRow* Row = nullptr;
 			if (!FindIt)
 			{
 				Row = &Table->AddRow();
 				Row->PackageInfo = Event.Package;
-				Row->EventType = Event.PackageEventType;
-				PackagesMap.Add(Key, Row);
+				PackagesMap.Add(Event.Package, Row);
 			}
 			else
 			{
 				Row = *FindIt;
 			}
 
-			if (Event.Export && Event.ExportEventType == LoadTimeProfilerObjectEventType_Serialize)
+			if (Event.Export && Event.EventType == LoadTimeProfilerObjectEventType_Serialize)
 			{
 				++Row->SerializedExportsCount;
 				Row->SerializedExportsSize += Event.Export->SerialSize;
 			}
 
-			if (!Event.Export && Event.PackageEventType == LoadTimeProfilerPackageEventType_CreateLinker)
-			{
-				Row->SerializedHeaderSize += Event.Package->Summary.TotalHeaderSize;
-			}
+			Row->SerializedHeaderSize = Event.Package->Summary.TotalHeaderSize;
 
 			return Row;
 		}
@@ -200,23 +288,25 @@ ITable<FPackagesTableRow>* FLoadTimeProfilerProvider::CreatePackageDetailsTable(
 		}
 	};
 
-	AsyncLoadingThreadCpuTimeline->EnumerateEvents(IntervalStart, IntervalEnd, [Table, &PackagesMap, FindRow](double StartTime, double EndTime, uint32, const FLoadTimeProfilerCpuEvent& Event)
-	{
-		FPackagesTableRow* Row = FindRow(Event);
-		if (Row)
-		{
-			Row->AsyncLoadingThreadTime += EndTime - StartTime; // TODO: Should be exclusive time
-		}
-	});
-	MainThreadCpuTimeline->EnumerateEvents(IntervalStart, IntervalEnd, [Table, &PackagesMap, FindRow](double StartTime, double EndTime, uint32, const FLoadTimeProfilerCpuEvent& Event)
+	for (int32 TimelineIndex = 0, TimelineCount = CpuTimelines.Num(); TimelineIndex < TimelineCount; ++TimelineIndex)
 	{
-		FPackagesTableRow* Row = FindRow(Event);
-		if (Row)
+		CpuTimelines[TimelineIndex]->EnumerateEvents(IntervalStart, IntervalEnd, [Table, &PackagesMap, FindRow, TimelineIndex](double StartTime, double EndTime, uint32, const FLoadTimeProfilerCpuEvent& Event)
 		{
-			Row->MainThreadTime += EndTime - StartTime; // TODO: Should be exclusive time
-		}
-	});
-
+			FPackagesTableRow* Row = FindRow(Event);
+			if (Row)
+			{
+				if (TimelineIndex == 0)
+				{
+					Row->MainThreadTime += EndTime - StartTime; // TODO: Should be exclusive time
+				}
+				else
+				{
+					Row->AsyncLoadingThreadTime += EndTime - StartTime; // TODO: Should be exclusive time
+				}
+			}
+		});
+	}
+	
 	return Table;
 }
 
@@ -230,14 +320,14 @@ ITable<FExportsTableRow>* FLoadTimeProfilerProvider::CreateExportDetailsTable(do
 	{
 		if (Event.Export)
 		{
-			auto Key = MakeTuple(Event.Export, Event.ExportEventType);
+			auto Key = MakeTuple(Event.Export, Event.EventType);
 			FExportsTableRow** FindIt = ExportsMap.Find(Key);
 			FExportsTableRow* Row = nullptr;
 			if (!FindIt)
 			{
 				Row = &Table->AddRow();
 				Row->ExportInfo = Event.Export;
-				Row->EventType = Event.ExportEventType;
+				Row->EventType = Event.EventType;
 				ExportsMap.Add(Key, Row);
 			}
 			else
@@ -245,7 +335,7 @@ ITable<FExportsTableRow>* FLoadTimeProfilerProvider::CreateExportDetailsTable(do
 				Row = *FindIt;
 			}
 
-			if (Event.ExportEventType == LoadTimeProfilerObjectEventType_Serialize)
+			if (Event.EventType == LoadTimeProfilerObjectEventType_Serialize)
 			{
 				Row->SerializedSize += Event.Export->SerialSize;
 			}
@@ -258,23 +348,24 @@ ITable<FExportsTableRow>* FLoadTimeProfilerProvider::CreateExportDetailsTable(do
 		}
 	};
 
-	AsyncLoadingThreadCpuTimeline->EnumerateEvents(IntervalStart, IntervalEnd, [Table, &ExportsMap, FindRow](double StartTime, double EndTime, uint32, const FLoadTimeProfilerCpuEvent& Event)
+	for (int32 TimelineIndex = 0, TimelineCount = CpuTimelines.Num(); TimelineIndex < TimelineCount; ++TimelineIndex)
 	{
-		FExportsTableRow* Row = FindRow(Event);
-		if (Row)
+		CpuTimelines[TimelineIndex]->EnumerateEvents(IntervalStart, IntervalEnd, [Table, &ExportsMap, FindRow, TimelineIndex](double StartTime, double EndTime, uint32, const FLoadTimeProfilerCpuEvent& Event)
 		{
-			Row->AsyncLoadingThreadTime += EndTime - StartTime; // TODO: Should be exclusive time
-		}
-	});
-	MainThreadCpuTimeline->EnumerateEvents(IntervalStart, IntervalEnd, [Table, &ExportsMap, FindRow](double StartTime, double EndTime, uint32, const FLoadTimeProfilerCpuEvent& Event)
-	{
-		FExportsTableRow* Row = FindRow(Event);
-		if (Row)
-		{
-			Row->MainThreadTime += EndTime - StartTime; // TODO: Should be exclusive time
-		}
-	});
-
+			FExportsTableRow* Row = FindRow(Event);
+			if (Row)
+			{
+				if (TimelineIndex == 0)
+				{
+					Row->MainThreadTime += EndTime - StartTime; // TODO: Should be exclusive time
+				}
+				else
+				{
+					Row->AsyncLoadingThreadTime += EndTime - StartTime; // TODO: Should be exclusive time
+				}
+			}
+		});
+	}
 	return Table;
 }
 
@@ -295,7 +386,7 @@ Trace::FLoadRequest& FLoadTimeProfilerProvider::CreateRequest()
 	return RequestInfo;
 }
 
-Trace::FPackageInfo& FLoadTimeProfilerProvider::CreatePackage(const TCHAR* PackageName)
+Trace::FPackageInfo& FLoadTimeProfilerProvider::EditPackageInfo(const TCHAR* PackageName)
 {
 	Session.WriteAccessCheck();
 
@@ -306,6 +397,120 @@ Trace::FPackageInfo& FLoadTimeProfilerProvider::CreatePackage(const TCHAR* Packa
 	return Package;
 }
 
+void FLoadTimeProfilerProvider::DistributeBytesAcrossFrames(uint64 ByteCount, double StartTime, double EndTime, uint64 FLoaderFrame::* FrameVariable)
+{
+	uint64 OriginalByteCount = ByteCount;
+
+	int32 BeginFrameIndex = int32(FMath::RoundToZero(StartTime / LoaderFrameLength));
+	int32 EndFrameIndex = int32(FMath::RoundToZero(EndTime / LoaderFrameLength));
+	while (Frames.Num() <= EndFrameIndex + 1)
+	{
+		Frames.PushBack();
+	}
+
+	if (BeginFrameIndex == EndFrameIndex)
+	{
+		FLoaderFrame& Frame = Frames[BeginFrameIndex];
+		Frame.*FrameVariable += ByteCount;
+		return;
+	}
+	
+	double TotalTime = EndTime - StartTime;
+	double TimeInFirstFrame = (BeginFrameIndex + 1) * LoaderFrameLength - StartTime;
+	check(TimeInFirstFrame >= 0.0);
+	uint64 BytesInFirstFrame = uint64(FMath::RoundToZero(TimeInFirstFrame / TotalTime * ByteCount));
+	Frames[BeginFrameIndex].*FrameVariable += BytesInFirstFrame;
+
+	uint64 TotalDistributed = BytesInFirstFrame;
+
+	double TimeInLastFrame = EndTime - EndFrameIndex * LoaderFrameLength;
+	check(TimeInLastFrame >= 0.0);
+	uint64 BytesInLastFrame = uint64(FMath::RoundToZero(TimeInLastFrame / TotalTime * ByteCount));
+
+	check(BytesInFirstFrame + BytesInLastFrame <= ByteCount);
+	ByteCount -= BytesInFirstFrame + BytesInLastFrame;
+
+	int32 FullFramesLeft = EndFrameIndex - BeginFrameIndex - 1;
+	for (int32 FrameIndex = BeginFrameIndex + 1; FrameIndex < EndFrameIndex; ++FrameIndex)
+	{
+		uint64 BytesInFrame = ByteCount / FullFramesLeft;
+		Frames[FrameIndex].*FrameVariable += BytesInFrame;
+
+		TotalDistributed += BytesInFrame;
+
+		ByteCount -= BytesInFrame;
+		--FullFramesLeft;
+	}
+
+	BytesInLastFrame += ByteCount;
+	Frames[EndFrameIndex].*FrameVariable += BytesInLastFrame;
+	TotalDistributed += BytesInLastFrame;
+	
+	check(TotalDistributed == OriginalByteCount);
+}
+
+uint64 FLoadTimeProfilerProvider::BeginLoadPackage(const FPackageInfo& PackageInfo, double Time)
+{
+	Session.WriteAccessCheck();
+
+	CreateCounters();
+
+	LoadingPackagesCounter->AddValue(Time, int64(1));
+
+	uint64 LoadHandle = PackageLoads.Num();
+	FPackageLoad& PackageLoad = PackageLoads.PushBack();
+	PackageLoad.Package = &PackageInfo;
+	PackageLoad.StartTime = Time;
+	return LoadHandle;
+}
+
+void FLoadTimeProfilerProvider::EndLoadPackage(uint64 LoadHandle, double Time)
+{
+	Session.WriteAccessCheck();
+
+	FPackageLoad& PackageLoad = PackageLoads[LoadHandle];
+	PackageLoad.EndTime = Time;
+	check(PackageLoad.EndTime >= PackageLoad.StartTime);
+	
+	DistributeBytesAcrossFrames(PackageLoad.Package->Summary.TotalHeaderSize, PackageLoad.StartTime, PackageLoad.EndTime, &FLoaderFrame::HeaderLoadedBytes);
+	DistributeBytesAcrossFrames(PackageLoad.Package->TotalExportsSerialSize, PackageLoad.StartTime, PackageLoad.EndTime, &FLoaderFrame::ExportLoadedBytes);
+
+	TotalLoaderBytesLoadedCounter->AddValue(Time, int64(PackageLoad.Package->Summary.TotalHeaderSize + PackageLoad.Package->TotalExportsSerialSize));
+	LoadingPackagesCounter->AddValue(Time, int64(-1));
+}
+
+uint64 FLoadTimeProfilerProvider::BeginIoDispatcherBatch(uint64 BatchId, double Time)
+{
+	Session.WriteAccessCheck();
+
+	CreateCounters();
+
+	ActiveIoDispatcherBatchesCounter->AddValue(Time, int64(1));
+
+	uint64 BatchHandle = IoDispatcherBatches.Num();
+
+	FIoDispatcherBatch& Batch = IoDispatcherBatches.PushBack();
+	Batch.StartTime = Time;
+
+	return BatchHandle;
+}
+
+void FLoadTimeProfilerProvider::EndIoDispatcherBatch(uint64 BatchHandle, double Time, uint64 TotalSize)
+{
+	Session.WriteAccessCheck();
+
+	FIoDispatcherBatch& Batch = IoDispatcherBatches[BatchHandle];
+	Batch.EndTime = Time;
+	Batch.TotalSize = TotalSize;
+
+	check(Batch.EndTime >= Batch.StartTime);
+
+	DistributeBytesAcrossFrames(Batch.TotalSize, Batch.StartTime, Batch.EndTime, &FLoaderFrame::IoDispatcherReadBytes);
+
+	TotalIoDispatcherBytesReadCounter->AddValue(Time, int64(Batch.TotalSize));
+	ActiveIoDispatcherBatchesCounter->AddValue(Time, int64(-1));
+}
+
 Trace::FPackageExportInfo& FLoadTimeProfilerProvider::CreateExport()
 {
 	Session.WriteAccessCheck();
@@ -316,16 +521,19 @@ Trace::FPackageExportInfo& FLoadTimeProfilerProvider::CreateExport()
 	return Export;
 }
 
-FLoadTimeProfilerProvider::CpuTimelineInternal& FLoadTimeProfilerProvider::EditAdditionalCpuTimeline(uint32 ThreadId)
+FLoadTimeProfilerProvider::CpuTimelineInternal& FLoadTimeProfilerProvider::EditCpuTimeline(uint32 ThreadId)
 {
-	if (AdditionalCpuTimelinesMap.Contains(ThreadId))
+	uint32* FindTimelineIndex = CpuTimelinesThreadMap.Find(ThreadId);
+	if (FindTimelineIndex)
 	{
-		return AdditionalCpuTimelinesMap[ThreadId].Get();
+		return CpuTimelines[*FindTimelineIndex].Get();
 	}
 	else
 	{
 		TSharedRef<CpuTimelineInternal> Timeline = MakeShared<CpuTimelineInternal>(Session.GetLinearAllocator());
-		AdditionalCpuTimelinesMap.Add(ThreadId, Timeline);
+		uint32 TimelineIndex = CpuTimelines.Num();
+		CpuTimelinesThreadMap.Add(ThreadId, TimelineIndex);
+		CpuTimelines.Add(Timeline);
 		return Timeline.Get();
 	}
 }
@@ -344,38 +552,31 @@ uint64 FLoadTimeProfilerProvider::PackageSizeSum(const FLoadRequest& Row)
 	return Sum;
 }
 
-const TCHAR* GetLoadTimeProfilerPackageEventTypeString(ELoadTimeProfilerPackageEventType EventType)
+void FLoadTimeProfilerProvider::CreateCounters()
 {
-	switch (EventType)
+	if (!bHasCreatedCounters)
 	{
-	case LoadTimeProfilerPackageEventType_CreateLinker:
-		return TEXT("CreateLinker");
-	case LoadTimeProfilerPackageEventType_FinishLinker:
-		return TEXT("FinishLinker");
-	case LoadTimeProfilerPackageEventType_StartImportPackages:
-		return TEXT("StartImportPackages");
-	case LoadTimeProfilerPackageEventType_SetupImports:
-		return TEXT("SetupImports");
-	case LoadTimeProfilerPackageEventType_SetupExports:
-		return TEXT("SetupExports");
-	case LoadTimeProfilerPackageEventType_ProcessImportsAndExports:
-		return TEXT("ProcessImportsAndExports");
-	case LoadTimeProfilerPackageEventType_ExportsDone:
-		return TEXT("ExportsDone");
-	case LoadTimeProfilerPackageEventType_PostLoadWait:
-		return TEXT("PostLoadWait");
-	case LoadTimeProfilerPackageEventType_StartPostLoad:
-		return TEXT("StartPostLoad");
-	case LoadTimeProfilerPackageEventType_Tick:
-		return TEXT("Tick");
-	case LoadTimeProfilerPackageEventType_DeferredPostLoad:
-		return TEXT("DeferredPostLoad");
-	case LoadTimeProfilerPackageEventType_Finish:
-		return TEXT("Finish");
-	case LoadTimeProfilerPackageEventType_None:
-		return TEXT("None");
+		for (int32 CounterType = 0; CounterType < FLoaderFrameCounter::LoaderFrameCounterType_Count; ++CounterType)
+		{
+			CounterProvider.AddCounter(new FLoaderFrameCounter(FLoaderFrameCounter::ELoaderFrameCounterType(CounterType), Frames));
+		}
+	
+		ActiveIoDispatcherBatchesCounter = CounterProvider.CreateCounter();
+		ActiveIoDispatcherBatchesCounter->SetName(TEXT("AssetLoading/IoDispatcher/ActiveBatches"));
+
+		TotalIoDispatcherBytesReadCounter = CounterProvider.CreateCounter();
+		TotalIoDispatcherBytesReadCounter->SetName(TEXT("AssetLoading/IoDispatcher/TotalBytesRead"));
+		TotalIoDispatcherBytesReadCounter->SetDisplayHint(CounterDisplayHint_Memory);
+
+		LoadingPackagesCounter = CounterProvider.CreateCounter();
+		LoadingPackagesCounter->SetName(TEXT("AssetLoading/AsyncLoading/ActiveLoadingPackages"));
+
+		TotalLoaderBytesLoadedCounter = CounterProvider.CreateCounter();
+		TotalLoaderBytesLoadedCounter->SetName(TEXT("AssetLoading/AsyncLoading/TotalBytesLoaded"));
+		TotalLoaderBytesLoadedCounter->SetDisplayHint(CounterDisplayHint_Memory);
+
+		bHasCreatedCounters = true;
 	}
-	return TEXT("[invalid]");
 }
 
 const TCHAR* GetLoadTimeProfilerObjectEventTypeString(ELoadTimeProfilerObjectEventType EventType)
diff --git a/Engine/Source/Developer/TraceServices/Private/Model/LoadTimeProfilerPrivate.h b/Engine/Source/Developer/TraceServices/Private/Model/LoadTimeProfilerPrivate.h
--- a/Engine/Source/Developer/TraceServices/Private/Model/LoadTimeProfilerPrivate.h
+++ b/Engine/Source/Developer/TraceServices/Private/Model/LoadTimeProfilerPrivate.h
@@ -21,9 +21,10 @@ class FLoadTimeProfilerProvider
 public:
 	typedef TMonotonicTimeline<FLoadTimeProfilerCpuEvent> CpuTimelineInternal;
 
-	FLoadTimeProfilerProvider(IAnalysisSession& Session);
-	virtual void ReadMainThreadCpuTimeline(TFunctionRef<void(const CpuTimeline&)> Callback) const override;
-	virtual void ReadAsyncLoadingThreadCpuTimeline(TFunctionRef<void(const CpuTimeline&)> Callback) const override;
+	FLoadTimeProfilerProvider(IAnalysisSession& Session, ICounterProvider& CounterProvider);
+	virtual uint64 GetTimelineCount() const override { return CpuTimelines.Num(); }
+	virtual bool GetCpuThreadTimelineIndex(uint32 ThreadId, uint32& OutTimelineIndex) const override;
+	virtual bool ReadTimeline(uint32 Index, TFunctionRef<void(const CpuTimeline&)> Callback) const override;
 	virtual ITable<FLoadTimeProfilerAggregatedStats>* CreateEventAggregation(double IntervalStart, double IntervalEnd) const override;
 	virtual ITable<FLoadTimeProfilerAggregatedStats>* CreateObjectTypeAggregation(double IntervalStart, double IntervalEnd) const override;
 	virtual ITable<FPackagesTableRow>* CreatePackageDetailsTable(double IntervalStart, double IntervalEnd) const override;
@@ -31,33 +32,88 @@ public:
 	virtual const ITable<FLoadRequest>& GetRequestsTable() const override { return RequestsTable; }
 	const FClassInfo& AddClassInfo(const TCHAR* ClassName);
 	FLoadRequest& CreateRequest();
-	FPackageInfo& CreatePackage(const TCHAR* PackageName);
+	FPackageInfo& EditPackageInfo(const TCHAR* PackageName);
+	uint64 BeginLoadPackage(const FPackageInfo& PackageInfo, double Time);
+	void EndLoadPackage(uint64 LoadHandle, double Time);
 	FPackageExportInfo& CreateExport();
-	CpuTimelineInternal& EditMainThreadCpuTimeline() { return MainThreadCpuTimeline.Get(); }
-	CpuTimelineInternal& EditAsyncLoadingThreadCpuTimeline() { return AsyncLoadingThreadCpuTimeline.Get(); }
-	CpuTimelineInternal& EditAdditionalCpuTimeline(uint32 ThreadId);
-	virtual uint32 GetMainThreadId() const override { return MainThreadId; }
-	void SetMainThreadId(uint32 ThreadId) { MainThreadId = ThreadId; }
-	virtual uint32 GetAsyncLoadingThreadId() const override { return AsyncLoadingThreadId; }
-	void SetAsyncLoadingThreadId(uint32 ThreadId) { AsyncLoadingThreadId = ThreadId; }
+	CpuTimelineInternal& EditCpuTimeline(uint32 ThreadId);
+	uint64 BeginIoDispatcherBatch(uint64 BatchId, double Time);
+	void EndIoDispatcherBatch(uint64 BatchHandle, double Time, uint64 TotalSize);
 
 private:
+	struct FLoaderFrame;
+
 	static uint64 PackageSizeSum(const FLoadRequest& Row);
+	void CreateCounters();
+	void DistributeBytesAcrossFrames(uint64 ByteCount, double StartTime, double EndTime, uint64 FLoaderFrame::* FrameVariable);
+
+	static constexpr double LoaderFrameLength = 1.0 / 60.0;
+
+	struct FPackageLoad
+	{
+		const FPackageInfo* Package = nullptr;
+		double StartTime = 0.0;
+		double EndTime = 0.0;
+	};
+
+	struct FIoDispatcherBatch
+	{
+		double StartTime = 0.0;
+		double EndTime = 0.0;
+		uint64 TotalSize = 0;
+	};
+
+	struct FLoaderFrame
+	{
+		uint64 IoDispatcherReadBytes = 0;
+		uint64 HeaderLoadedBytes = 0;
+		uint64 ExportLoadedBytes = 0;
+	};
+
+	class FLoaderFrameCounter
+		: public ICounter
+	{
+	public:
+		enum ELoaderFrameCounterType
+		{
+			LoaderFrameCounterType_IoDispatcherThroughput,
+			LoaderFrameCounterType_LoaderThroughput,
+			LoaderFrameCounterType_Count
+		};
+
+		FLoaderFrameCounter(ELoaderFrameCounterType Type, const TPagedArray<FLoaderFrame>& Frames);
+		virtual const TCHAR* GetName() const override;
+		virtual const TCHAR* GetDescription() const override;
+		virtual bool IsFloatingPoint() const override;
+		virtual ECounterDisplayHint GetDisplayHint() const override;
+		virtual void EnumerateValues(double IntervalStart, double IntervalEnd, bool bIncludeExternalBounds, TFunctionRef<void(double, int64)> Callback) const override;
+		virtual void EnumerateFloatValues(double IntervalStart, double IntervalEnd, bool bIncludeExternalBounds, TFunctionRef<void(double, double)> Callback) const override;
+
+	private:
+		const TPagedArray<FLoaderFrame>& Frames;
+		ELoaderFrameCounterType Type;
+	};
 
 	IAnalysisSession& Session;
+	ICounterProvider& CounterProvider;
 	TPagedArray<FClassInfo> ClassInfos;
 	TPagedArray<FLoadRequest> Requests;
 	TPagedArray<FPackageInfo> Packages;
+	TPagedArray<FPackageLoad> PackageLoads;
+	TPagedArray<FIoDispatcherBatch> IoDispatcherBatches;
 	TPagedArray<FPackageExportInfo> Exports;
-	TSharedRef<CpuTimelineInternal> MainThreadCpuTimeline;
-	TSharedRef<CpuTimelineInternal> AsyncLoadingThreadCpuTimeline;
-	TMap<uint32, TSharedRef<CpuTimelineInternal>> AdditionalCpuTimelinesMap;
-	uint32 MainThreadId = uint32(-1);
-	uint32 AsyncLoadingThreadId = uint32(-1);
+	TArray<TSharedRef<CpuTimelineInternal>> CpuTimelines;
+	TMap<uint32, uint32> CpuTimelinesThreadMap;
 	TTableView<FLoadRequest> RequestsTable;
 	TTableLayout<FLoadTimeProfilerAggregatedStats> AggregatedStatsTableLayout;
 	TTableLayout<FPackagesTableRow> PackagesTableLayout;
 	TTableLayout<FExportsTableRow> ExportsTableLayout;
+	TPagedArray<FLoaderFrame> Frames;
+	bool bHasCreatedCounters = false;
+	IEditableCounter* ActiveIoDispatcherBatchesCounter = nullptr;
+	IEditableCounter* TotalIoDispatcherBytesReadCounter = nullptr;
+	IEditableCounter* LoadingPackagesCounter = nullptr;
+	IEditableCounter* TotalLoaderBytesLoadedCounter = nullptr;
 };
 
 }
diff --git a/Engine/Source/Developer/TraceServices/Private/Modules/LoadTimeProfilerModule.cpp b/Engine/Source/Developer/TraceServices/Private/Modules/LoadTimeProfilerModule.cpp
--- a/Engine/Source/Developer/TraceServices/Private/Modules/LoadTimeProfilerModule.cpp
+++ b/Engine/Source/Developer/TraceServices/Private/Modules/LoadTimeProfilerModule.cpp
@@ -35,7 +35,7 @@ void FLoadTimeProfilerModule::GetModuleInfo(FModuleInfo& OutModuleInfo)
 
 void FLoadTimeProfilerModule::OnAnalysisBegin(IAnalysisSession& Session)
 {
-	FLoadTimeProfilerProvider* LoadTimeProfilerProvider = new FLoadTimeProfilerProvider(Session);
+	FLoadTimeProfilerProvider* LoadTimeProfilerProvider = new FLoadTimeProfilerProvider(Session, EditCounterProvider(Session));
 	Session.AddProvider(LoadTimeProfilerProviderName, LoadTimeProfilerProvider);
 	Session.AddAnalyzer(new FAsyncLoadingTraceAnalyzer(Session, *LoadTimeProfilerProvider));
 	FFileActivityProvider* FileActivityProvider = new FFileActivityProvider(Session);
diff --git a/Engine/Source/Developer/TraceServices/Public/TraceServices/Model/LoadTimeProfiler.h b/Engine/Source/Developer/TraceServices/Public/TraceServices/Model/LoadTimeProfiler.h
--- a/Engine/Source/Developer/TraceServices/Public/TraceServices/Model/LoadTimeProfiler.h
+++ b/Engine/Source/Developer/TraceServices/Public/TraceServices/Model/LoadTimeProfiler.h
@@ -56,7 +56,6 @@ public:
 struct FPackageSummaryInfo
 {
 	uint32 TotalHeaderSize = 0;
-	uint32 NameCount = 0;
 	uint32 ImportCount = 0;
 	uint32 ExportCount = 0;
 };
@@ -72,10 +71,8 @@ struct FPackageExportInfo
 {
 	uint32 Id;
 	const FClassInfo* Class = nullptr;
-	uint64 SerialOffset = 0;
 	uint64 SerialSize = 0;
 	const FPackageInfo* Package = nullptr;
-	bool IsAsset = false;
 };
 
 struct FPackageInfo
@@ -84,6 +81,7 @@ struct FPackageInfo
 	const TCHAR* Name = nullptr;
 	FPackageSummaryInfo Summary;
 	TArray<const FPackageExportInfo*> Exports;
+	uint64 TotalExportsSerialSize = 0;
 };
 
 struct FLoadRequest
@@ -95,6 +93,14 @@ struct FLoadRequest
 	TArray<const FPackageInfo*> Packages;
 };
 
+enum ELoadTimeProfilerObjectEventType
+{
+	LoadTimeProfilerObjectEventType_Create,
+	LoadTimeProfilerObjectEventType_Serialize,
+	LoadTimeProfilerObjectEventType_PostLoad,
+	LoadTimeProfilerObjectEventType_None
+};
+
 struct FExportsTableRow
 {
 	const FPackageExportInfo* ExportInfo = nullptr;
@@ -112,15 +118,13 @@ struct FPackagesTableRow
 	uint64 SerializedExportsSize = 0;
 	double MainThreadTime = 0.0;
 	double AsyncLoadingThreadTime = 0.0;
-	ELoadTimeProfilerPackageEventType EventType = LoadTimeProfilerPackageEventType_None;
 };
 
 struct FLoadTimeProfilerCpuEvent
 {
 	const FPackageInfo* Package = nullptr;
 	const FPackageExportInfo* Export = nullptr;
-	ELoadTimeProfilerPackageEventType PackageEventType = LoadTimeProfilerPackageEventType_None;
-	ELoadTimeProfilerObjectEventType ExportEventType = LoadTimeProfilerObjectEventType_None;
+	ELoadTimeProfilerObjectEventType EventType = LoadTimeProfilerObjectEventType_None;
 };
 
 struct FLoadTimeProfilerAggregatedStats
@@ -134,7 +138,6 @@ struct FLoadTimeProfilerAggregatedStats
 	double Median;
 };
 
-TRACESERVICES_API const TCHAR* GetLoadTimeProfilerPackageEventTypeString(ELoadTimeProfilerPackageEventType EventType);
 TRACESERVICES_API const TCHAR* GetLoadTimeProfilerObjectEventTypeString(ELoadTimeProfilerObjectEventType EventType);
 
 class ILoadTimeProfilerProvider
@@ -144,10 +147,9 @@ public:
 	typedef ITimeline<FLoadTimeProfilerCpuEvent> CpuTimeline;
 
 	virtual ~ILoadTimeProfilerProvider() = default;
-	virtual uint32 GetMainThreadId() const = 0;
-	virtual void ReadMainThreadCpuTimeline(TFunctionRef<void(const CpuTimeline&)> Callback) const = 0;
-	virtual uint32 GetAsyncLoadingThreadId() const = 0;
-	virtual void ReadAsyncLoadingThreadCpuTimeline(TFunctionRef<void(const CpuTimeline&)> Callback) const = 0;
+	virtual uint64 GetTimelineCount() const = 0;
+	virtual bool GetCpuThreadTimelineIndex(uint32 ThreadId, uint32& OutTimelineIndex) const = 0;
+	virtual bool ReadTimeline(uint32 Index, TFunctionRef<void(const CpuTimeline&)> Callback) const = 0;
 	virtual ITable<FLoadTimeProfilerAggregatedStats>* CreateEventAggregation(double IntervalStart, double IntervalEnd) const = 0;
 	virtual ITable<FLoadTimeProfilerAggregatedStats>* CreateObjectTypeAggregation(double IntervalStart, double IntervalEnd) const = 0;
 	virtual ITable<FPackagesTableRow>* CreatePackageDetailsTable(double IntervalStart, double IntervalEnd) const = 0;
diff --git a/Engine/Source/Runtime/Core/Private/IO/IoDispatcher.cpp b/Engine/Source/Runtime/Core/Private/IO/IoDispatcher.cpp
--- a/Engine/Source/Runtime/Core/Private/IO/IoDispatcher.cpp
+++ b/Engine/Source/Runtime/Core/Private/IO/IoDispatcher.cpp
@@ -17,11 +17,25 @@
 #include "Misc/StringBuilder.h"
 #include "Misc/LazySingleton.h"
 #include "Misc/CoreDelegates.h"
+#include "Trace/Trace.h"
+
+#define IODISPATCHER_TRACE_ENABLED !UE_BUILD_SHIPPING
 
 DEFINE_LOG_CATEGORY_STATIC(LogIoDispatch, Log, All);
 
 //////////////////////////////////////////////////////////////////////////
 
+UE_TRACE_EVENT_BEGIN(IoDispatcher, BatchIssued, Always)
+	UE_TRACE_EVENT_FIELD(uint64, Cycle)
+	UE_TRACE_EVENT_FIELD(uint64, BatchId)
+UE_TRACE_EVENT_END()
+
+UE_TRACE_EVENT_BEGIN(IoDispatcher, BatchResolved, Always)
+	UE_TRACE_EVENT_FIELD(uint64, Cycle)
+	UE_TRACE_EVENT_FIELD(uint64, BatchId)
+	UE_TRACE_EVENT_FIELD(uint64, TotalSize)
+UE_TRACE_EVENT_END()
+
 template <typename T, uint32 BlockSize = 128>
 class TBlockAllocator
 {
@@ -1064,12 +1078,29 @@ public:
 		// At this point the batch is immutable and we should start
 		// doing the work.
 
-		IterateBatch(Batch, [this](FIoRequestImpl* Request)
+#if IODISPATCHER_TRACE_ENABLED
+		UE_TRACE_LOG(IoDispatcher, BatchIssued)
+			<< BatchIssued.Cycle(FPlatformTime::Cycles64())
+			<< BatchIssued.BatchId(uint64(Batch));
+		uint64 TotalBatchSize = 0;
+#endif
+
+		IterateBatch(Batch, [this, &TotalBatchSize](FIoRequestImpl* Request)
 		{
 			Request->Result = IoStore->Resolve(Request->ChunkId);
 
+#if IODISPATCHER_TRACE_ENABLED
+			TotalBatchSize += Request->Result.ValueOrDie().DataSize();
+#endif
 			return true;
 		});
+
+#if IODISPATCHER_TRACE_ENABLED
+		UE_TRACE_LOG(IoDispatcher, BatchResolved)
+			<< BatchResolved.Cycle(FPlatformTime::Cycles64())
+			<< BatchResolved.BatchId(uint64(Batch))
+			<< BatchResolved.TotalSize(TotalBatchSize);
+#endif
 	}
 
 	bool IsBatchReady(const FIoBatchImpl* Batch)
diff --git a/Engine/Source/Runtime/Core/Public/Serialization/LoadTimeTrace.h b/Engine/Source/Runtime/Core/Public/Serialization/LoadTimeTrace.h
--- a/Engine/Source/Runtime/Core/Public/Serialization/LoadTimeTrace.h
+++ b/Engine/Source/Runtime/Core/Public/Serialization/LoadTimeTrace.h
@@ -12,31 +12,6 @@
 #define LOADTIMEPROFILERTRACE_ENABLED 0
 #endif
 
-enum ELoadTimeProfilerPackageEventType
-{
-	LoadTimeProfilerPackageEventType_CreateLinker,
-	LoadTimeProfilerPackageEventType_FinishLinker,
-	LoadTimeProfilerPackageEventType_StartImportPackages,
-	LoadTimeProfilerPackageEventType_SetupImports,
-	LoadTimeProfilerPackageEventType_SetupExports,
-	LoadTimeProfilerPackageEventType_ProcessImportsAndExports,
-	LoadTimeProfilerPackageEventType_ExportsDone,
-	LoadTimeProfilerPackageEventType_PostLoadWait,
-	LoadTimeProfilerPackageEventType_StartPostLoad,
-	LoadTimeProfilerPackageEventType_Tick,
-	LoadTimeProfilerPackageEventType_Finish,
-	LoadTimeProfilerPackageEventType_DeferredPostLoad,
-	LoadTimeProfilerPackageEventType_None,
-};
-
-enum ELoadTimeProfilerObjectEventType
-{
-	LoadTimeProfilerObjectEventType_Create,
-	LoadTimeProfilerObjectEventType_Serialize,
-	LoadTimeProfilerObjectEventType_PostLoad,
-	LoadTimeProfilerObjectEventType_None
-};
-
 #if LOADTIMEPROFILERTRACE_ENABLED
 
 struct FLoadTimeProfilerTrace
diff --git a/Engine/Source/Runtime/CoreUObject/Private/Serialization/AsyncLoading.cpp b/Engine/Source/Runtime/CoreUObject/Private/Serialization/AsyncLoading.cpp
--- a/Engine/Source/Runtime/CoreUObject/Private/Serialization/AsyncLoading.cpp
+++ b/Engine/Source/Runtime/CoreUObject/Private/Serialization/AsyncLoading.cpp
@@ -1370,6 +1370,8 @@ FORCEINLINE static bool CanAddWaitingPackages(FAsyncLoadingThread& AsyncLoadingT
 
 void FAsyncLoadingThread::QueueEvent_CreateLinker(FAsyncPackage* Package, int32 EventSystemPriority)
 {
+	TRACE_LOADTIME_BEGIN_LOAD_ASYNC_PACKAGE(Package);
+
 	FileOpenLogActive();  // make sure GRandomizeLoadOrder is set up
 	check(Package);
 	Package->AddNode(EEventLoadNode::Package_LoadSummary);
@@ -1385,7 +1387,6 @@ void FAsyncLoadingThread::QueueEvent_CreateLinker(FAsyncPackage* Package, int32
 				check(Pkg);
 				if (Pkg)
 				{
-					TRACE_LOADTIME_ASYNC_PACKAGE_SCOPE(Pkg, LoadTimeProfilerPackageEventType_CreateLinker);
 					Pkg->SetTimeLimit(Args, TEXT("Create Linker"));
 					Pkg->Event_CreateLinker();
 					Args.OutLastObjectWorkWasPerformedOn = Pkg->GetLinkerRoot();
@@ -1450,7 +1451,6 @@ void FAsyncLoadingThread::QueueEvent_FinishLinker(FWeakAsyncPackagePtr WeakPtr,
 			check(PkgInner);
 			if (PkgInner)
 			{
-				TRACE_LOADTIME_ASYNC_PACKAGE_SCOPE(PkgInner, LoadTimeProfilerPackageEventType_FinishLinker);
 				PkgInner->SetTimeLimit(Args, TEXT("Finish Linker"));
 				PkgInner->Event_FinishLinker();
 			}
@@ -1569,6 +1569,8 @@ void FAsyncPackage::Event_FinishLinker()
 			}
 		}
 
+		TRACE_LOADTIME_PACKAGE_SUMMARY(this, Linker->Summary.TotalHeaderSize, Linker->Summary.ImportCount, Linker->Summary.ExportCount);
+
 		check(AsyncPackageLoadingState == EAsyncPackageLoadingState::WaitingForSummary);
 		AsyncPackageLoadingState = EAsyncPackageLoadingState::StartImportPackages;
 		AsyncLoadingThread.QueueEvent_StartImportPackages(this, FAsyncLoadEvent::EventSystemPriority_MAX - 1);
@@ -1597,7 +1599,6 @@ void FAsyncLoadingThread::QueueEvent_StartImportPackages(FAsyncPackage* Package,
 		FAsyncPackage* Pkg = GetPackage(WeakPtr);
 		if (Pkg)
 		{
-			TRACE_LOADTIME_ASYNC_PACKAGE_SCOPE(Pkg, LoadTimeProfilerPackageEventType_StartImportPackages);
 			Pkg->SetTimeLimit(Args, TEXT("Start Import Packages"));
 			Pkg->Event_StartImportPackages();
 		}
@@ -1883,7 +1884,6 @@ void FAsyncLoadingThread::QueueEvent_SetupImports(FAsyncPackage* Package, int32
 		FAsyncPackage* Pkg = GetPackage(WeakPtr);
 		if (Pkg)
 		{
-			TRACE_LOADTIME_ASYNC_PACKAGE_SCOPE(Pkg, LoadTimeProfilerPackageEventType_SetupImports);
 			Pkg->SetTimeLimit(Args, TEXT("Setup Imports"));
 			Pkg->Event_SetupImports();
 		}
@@ -2245,7 +2245,6 @@ void FAsyncLoadingThread::QueueEvent_SetupExports(FAsyncPackage* Package, int32
 		FAsyncPackage* Pkg = GetPackage(WeakPtr);
 		if (Pkg)
 		{
-			TRACE_LOADTIME_ASYNC_PACKAGE_SCOPE(Pkg, LoadTimeProfilerPackageEventType_SetupExports);
 			Pkg->SetTimeLimit(Args, TEXT("Setup Exports"));
 			Pkg->Event_SetupExports();
 		}
@@ -2282,7 +2281,6 @@ void FAsyncLoadingThread::QueueEvent_ProcessImportsAndExports(FAsyncPackage* Pac
 		FAsyncPackage* Pkg = GetPackage(WeakPtr);
 		if (Pkg)
 		{
-			TRACE_LOADTIME_ASYNC_PACKAGE_SCOPE(Pkg, LoadTimeProfilerPackageEventType_ProcessImportsAndExports);
 			Pkg->SetTimeLimit(Args, TEXT("ProcessImportsAndExports"));
 			Pkg->Event_ProcessImportsAndExports();
 		}
@@ -2304,7 +2302,6 @@ void FAsyncLoadingThread::QueueEvent_ProcessPostloadWait(FAsyncPackage* Package,
 		FAsyncPackage* Pkg = GetPackage(WeakPtr);
 		if (Pkg)
 		{
-			TRACE_LOADTIME_ASYNC_PACKAGE_SCOPE(Pkg, LoadTimeProfilerPackageEventType_PostLoadWait);
 			Pkg->SetTimeLimit(Args, TEXT("Process Process Postload Wait"));
 			Pkg->Event_ProcessPostloadWait();
 		}
@@ -2326,7 +2323,6 @@ void FAsyncLoadingThread::QueueEvent_ExportsDone(FAsyncPackage* Package, int32 E
 		FAsyncPackage* Pkg = GetPackage(WeakPtr);
 		if (Pkg)
 		{
-			TRACE_LOADTIME_ASYNC_PACKAGE_SCOPE(Pkg, LoadTimeProfilerPackageEventType_ExportsDone);
 			Pkg->SetTimeLimit(Args, TEXT("Exports Done"));
 			Pkg->Event_ExportsDone();
 		}
@@ -2348,7 +2344,6 @@ void FAsyncLoadingThread::QueueEvent_StartPostLoad(FAsyncPackage* Package, int32
 		FAsyncPackage* Pkg = GetPackage(WeakPtr);
 		if (Pkg)
 		{
-			TRACE_LOADTIME_ASYNC_PACKAGE_SCOPE(Pkg, LoadTimeProfilerPackageEventType_StartPostLoad);
 			Pkg->SetTimeLimit(Args, TEXT("Start Post Load"));
 			Pkg->Event_StartPostload();
 		}
@@ -2767,7 +2762,7 @@ void FAsyncPackage::EventDrivenCreateExport(int32 LocalExportIndex)
 	SCOPED_LOADTIMER(Package_CreateExports);
 	FObjectExport& Export = Linker->ExportMap[LocalExportIndex];
 
-	TRACE_LOADTIME_CREATE_EXPORT_SCOPE(Linker, &Export.Object, Export.SerialOffset, Export.SerialSize, Export.bIsAsset);
+	TRACE_LOADTIME_CREATE_EXPORT_SCOPE(this, &Export.Object);
 
 	LLM_SCOPE(ELLMTag::AsyncLoading);
 	LLM_SCOPED_TAG_WITH_OBJECT_IN_SET(GetLinkerRoot(), ELLMTagSet::Assets);
@@ -3151,7 +3146,7 @@ void FAsyncPackage::EventDrivenSerializeExport(int32 LocalExportIndex)
 
 		Object->ClearFlags(RF_NeedLoad);
 
-		TRACE_LOADTIME_OBJECT_SCOPE(Object, LoadTimeProfilerObjectEventType_Serialize);
+		TRACE_LOADTIME_SERIALIZE_EXPORT_SCOPE(Object, Export.SerialSize);
 
 		FUObjectSerializeContext* LoadContext = GetSerializeContext();
 		UObject* PrevSerializedObject = LoadContext->SerializedObject;
@@ -4377,6 +4372,8 @@ EAsyncPackageState::Type FAsyncLoadingThread::ProcessLoadedPackages(bool bUseTim
 				PackagesToDelete.AddUnique(Package);
 				Package->MarkRequestIDsAsComplete();
 
+				TRACE_LOADTIME_END_LOAD_ASYNC_PACKAGE(Package);
+
 				if (IsTimeLimitExceeded(TickStartTime, bUseTimeLimit, TimeLimit, TEXT("ProcessLoadedPackages Misc")) || (FlushTree && !ContainsRequestID(FlushTree->RequestId)))
 				{
 					// The only package we care about has finished loading, so we're good to exit
@@ -5197,7 +5194,6 @@ FAsyncPackage::FAsyncPackage(FAsyncLoadingThread& InThread, const FAsyncPackageD
 
 FAsyncPackage::~FAsyncPackage()
 {
-	TRACE_LOADTIME_DESTROY_ASYNC_PACKAGE(this);
 #if DO_CHECK
 	if (GEventDrivenLoaderEnabled)
 	{
@@ -5224,6 +5220,8 @@ FAsyncPackage::~FAsyncPackage()
 	}
 
 	EmptyReferencedObjects();
+
+	TRACE_LOADTIME_DESTROY_ASYNC_PACKAGE(this);
 }
 
 
@@ -5419,7 +5417,6 @@ EAsyncPackageState::Type FAsyncPackage::TickAsyncPackage(bool InbUseTimeLimit, b
 
 	ReentryCount++;
 
-	TRACE_LOADTIME_ASYNC_PACKAGE_SCOPE(this, LoadTimeProfilerPackageEventType_Tick);
 	SCOPE_CYCLE_COUNTER(STAT_FAsyncPackage_Tick);
 	SCOPED_LOADTIMER(Package_Tick);
 
@@ -5786,7 +5783,6 @@ EAsyncPackageState::Type FAsyncPackage::CreateLinker()
 		check(Linker);
 		check(Linker->AsyncRoot == nullptr || Linker->AsyncRoot == this);
 		Linker->AsyncRoot = this;
-		TRACE_LOADTIME_ASYNC_PACKAGE_LINKER_ASSOCIATION(this, Linker);
 
 		UE_LOG(LogStreaming, Verbose, TEXT("FAsyncPackage::CreateLinker for %s finished."), *Desc.NameToLoad.ToString());
 	}
@@ -6351,7 +6347,7 @@ EAsyncPackageState::Type FAsyncPackage::PostLoadObjects()
 
 				ThreadContext.CurrentlyPostLoadedObjectByALT = Object;
 				{
-					TRACE_LOADTIME_OBJECT_SCOPE(Object, LoadTimeProfilerObjectEventType_PostLoad);
+					TRACE_LOADTIME_POSTLOAD_EXPORT_SCOPE(Object);
 					Object->ConditionalPostLoad();
 					Object->AtomicallyClearInternalFlags(EInternalObjectFlags::AsyncLoading);
 				}
@@ -6427,7 +6423,7 @@ EAsyncPackageState::Type FAsyncPackage::PostLoadDeferredObjects(double InTickSta
 
 		PackageScope.ThreadContext.CurrentlyPostLoadedObjectByALT = Object;
 		{
-			TRACE_LOADTIME_OBJECT_SCOPE(Object, LoadTimeProfilerObjectEventType_PostLoad);
+			TRACE_LOADTIME_POSTLOAD_EXPORT_SCOPE(Object);
 			Object->ConditionalPostLoad();
 		}
 		PackageScope.ThreadContext.CurrentlyPostLoadedObjectByALT = nullptr;
diff --git a/Engine/Source/Runtime/CoreUObject/Private/Serialization/AsyncLoading2.cpp b/Engine/Source/Runtime/CoreUObject/Private/Serialization/AsyncLoading2.cpp
--- a/Engine/Source/Runtime/CoreUObject/Private/Serialization/AsyncLoading2.cpp
+++ b/Engine/Source/Runtime/CoreUObject/Private/Serialization/AsyncLoading2.cpp
@@ -960,7 +960,8 @@ private:
 	FEventLoadNode2* ExportNodes = nullptr;
 	uint32 ImportNodeCount = 0;
 	uint32 ExportNodeCount = 0;
-	
+
+	uint64 PackageSummarySize = 0;
 	FIoBuffer PackageSummaryIoBuffer;
 	TArray<FIoBuffer> ExportIoBuffers;
 
@@ -2106,6 +2107,7 @@ uint32 FAsyncLoadingThreadWorker::Run()
 void FAsyncLoadingThread2Impl::QueueEvent_CreateLinker(FAsyncPackage2* Package)
 {
 	check(Package);
+	TRACE_LOADTIME_BEGIN_LOAD_ASYNC_PACKAGE(Package);
 	Package->GetNode(EEventLoadNode2::Package_CreateLinker)->ReleaseBarrier();
 }
 
@@ -2278,6 +2280,7 @@ void FAsyncPackage2::ImportPackagesRecursive()
 		Info.Priority = Desc.Priority;
 		bool bInserted;
 		FAsyncPackage2* ImportedAsyncPackage = AsyncLoadingThread.FindOrInsertPackage(&Info, bInserted);
+		TRACE_LOADTIME_ASYNC_PACKAGE_IMPORT_DEPENDENCY(this, ImportedAsyncPackage);
 		ImportedAsyncPackage->AddRef();
 		ImportedAsyncPackages.Add(ImportedAsyncPackage);
 		if (bInserted)
@@ -2351,6 +2354,7 @@ void FAsyncPackage2::ImportPackagesRecursive()
 			Info.Priority = Desc.Priority;
 			bool bInserted;
 			FAsyncPackage2* ImportedAsyncPackage = AsyncLoadingThread.FindOrInsertPackage(&Info, bInserted);
+			TRACE_LOADTIME_ASYNC_PACKAGE_IMPORT_DEPENDENCY(this, ImportedAsyncPackage);
 			ImportedAsyncPackage->AddRef();
 			ImportedAsyncPackages.Add(ImportedAsyncPackage);
 			if (bInserted)
@@ -2546,16 +2550,16 @@ void FAsyncPackage2::LinkSlimport(int32 LocalImportIndex, int32 GlobalImportInde
 		check(Outer);
 
 		Object = StaticFindObjectFast(UObject::StaticClass(), Outer, ObjectName, false, true);
-		if (!Object)
+		/*if (!Object)
 		{
 			TRACE_CPUPROFILER_EVENT_SCOPE(LinkImport_SpinWait);
 			while (!Object)
 			{
 				Object = StaticFindObjectFast(UObject::StaticClass(), Outer, ObjectName, false, true);
 			}
-		}
+		}*/
 	}
-	check(Object);
+	//check(Object);
 }
 
 UObject* FAsyncPackage2::EventDrivenIndexToObject(FPackageIndex Index, bool bCheckSerialized)
@@ -2604,7 +2608,7 @@ void FAsyncPackage2::EventDrivenCreateExport(int32 LocalExportIndex)
 	//SCOPED_LOADTIMER(Package_CreateExports);
 	FObjectExport& Export = ExportMap[LocalExportIndex];
 
-	TRACE_LOADTIME_CREATE_EXPORT_SCOPE(Linker, &Export.Object, Export.SerialOffset, Export.SerialSize, Export.bIsAsset);
+	TRACE_LOADTIME_CREATE_EXPORT_SCOPE(this, &Export.Object);
 
 	LLM_SCOPE(ELLMTag::AsyncLoading);
 	LLM_SCOPED_TAG_WITH_OBJECT_IN_SET(GetLinkerRoot(), ELLMTagSet::Assets);
@@ -2960,7 +2964,7 @@ void FAsyncPackage2::EventDrivenSerializeExport(int32 LocalExportIndex)
 
 		Object->ClearFlags(RF_NeedLoad);
 
-		TRACE_LOADTIME_OBJECT_SCOPE(Object, LoadTimeProfilerObjectEventType_Serialize);
+		TRACE_LOADTIME_SERIALIZE_EXPORT_SCOPE(Object, ExportIoBuffers[LocalExportIndex].DataSize());
 
 		FUObjectSerializeContext* LoadContext = GetSerializeContext();
 		UObject* PrevSerializedObject = LoadContext->SerializedObject;
@@ -3148,6 +3152,7 @@ void FAsyncPackage2::ProcessIoRequest(const FIoChunkId& ChunkId, TIoStatusOr<FIo
 	{
 	case EChunkType::PackageSummary:
 	{
+		PackageSummarySize = IoBuffer.DataSize();
 		PackageSummaryIoBuffer = IoBuffer;
 		GetNode(EEventLoadNode2::Package_LoadSummary)->ReleaseBarrier();
 		break;
@@ -3403,6 +3408,8 @@ EAsyncPackageState::Type FAsyncLoadingThread2Impl::ProcessLoadedPackagesFromGame
 
 			UE_CLOG(NewExistingAsyncPackagesCounterValue < 0, LogStreaming, Fatal, TEXT("ExistingAsyncPackagesCounter is negative, this means we loaded more packages then requested so there must be a bug in async loading code."));
 
+			TRACE_LOADTIME_END_LOAD_ASYNC_PACKAGE(Package);
+
 			// Call external callbacks
 			const bool bInternalCallbacks = false;
 			const EAsyncLoadingResult::Type LoadingResult = Package->HasLoadFailed() ? EAsyncLoadingResult::Failed : EAsyncLoadingResult::Succeeded;
@@ -4152,7 +4159,7 @@ FAsyncPackage2::FAsyncPackage2(const FAsyncPackageDesc& InDesc, int32 InSerialNu
 , bAllExportsSerialized(false)
 {
 	TRACE_CPUPROFILER_EVENT_SCOPE(NewAsyncPackage);
-	//TRACE_LOADTIME_NEW_ASYNC_PACKAGE(this, *InDesc.Name.ToString());
+	TRACE_LOADTIME_NEW_ASYNC_PACKAGE(this, *InDesc.Name.ToString());
 	AddRequestID(InDesc.RequestID);
 
 	const int32* NameIndex = AsyncLoadingThread.GlobalNameMap.GetIndex(Desc.NameToLoad);
@@ -4246,7 +4253,7 @@ FAsyncPackage2::~FAsyncPackage2()
 
 	FAsyncLoadingThreadState2::Get().DeferredFreeNodes.Add(MakeTuple(PackageNodes, EEventLoadNode2::Package_NumPhases + ImportNodeCount + ExportNodeCount));
 
-	//TRACE_LOADTIME_DESTROY_ASYNC_PACKAGE(this);
+	TRACE_LOADTIME_DESTROY_ASYNC_PACKAGE(this);
 
 	MarkRequestIDsAsComplete();
 	DetachLinker();
@@ -4291,7 +4298,7 @@ void FAsyncPackage2::AddRequestID(int32 Id)
 		}
 		RequestIDs.Add(Id);
 		AsyncLoadingThread.AddPendingRequest(Id);
-		//TRACE_LOADTIME_ASYNC_PACKAGE_REQUEST_ASSOCIATION(this, Id);
+		TRACE_LOADTIME_ASYNC_PACKAGE_REQUEST_ASSOCIATION(this, Id);
 	}
 }
 
@@ -4474,6 +4481,8 @@ EAsyncPackageState::Type FAsyncPackage2::FinishLinker()
 	const int32 ExportCount = AsyncLoadingThread.GetPackageExportCount(GlobalPackageId);
 	check(ExportCount);
 
+	TRACE_LOADTIME_PACKAGE_SUMMARY(this, PackageSummarySize, AsyncLoadingThread.GetPackageImportCount(GlobalPackageId), ExportCount);
+
 	{
 		SCOPED_LOADTIMER(LinkerLoad_SerializePackageFileSummary);
 
@@ -4593,7 +4602,7 @@ EAsyncPackageState::Type FAsyncPackage2::PostLoadObjects()
 
 				ThreadContext.CurrentlyPostLoadedObjectByALT = Object;
 				{
-					TRACE_LOADTIME_OBJECT_SCOPE(Object, LoadTimeProfilerObjectEventType_PostLoad);
+					TRACE_LOADTIME_POSTLOAD_EXPORT_SCOPE(Object);
 					Object->ConditionalPostLoad();
 					Object->AtomicallyClearInternalFlags(EInternalObjectFlags::AsyncLoading);
 				}
@@ -4653,7 +4662,7 @@ EAsyncPackageState::Type FAsyncPackage2::PostLoadDeferredObjects()
 
 		PackageScope.ThreadContext.CurrentlyPostLoadedObjectByALT = Object;
 		{
-			TRACE_LOADTIME_OBJECT_SCOPE(Object, LoadTimeProfilerObjectEventType_PostLoad);
+			TRACE_LOADTIME_POSTLOAD_EXPORT_SCOPE(Object);
 			Object->ConditionalPostLoad();
 		}
 		PackageScope.ThreadContext.CurrentlyPostLoadedObjectByALT = nullptr;
diff --git a/Engine/Source/Runtime/CoreUObject/Private/Serialization/LoadTimeTrace.cpp b/Engine/Source/Runtime/CoreUObject/Private/Serialization/LoadTimeTrace.cpp
--- a/Engine/Source/Runtime/CoreUObject/Private/Serialization/LoadTimeTrace.cpp
+++ b/Engine/Source/Runtime/CoreUObject/Private/Serialization/LoadTimeTrace.cpp
@@ -26,29 +26,17 @@ UE_TRACE_EVENT_BEGIN(LoadTime, ResumeAsyncLoading)
 	UE_TRACE_EVENT_FIELD(uint64, Cycle)
 UE_TRACE_EVENT_END()
 
-UE_TRACE_EVENT_BEGIN(LoadTime, NewLinker)
-	UE_TRACE_EVENT_FIELD(const FLinkerLoad*, Linker)
-UE_TRACE_EVENT_END()
-
-UE_TRACE_EVENT_BEGIN(LoadTime, DestroyLinker)
-	UE_TRACE_EVENT_FIELD(const FLinkerLoad*, Linker)
-UE_TRACE_EVENT_END()
-
 UE_TRACE_EVENT_BEGIN(LoadTime, PackageSummary)
-	UE_TRACE_EVENT_FIELD(const FLinkerLoad*, Linker)
+	UE_TRACE_EVENT_FIELD(const FGCObject*, AsyncPackage)
 	UE_TRACE_EVENT_FIELD(uint32, TotalHeaderSize)
-	UE_TRACE_EVENT_FIELD(uint32, NameCount)
 	UE_TRACE_EVENT_FIELD(uint32, ImportCount)
 	UE_TRACE_EVENT_FIELD(uint32, ExportCount)
 UE_TRACE_EVENT_END()
 
 UE_TRACE_EVENT_BEGIN(LoadTime, BeginCreateExport)
 	UE_TRACE_EVENT_FIELD(uint64, Cycle)
-	UE_TRACE_EVENT_FIELD(const FLinkerLoad*, Linker)
-	UE_TRACE_EVENT_FIELD(uint64, SerialOffset)
-	UE_TRACE_EVENT_FIELD(uint64, SerialSize)
+	UE_TRACE_EVENT_FIELD(const FGCObject*, AsyncPackage)
 	UE_TRACE_EVENT_FIELD(uint32, ThreadId)
-	UE_TRACE_EVENT_FIELD(bool, IsAsset)
 UE_TRACE_EVENT_END()
 
 UE_TRACE_EVENT_BEGIN(LoadTime, EndCreateExport)
@@ -58,14 +46,25 @@ UE_TRACE_EVENT_BEGIN(LoadTime, EndCreateExport)
 	UE_TRACE_EVENT_FIELD(uint32, ThreadId)
 UE_TRACE_EVENT_END()
 
-UE_TRACE_EVENT_BEGIN(LoadTime, BeginObjectScope)
+UE_TRACE_EVENT_BEGIN(LoadTime, BeginSerializeExport)
+	UE_TRACE_EVENT_FIELD(uint64, Cycle)
+	UE_TRACE_EVENT_FIELD(const UObject*, Object)
+	UE_TRACE_EVENT_FIELD(uint64, SerialSize)
+	UE_TRACE_EVENT_FIELD(uint32, ThreadId)
+UE_TRACE_EVENT_END()
+
+UE_TRACE_EVENT_BEGIN(LoadTime, EndSerializeExport)
+	UE_TRACE_EVENT_FIELD(uint64, Cycle)
+	UE_TRACE_EVENT_FIELD(uint32, ThreadId)
+UE_TRACE_EVENT_END()
+
+UE_TRACE_EVENT_BEGIN(LoadTime, BeginPostLoadExport)
 	UE_TRACE_EVENT_FIELD(uint64, Cycle)
 	UE_TRACE_EVENT_FIELD(const UObject*, Object)
 	UE_TRACE_EVENT_FIELD(uint32, ThreadId)
-	UE_TRACE_EVENT_FIELD(uint8, EventType)
 UE_TRACE_EVENT_END()
 
-UE_TRACE_EVENT_BEGIN(LoadTime, EndObjectScope)
+UE_TRACE_EVENT_BEGIN(LoadTime, EndPostLoadExport)
 	UE_TRACE_EVENT_FIELD(uint64, Cycle)
 	UE_TRACE_EVENT_FIELD(uint32, ThreadId)
 UE_TRACE_EVENT_END()
@@ -82,38 +81,31 @@ UE_TRACE_EVENT_BEGIN(LoadTime, EndRequest)
 UE_TRACE_EVENT_END()
 
 UE_TRACE_EVENT_BEGIN(LoadTime, NewAsyncPackage)
-	UE_TRACE_EVENT_FIELD(const FAsyncPackage*, AsyncPackage)
+	UE_TRACE_EVENT_FIELD(const FGCObject*, AsyncPackage)
 UE_TRACE_EVENT_END()
 
-UE_TRACE_EVENT_BEGIN(LoadTime, DestroyAsyncPackage)
-	UE_TRACE_EVENT_FIELD(const FAsyncPackage*, AsyncPackage)
-UE_TRACE_EVENT_END()
-
-UE_TRACE_EVENT_BEGIN(LoadTime, AsyncPackageRequestAssociation)
-	UE_TRACE_EVENT_FIELD(const FAsyncPackage*, AsyncPackage)
-	UE_TRACE_EVENT_FIELD(uint64, RequestId)
+UE_TRACE_EVENT_BEGIN(LoadTime, BeginLoadAsyncPackage)
+	UE_TRACE_EVENT_FIELD(uint64, Cycle)
+	UE_TRACE_EVENT_FIELD(const FGCObject*, AsyncPackage)
 UE_TRACE_EVENT_END()
 
-UE_TRACE_EVENT_BEGIN(LoadTime, AsyncPackageImportDependency)
-	UE_TRACE_EVENT_FIELD(const FAsyncPackage*, AsyncPackage)
-	UE_TRACE_EVENT_FIELD(const FAsyncPackage*, ImportedAsyncPackage)
+UE_TRACE_EVENT_BEGIN(LoadTime, EndLoadAsyncPackage)
+	UE_TRACE_EVENT_FIELD(uint64, Cycle)
+	UE_TRACE_EVENT_FIELD(const FGCObject*, AsyncPackage)
 UE_TRACE_EVENT_END()
 
-UE_TRACE_EVENT_BEGIN(LoadTime, AsyncPackageLinkerAssociation)
-	UE_TRACE_EVENT_FIELD(const FAsyncPackage*, AsyncPackage)
-	UE_TRACE_EVENT_FIELD(const FLinkerLoad*, Linker)
+UE_TRACE_EVENT_BEGIN(LoadTime, DestroyAsyncPackage)
+	UE_TRACE_EVENT_FIELD(const FGCObject*, AsyncPackage)
 UE_TRACE_EVENT_END()
 
-UE_TRACE_EVENT_BEGIN(LoadTime, BeginAsyncPackageScope)
-	UE_TRACE_EVENT_FIELD(const FAsyncPackage*, AsyncPackage)
-	UE_TRACE_EVENT_FIELD(uint64, Cycle)
-	UE_TRACE_EVENT_FIELD(uint32, ThreadId)
-	UE_TRACE_EVENT_FIELD(uint8, EventType)
+UE_TRACE_EVENT_BEGIN(LoadTime, AsyncPackageRequestAssociation)
+	UE_TRACE_EVENT_FIELD(const FGCObject*, AsyncPackage)
+	UE_TRACE_EVENT_FIELD(uint64, RequestId)
 UE_TRACE_EVENT_END()
 
-UE_TRACE_EVENT_BEGIN(LoadTime, EndAsyncPackageScope)
-	UE_TRACE_EVENT_FIELD(uint64, Cycle)
-	UE_TRACE_EVENT_FIELD(uint32, ThreadId)
+UE_TRACE_EVENT_BEGIN(LoadTime, AsyncPackageImportDependency)
+	UE_TRACE_EVENT_FIELD(const FGCObject*, AsyncPackage)
+	UE_TRACE_EVENT_FIELD(const FGCObject*, ImportedAsyncPackage)
 UE_TRACE_EVENT_END()
 
 UE_TRACE_EVENT_BEGIN(LoadTime, ClassInfo, Always)
@@ -126,22 +118,21 @@ void FLoadTimeProfilerTracePrivate::Init()
 	UE_TRACE_EVENT_IS_ENABLED(LoadTime, StartAsyncLoading);
 	UE_TRACE_EVENT_IS_ENABLED(LoadTime, SuspendAsyncLoading);
 	UE_TRACE_EVENT_IS_ENABLED(LoadTime, ResumeAsyncLoading);
-	UE_TRACE_EVENT_IS_ENABLED(LoadTime, NewLinker);
-	UE_TRACE_EVENT_IS_ENABLED(LoadTime, DestroyLinker);
 	UE_TRACE_EVENT_IS_ENABLED(LoadTime, PackageSummary);
 	UE_TRACE_EVENT_IS_ENABLED(LoadTime, BeginCreateExport);
 	UE_TRACE_EVENT_IS_ENABLED(LoadTime, EndCreateExport);
-	UE_TRACE_EVENT_IS_ENABLED(LoadTime, BeginObjectScope);
-	UE_TRACE_EVENT_IS_ENABLED(LoadTime, EndObjectScope);
+	UE_TRACE_EVENT_IS_ENABLED(LoadTime, BeginSerializeExport);
+	UE_TRACE_EVENT_IS_ENABLED(LoadTime, EndSerializeExport);
+	UE_TRACE_EVENT_IS_ENABLED(LoadTime, BeginPostLoadExport);
+	UE_TRACE_EVENT_IS_ENABLED(LoadTime, EndPostLoadExport);
 	UE_TRACE_EVENT_IS_ENABLED(LoadTime, BeginRequest);
 	UE_TRACE_EVENT_IS_ENABLED(LoadTime, EndRequest);
 	UE_TRACE_EVENT_IS_ENABLED(LoadTime, NewAsyncPackage);
+	UE_TRACE_EVENT_IS_ENABLED(LoadTime, BeginLoadAsyncPackage);
+	UE_TRACE_EVENT_IS_ENABLED(LoadTime, EndLoadAsyncPackage);
 	UE_TRACE_EVENT_IS_ENABLED(LoadTime, DestroyAsyncPackage);
 	UE_TRACE_EVENT_IS_ENABLED(LoadTime, AsyncPackageRequestAssociation);
 	UE_TRACE_EVENT_IS_ENABLED(LoadTime, AsyncPackageImportDependency);
-	UE_TRACE_EVENT_IS_ENABLED(LoadTime, AsyncPackageLinkerAssociation);
-	UE_TRACE_EVENT_IS_ENABLED(LoadTime, BeginAsyncPackageScope);
-	UE_TRACE_EVENT_IS_ENABLED(LoadTime, EndAsyncPackageScope);
 	UE_TRACE_EVENT_IS_ENABLED(LoadTime, ClassInfo);
 	if (FParse::Param(FCommandLine::Get(), TEXT("loadtimetrace")))
 	{
@@ -170,18 +161,6 @@ void FLoadTimeProfilerTracePrivate::OutputResumeAsyncLoading()
 }
 PRAGMA_ENABLE_SHADOW_VARIABLE_WARNINGS
 
-void FLoadTimeProfilerTracePrivate::OutputNewLinker(const FLinkerLoad* Linker)
-{
-	UE_TRACE_LOG(LoadTime, NewLinker)
-		<< NewLinker.Linker(Linker);
-}
-
-void FLoadTimeProfilerTracePrivate::OutputDestroyLinker(const FLinkerLoad* Linker)
-{
-	UE_TRACE_LOG(LoadTime, DestroyLinker)
-		<< DestroyLinker.Linker(Linker);
-}
-
 void FLoadTimeProfilerTracePrivate::OutputBeginRequest(uint64 RequestId)
 {
 	UE_TRACE_LOG(LoadTime, BeginRequest)
@@ -197,7 +176,7 @@ void FLoadTimeProfilerTracePrivate::OutputEndRequest(uint64 RequestId)
 		<< EndRequest.RequestId(RequestId);
 }
 
-void FLoadTimeProfilerTracePrivate::OutputNewAsyncPackage(const FAsyncPackage* AsyncPackage, const TCHAR* PackageName)
+void FLoadTimeProfilerTracePrivate::OutputNewAsyncPackage(const FGCObject* AsyncPackage, const TCHAR* PackageName)
 {
 	uint16 NameSize = (FCString::Strlen(PackageName) + 1) * sizeof(TCHAR);
 	UE_TRACE_LOG(LoadTime, NewAsyncPackage, NameSize)
@@ -205,43 +184,49 @@ void FLoadTimeProfilerTracePrivate::OutputNewAsyncPackage(const FAsyncPackage* A
 		<< NewAsyncPackage.Attachment(PackageName, NameSize);
 }
 
-void FLoadTimeProfilerTracePrivate::OutputDestroyAsyncPackage(const FAsyncPackage* AsyncPackage)
+void FLoadTimeProfilerTracePrivate::OutputBeginLoadAsyncPackage(const FGCObject* AsyncPackage)
+{
+	UE_TRACE_LOG(LoadTime, BeginLoadAsyncPackage)
+		<< BeginLoadAsyncPackage.Cycle(FPlatformTime::Cycles64())
+		<< BeginLoadAsyncPackage.AsyncPackage(AsyncPackage);
+}
+
+void FLoadTimeProfilerTracePrivate::OutputEndLoadAsyncPackage(const FGCObject* AsyncPackage)
+{
+	UE_TRACE_LOG(LoadTime, EndLoadAsyncPackage)
+		<< EndLoadAsyncPackage.Cycle(FPlatformTime::Cycles64())
+		<< EndLoadAsyncPackage.AsyncPackage(AsyncPackage);
+}
+
+void FLoadTimeProfilerTracePrivate::OutputDestroyAsyncPackage(const FGCObject* AsyncPackage)
 {
 	UE_TRACE_LOG(LoadTime, DestroyAsyncPackage)
 		<< DestroyAsyncPackage.AsyncPackage(AsyncPackage);
 }
 
-void FLoadTimeProfilerTracePrivate::OutputPackageSummary(const FLinkerLoad* Linker, uint32 TotalHeaderSize, uint32 NameCount, uint32 ImportCount, uint32 ExportCount)
+void FLoadTimeProfilerTracePrivate::OutputPackageSummary(const FGCObject* AsyncPackage, uint32 TotalHeaderSize, uint32 ImportCount, uint32 ExportCount)
 {
 	UE_TRACE_LOG(LoadTime, PackageSummary)
-		<< PackageSummary.Linker(Linker)
+		<< PackageSummary.AsyncPackage(AsyncPackage)
 		<< PackageSummary.TotalHeaderSize(TotalHeaderSize)
-		<< PackageSummary.NameCount(NameCount)
 		<< PackageSummary.ImportCount(ImportCount)
 		<< PackageSummary.ExportCount(ExportCount);
 }
 
-void FLoadTimeProfilerTracePrivate::OutputAsyncPackageRequestAssociation(const FAsyncPackage* AsyncPackage, uint64 RequestId)
+void FLoadTimeProfilerTracePrivate::OutputAsyncPackageRequestAssociation(const FGCObject* AsyncPackage, uint64 RequestId)
 {
 	UE_TRACE_LOG(LoadTime, AsyncPackageRequestAssociation)
 		<< AsyncPackageRequestAssociation.AsyncPackage(AsyncPackage)
 		<< AsyncPackageRequestAssociation.RequestId(RequestId);
 }
 
-void FLoadTimeProfilerTracePrivate::OutputAsyncPackageImportDependency(const FAsyncPackage* Package, const FAsyncPackage* ImportedPackage)
+void FLoadTimeProfilerTracePrivate::OutputAsyncPackageImportDependency(const FGCObject* Package, const FGCObject* ImportedPackage)
 {
 	UE_TRACE_LOG(LoadTime, AsyncPackageImportDependency)
 		<< AsyncPackageImportDependency.AsyncPackage(Package)
 		<< AsyncPackageImportDependency.ImportedAsyncPackage(ImportedPackage);
 }
 
-void FLoadTimeProfilerTracePrivate::OutputAsyncPackageLinkerAssociation(const FAsyncPackage* AsyncPackage, const FLinkerLoad* Linker)
-{
-	UE_TRACE_LOG(LoadTime, AsyncPackageLinkerAssociation)
-		<< AsyncPackageLinkerAssociation.AsyncPackage(AsyncPackage)
-		<< AsyncPackageLinkerAssociation.Linker(Linker);
-}
-
 void FLoadTimeProfilerTracePrivate::OutputClassInfo(const UClass* Class, const TCHAR* Name)
 {
 	uint16 NameSize = (FCString::Strlen(Name) + 1) * sizeof(TCHAR);
@@ -250,32 +235,13 @@ void FLoadTimeProfilerTracePrivate::OutputClassInfo(const UClass* Class, const T
 		<< ClassInfo.Attachment(Name, NameSize);
 }
 
-FLoadTimeProfilerTracePrivate::FAsyncPackageScope::FAsyncPackageScope(const FAsyncPackage* AsyncPackage, ELoadTimeProfilerPackageEventType EventType)
-{
-	UE_TRACE_LOG(LoadTime, BeginAsyncPackageScope)
-		<< BeginAsyncPackageScope.AsyncPackage(AsyncPackage)
-		<< BeginAsyncPackageScope.Cycle(FPlatformTime::Cycles64())
-		<< BeginAsyncPackageScope.ThreadId(FPlatformTLS::GetCurrentThreadId())
-		<< BeginAsyncPackageScope.EventType(EventType);
-}
-
-FLoadTimeProfilerTracePrivate::FAsyncPackageScope::~FAsyncPackageScope()
-{
-	UE_TRACE_LOG(LoadTime, EndAsyncPackageScope)
-		<< EndAsyncPackageScope.Cycle(FPlatformTime::Cycles64())
-		<< EndAsyncPackageScope.ThreadId(FPlatformTLS::GetCurrentThreadId());
-}
-
-FLoadTimeProfilerTracePrivate::FCreateExportScope::FCreateExportScope(const FLinkerLoad* Linker, uint64 SerialOffset, uint64 SerialSize, bool IsAsset, const UObject* const* InObject)
+FLoadTimeProfilerTracePrivate::FCreateExportScope::FCreateExportScope(const FGCObject* AsyncPackage, const UObject* const* InObject)
 	: Object(InObject)
 {
 	UE_TRACE_LOG(LoadTime, BeginCreateExport)
 		<< BeginCreateExport.Cycle(FPlatformTime::Cycles64())
-		<< BeginCreateExport.Linker(Linker)
-		<< BeginCreateExport.SerialOffset(SerialOffset)
-		<< BeginCreateExport.SerialSize(SerialSize)
-		<< BeginCreateExport.ThreadId(FPlatformTLS::GetCurrentThreadId())
-		<< BeginCreateExport.IsAsset(IsAsset);
+		<< BeginCreateExport.AsyncPackage(AsyncPackage)
+		<< BeginCreateExport.ThreadId(FPlatformTLS::GetCurrentThreadId());
 }
 
 FLoadTimeProfilerTracePrivate::FCreateExportScope::~FCreateExportScope()
@@ -287,20 +253,35 @@ FLoadTimeProfilerTracePrivate::FCreateExportScope::~FCreateExportScope()
 		<< EndCreateExport.ThreadId(FPlatformTLS::GetCurrentThreadId());
 }
 
-FLoadTimeProfilerTracePrivate::FObjectScope::FObjectScope(const UObject* Object, ELoadTimeProfilerObjectEventType EventType)
+FLoadTimeProfilerTracePrivate::FSerializeExportScope::FSerializeExportScope(const UObject* Object, uint64 SerialSize)
+{
+	UE_TRACE_LOG(LoadTime, BeginSerializeExport)
+		<< BeginSerializeExport.Cycle(FPlatformTime::Cycles64())
+		<< BeginSerializeExport.Object(Object)
+		<< BeginSerializeExport.SerialSize(SerialSize)
+		<< BeginSerializeExport.ThreadId(FPlatformTLS::GetCurrentThreadId());
+}
+
+FLoadTimeProfilerTracePrivate::FSerializeExportScope::~FSerializeExportScope()
+{
+	UE_TRACE_LOG(LoadTime, EndSerializeExport)
+		<< EndSerializeExport.Cycle(FPlatformTime::Cycles64())
+		<< EndSerializeExport.ThreadId(FPlatformTLS::GetCurrentThreadId());
+}
+
+FLoadTimeProfilerTracePrivate::FPostLoadExportScope::FPostLoadExportScope(const UObject* Object)
 {
-	UE_TRACE_LOG(LoadTime, BeginObjectScope)
-		<< BeginObjectScope.Cycle(FPlatformTime::Cycles64())
-		<< BeginObjectScope.Object(Object)
-		<< BeginObjectScope.ThreadId(FPlatformTLS::GetCurrentThreadId())
-		<< BeginObjectScope.EventType(EventType);
+	UE_TRACE_LOG(LoadTime, BeginPostLoadExport)
+		<< BeginPostLoadExport.Cycle(FPlatformTime::Cycles64())
+		<< BeginPostLoadExport.Object(Object)
+		<< BeginPostLoadExport.ThreadId(FPlatformTLS::GetCurrentThreadId());
 }
 
-FLoadTimeProfilerTracePrivate::FObjectScope::~FObjectScope()
+FLoadTimeProfilerTracePrivate::FPostLoadExportScope::~FPostLoadExportScope()
 {
-	UE_TRACE_LOG(LoadTime, EndObjectScope)
-		<< EndObjectScope.Cycle(FPlatformTime::Cycles64())
-		<< EndObjectScope.ThreadId(FPlatformTLS::GetCurrentThreadId());
+	UE_TRACE_LOG(LoadTime, EndPostLoadExport)
+		<< EndPostLoadExport.Cycle(FPlatformTime::Cycles64())
+		<< EndPostLoadExport.ThreadId(FPlatformTLS::GetCurrentThreadId());
 }
 
 #endif
\ No newline at end of file
diff --git a/Engine/Source/Runtime/CoreUObject/Private/Serialization/LoadTimeTracePrivate.h b/Engine/Source/Runtime/CoreUObject/Private/Serialization/LoadTimeTracePrivate.h
--- a/Engine/Source/Runtime/CoreUObject/Private/Serialization/LoadTimeTracePrivate.h
+++ b/Engine/Source/Runtime/CoreUObject/Private/Serialization/LoadTimeTracePrivate.h
@@ -6,13 +6,7 @@
 
 #if LOADTIMEPROFILERTRACE_ENABLED
 
-#include "UObject/UObjectGlobals.h"
-
-class FArchive;
-struct FAsyncPackage;
-class FLinkerLoad;
 class UObject;
-class UPackage;
 
 struct FLoadTimeProfilerTracePrivate
 {
@@ -20,44 +14,38 @@ struct FLoadTimeProfilerTracePrivate
 	static void OutputStartAsyncLoading();
 	static void OutputSuspendAsyncLoading();
 	static void OutputResumeAsyncLoading();
-	static void OutputNewLinker(const FLinkerLoad* Linker);
-	static void OutputDestroyLinker(const FLinkerLoad* Linker);
-	static void OutputNewAsyncPackage(const FAsyncPackage* AsyncPackage, const TCHAR* PackageName);
-	static void OutputDestroyAsyncPackage(const FAsyncPackage* AsyncPackage);
+	static void OutputNewAsyncPackage(const FGCObject* AsyncPackage, const TCHAR* PackageName);
+	static void OutputBeginLoadAsyncPackage(const FGCObject* AsyncPackage);
+	static void OutputEndLoadAsyncPackage(const FGCObject* AsyncPackage);
+	static void OutputDestroyAsyncPackage(const FGCObject* AsyncPackage);
 	static void OutputBeginRequest(uint64 RequestId);
 	static void OutputEndRequest(uint64 RequestId);
-	static void OutputPackageSummary(const FLinkerLoad* Linker, uint32 TotalHeaderSize, uint32 NameCount, uint32 ImportCount, uint32 ExportCount);
-	static void OutputAsyncPackageImportDependency(const FAsyncPackage* Package, const FAsyncPackage* ImportedPackage);
-	static void OutputAsyncPackageRequestAssociation(const FAsyncPackage* AsyncPackage, uint64 RequestId);
-	static void OutputAsyncPackageLinkerAssociation(const FAsyncPackage* AsyncPackage, const FLinkerLoad* Linker);
+	static void OutputPackageSummary(const FGCObject* AsyncPackage, uint32 TotalHeaderSize, uint32 ImportCount, uint32 ExportCount);
+	static void OutputAsyncPackageImportDependency(const FGCObject* Package, const FGCObject* ImportedPackage);
+	static void OutputAsyncPackageRequestAssociation(const FGCObject* AsyncPackage, uint64 RequestId);
 	static void OutputClassInfo(const UClass* Class, const TCHAR* Name);
 
-	struct FAsyncPackageScope
-	{
-		FAsyncPackageScope(const FAsyncPackage* AsyncPackage, ELoadTimeProfilerPackageEventType EventType);
-		~FAsyncPackageScope();
-	};
-
 	struct FCreateExportScope
 	{
-		FCreateExportScope(const FLinkerLoad* Linker, uint64 SerialOffset, uint64 SerialSize, bool IsAsset, const UObject* const* InObject);
+		FCreateExportScope(const FGCObject* AsyncPackage, const UObject* const* InObject);
 		~FCreateExportScope();
 
 	private:
 		const UObject* const* Object;
 	};
 
-	struct FObjectScope
+	struct FSerializeExportScope
 	{
-		FObjectScope(const UObject* Object, ELoadTimeProfilerObjectEventType EventType);
-		~FObjectScope();
+		FSerializeExportScope(const UObject* Object, uint64 SerialSize);
+		~FSerializeExportScope();
 	};
 
-	struct FLinkerScope
+	struct FPostLoadExportScope
 	{
-		FLinkerScope(const FLinkerLoad* Linker);
-		~FLinkerScope();
+		FPostLoadExportScope(const UObject* Object);
+		~FPostLoadExportScope();
 	};
+
 };
 
 #define TRACE_LOADTIME_START_ASYNC_LOADING() \
@@ -69,15 +57,6 @@ struct FLoadTimeProfilerTracePrivate
 #define TRACE_LOADTIME_RESUME_ASYNC_LOADING() \
 	FLoadTimeProfilerTracePrivate::OutputResumeAsyncLoading();
 
-#define TRACE_LOADTIME_NEW_LINKER(Linker) \
-	FLoadTimeProfilerTracePrivate::OutputNewLinker(Linker);
-
-#define TRACE_LOADTIME_DESTROY_LINKER(Linker) \
-	FLoadTimeProfilerTracePrivate::OutputDestroyLinker(Linker);
-
-#define TRACE_LOADTIME_PACKAGE_SUMMARY(Linker, TotalHeaderSize, NameCount, ImportCount, ExportCount) \
-	FLoadTimeProfilerTracePrivate::OutputPackageSummary(Linker, TotalHeaderSize, NameCount, ImportCount, ExportCount);
-
 #define TRACE_LOADTIME_BEGIN_REQUEST(RequestId) \
 	FLoadTimeProfilerTracePrivate::OutputBeginRequest(RequestId);
 
@@ -87,9 +66,18 @@ struct FLoadTimeProfilerTracePrivate
 #define TRACE_LOADTIME_NEW_ASYNC_PACKAGE(AsyncPackage, PackageName) \
 	FLoadTimeProfilerTracePrivate::OutputNewAsyncPackage(AsyncPackage, PackageName)
 
+#define TRACE_LOADTIME_BEGIN_LOAD_ASYNC_PACKAGE(AsyncPackage) \
+	FLoadTimeProfilerTracePrivate::OutputBeginLoadAsyncPackage(AsyncPackage)
+
+#define TRACE_LOADTIME_END_LOAD_ASYNC_PACKAGE(AsyncPackage) \
+	FLoadTimeProfilerTracePrivate::OutputEndLoadAsyncPackage(AsyncPackage)
+
 #define TRACE_LOADTIME_DESTROY_ASYNC_PACKAGE(AsyncPackage) \
 	FLoadTimeProfilerTracePrivate::OutputDestroyAsyncPackage(AsyncPackage);
 
+#define TRACE_LOADTIME_PACKAGE_SUMMARY(AsyncPackage, TotalHeaderSize, ImportCount, ExportCount) \
+	FLoadTimeProfilerTracePrivate::OutputPackageSummary(AsyncPackage, TotalHeaderSize, ImportCount, ExportCount);
+
 #define TRACE_LOADTIME_ASYNC_PACKAGE_REQUEST_ASSOCIATION(AsyncPackage, RequestId) \
 	FLoadTimeProfilerTracePrivate::OutputAsyncPackageRequestAssociation(AsyncPackage, RequestId);
 
@@ -99,14 +87,14 @@ struct FLoadTimeProfilerTracePrivate
 #define TRACE_LOADTIME_ASYNC_PACKAGE_IMPORT_DEPENDENCY(AsyncPackage, ImportedAsyncPackage) \
 	FLoadTimeProfilerTracePrivate::OutputAsyncPackageImportDependency(AsyncPackage, ImportedAsyncPackage);
 
-#define TRACE_LOADTIME_ASYNC_PACKAGE_SCOPE(AsyncPackage, EventType) \
-	FLoadTimeProfilerTracePrivate::FAsyncPackageScope __LoadTimeTraceAsyncPackageScope(AsyncPackage, EventType);
+#define TRACE_LOADTIME_CREATE_EXPORT_SCOPE(AsyncPackage, Object) \
+	FLoadTimeProfilerTracePrivate::FCreateExportScope __LoadTimeTraceCreateExportScope(AsyncPackage, Object);
 
-#define TRACE_LOADTIME_CREATE_EXPORT_SCOPE(Linker, Object, SerialOffset, SerialSize, IsAsset) \
-	FLoadTimeProfilerTracePrivate::FCreateExportScope __LoadTimeTraceCreateExportScope(Linker, SerialOffset, SerialSize, IsAsset, Object);
+#define TRACE_LOADTIME_SERIALIZE_EXPORT_SCOPE(Object, SerialSize) \
+	FLoadTimeProfilerTracePrivate::FSerializeExportScope __LoadTimeTraceSerializeExportScope(Object, SerialSize);
 
-#define TRACE_LOADTIME_OBJECT_SCOPE(Object, EventType) \
-	FLoadTimeProfilerTracePrivate::FObjectScope __LoadTimeTraceObjectScope(Object, EventType);
+#define TRACE_LOADTIME_POSTLOAD_EXPORT_SCOPE(Object) \
+	FLoadTimeProfilerTracePrivate::FPostLoadExportScope __LoadTimeTracePostLoadExportScope(Object);
 
 #define TRACE_LOADTIME_CLASS_INFO(Class, Name) \
 	FLoadTimeProfilerTracePrivate::OutputClassInfo(Class, Name);
@@ -116,21 +104,18 @@ struct FLoadTimeProfilerTracePrivate
 #define TRACE_LOADTIME_START_ASYNC_LOADING(...)
 #define TRACE_LOADTIME_SUSPEND_ASYNC_LOADING(...)
 #define TRACE_LOADTIME_RESUME_ASYNC_LOADING(...)
-#define TRACE_LOADTIME_NEW_PACKAGE(...)
-#define TRACE_LOADTIME_NEW_LINKER(...)
-#define TRACE_LOADTIME_DESTROY_LINKER(...)
-#define TRACE_LOADTIME_PACKAGE_SUMMARY(...)
 #define TRACE_LOADTIME_BEGIN_REQUEST(...)
 #define TRACE_LOADTIME_END_REQUEST(...)
 #define TRACE_LOADTIME_NEW_ASYNC_PACKAGE(...)
+#define TRACE_LOADTIME_BEGIN_LOAD_ASYNC_PACKAGE(...)
+#define TRACE_LOADTIME_END_LOAD_ASYNC_PACKAGE(...)
 #define TRACE_LOADTIME_DESTROY_ASYNC_PACKAGE(...)
+#define TRACE_LOADTIME_PACKAGE_SUMMARY(...)
 #define TRACE_LOADTIME_ASYNC_PACKAGE_REQUEST_ASSOCIATION(...)
-#define TRACE_LOADTIME_ASYNC_PACKAGE_LINKER_ASSOCIATION(...)
 #define TRACE_LOADTIME_ASYNC_PACKAGE_IMPORT_DEPENDENCY(...)
-#define TRACE_LOADTIME_ASYNC_PACKAGE_SCOPE(...)
-#define TRACE_LOADTIME_TIMER_SCOPE(...)
 #define TRACE_LOADTIME_CREATE_EXPORT_SCOPE(...)
-#define TRACE_LOADTIME_OBJECT_SCOPE(...)
+#define TRACE_LOADTIME_SERIALIZE_EXPORT_SCOPE(...)
+#define TRACE_LOADTIME_POSTLOAD_EXPORT_SCOPE(...)
 #define TRACE_LOADTIME_CLASS_INFO(...)
 
 #endif
\ No newline at end of file
diff --git a/Engine/Source/Runtime/CoreUObject/Private/UObject/LinkerLoad.cpp b/Engine/Source/Runtime/CoreUObject/Private/UObject/LinkerLoad.cpp
--- a/Engine/Source/Runtime/CoreUObject/Private/UObject/LinkerLoad.cpp
+++ b/Engine/Source/Runtime/CoreUObject/Private/UObject/LinkerLoad.cpp
@@ -866,12 +866,10 @@ FLinkerLoad::FLinkerLoad(UPackage* InParent, const TCHAR* InFilename, uint32 InL
 #endif
 
 	OwnerThread = FPlatformTLS::GetCurrentThreadId();
-	TRACE_LOADTIME_NEW_LINKER(this);
 }
 
 FLinkerLoad::~FLinkerLoad()
 {
-	TRACE_LOADTIME_DESTROY_LINKER(this);
 #if !UE_BUILD_SHIPPING && !UE_BUILD_TEST
 	FLinkerManager::Get().RemoveLiveLinker(this);
 #endif
@@ -1166,8 +1164,6 @@ FLinkerLoad::ELinkerStatus FLinkerLoad::SerializePackageFileSummary()
 		// Read summary from file.
 		StructuredArchiveRootRecord.GetValue() << SA_VALUE(TEXT("Summary"), Summary);
 
-		TRACE_LOADTIME_PACKAGE_SUMMARY(this, Summary.TotalHeaderSize, Summary.NameCount, Summary.ImportCount, Summary.ExportCount);
-
 		// Check tag.
 		if( Summary.Tag != PACKAGE_FILE_TAG )
 		{
@@ -3698,7 +3694,6 @@ void FLinkerLoad::Preload( UObject* Object )
 				Object->ClearFlags ( RF_NeedLoad );
 
 				{
-					TRACE_LOADTIME_OBJECT_SCOPE(Export.Object, LoadTimeProfilerObjectEventType_Serialize);
 					SCOPE_CYCLE_COUNTER(STAT_LinkerSerialize);
 #if USE_CIRCULAR_DEPENDENCY_LOAD_DEFERRING
 					// communicate with FLinkerPlaceholderBase, what object is currently serializing in
@@ -4019,8 +4014,6 @@ UObject* FLinkerLoad::CreateExport( int32 Index )
 	// Check whether we already loaded the object and if not whether the context flags allow loading it.
 	if( !Export.Object && !FilterExport(Export) ) // for some acceptable position, it was not "not for" 
 	{
-		TRACE_LOADTIME_CREATE_EXPORT_SCOPE(this, &Export.Object, Export.SerialOffset, Export.SerialSize, Export.bIsAsset);
-
 		FUObjectSerializeContext* CurrentLoadContext = GetSerializeContext();
 		check(!GEventDrivenLoaderEnabled || !bLockoutLegacyOperations || !EVENT_DRIVEN_ASYNC_LOAD_ACTIVE_AT_RUNTIME);
 		check(Export.ObjectName!=NAME_None || !(Export.ObjectFlags&RF_Public));
