Removed unnecessary includes of "Trace/Trace.h".

#rnx

diff --git a/Engine/Source/Developer/TraceServices/Private/Analyzers/CountersTraceAnalysis.h b/Engine/Source/Developer/TraceServices/Private/Analyzers/CountersTraceAnalysis.h
--- a/Engine/Source/Developer/TraceServices/Private/Analyzers/CountersTraceAnalysis.h
+++ b/Engine/Source/Developer/TraceServices/Private/Analyzers/CountersTraceAnalysis.h
@@ -2,7 +2,6 @@
 
 #pragma once
 
-#include "Trace/Trace.h"
 #include "Trace/Analyzer.h"
 #include "Containers/Map.h"
 
diff --git a/Engine/Source/Developer/TraceServices/Private/Analyzers/CpuProfilerTraceAnalysis.h b/Engine/Source/Developer/TraceServices/Private/Analyzers/CpuProfilerTraceAnalysis.h
--- a/Engine/Source/Developer/TraceServices/Private/Analyzers/CpuProfilerTraceAnalysis.h
+++ b/Engine/Source/Developer/TraceServices/Private/Analyzers/CpuProfilerTraceAnalysis.h
@@ -2,7 +2,6 @@
 
 #pragma once
 
-#include "Trace/Trace.h"
 #include "Trace/Analyzer.h"
 #include "Containers/UnrealString.h"
 #include "Model/TimingProfilerPrivate.h"
diff --git a/Engine/Source/Developer/TraceServices/Private/Analyzers/CsvProfilerTraceAnalysis.h b/Engine/Source/Developer/TraceServices/Private/Analyzers/CsvProfilerTraceAnalysis.h
--- a/Engine/Source/Developer/TraceServices/Private/Analyzers/CsvProfilerTraceAnalysis.h
+++ b/Engine/Source/Developer/TraceServices/Private/Analyzers/CsvProfilerTraceAnalysis.h
@@ -2,7 +2,6 @@
 
 #pragma once
 
-#include "Trace/Trace.h"
 #include "Trace/Analyzer.h"
 #include "Containers/Map.h"
 #include "Model/CsvProfilerPrivate.h"
diff --git a/Engine/Source/Developer/TraceServices/Private/Analyzers/GpuProfilerTraceAnalysis.h b/Engine/Source/Developer/TraceServices/Private/Analyzers/GpuProfilerTraceAnalysis.h
--- a/Engine/Source/Developer/TraceServices/Private/Analyzers/GpuProfilerTraceAnalysis.h
+++ b/Engine/Source/Developer/TraceServices/Private/Analyzers/GpuProfilerTraceAnalysis.h
@@ -2,7 +2,6 @@
 
 #pragma once
 
-#include "Trace/Trace.h"
 #include "Trace/Analyzer.h"
 #include "Model/TimingProfilerPrivate.h"
 
diff --git a/Engine/Source/Developer/TraceServices/Private/Analyzers/LoadTimeTraceAnalysis.cpp b/Engine/Source/Developer/TraceServices/Private/Analyzers/LoadTimeTraceAnalysis.cpp
--- a/Engine/Source/Developer/TraceServices/Private/Analyzers/LoadTimeTraceAnalysis.cpp
+++ b/Engine/Source/Developer/TraceServices/Private/Analyzers/LoadTimeTraceAnalysis.cpp
@@ -2,7 +2,6 @@
 #include "LoadTimeTraceAnalysis.h"
 
 #include "Serialization/LoadTimeTrace.h"
-#include "Trace/Trace.h"
 #include "Model/LoadTimeProfilerPrivate.h"
 #include "Analyzers/MiscTraceAnalysis.h"
 #include "Common/FormatArgs.h"
diff --git a/Engine/Source/Developer/TraceServices/Private/Analyzers/LogTraceAnalysis.h b/Engine/Source/Developer/TraceServices/Private/Analyzers/LogTraceAnalysis.h
--- a/Engine/Source/Developer/TraceServices/Private/Analyzers/LogTraceAnalysis.h
+++ b/Engine/Source/Developer/TraceServices/Private/Analyzers/LogTraceAnalysis.h
@@ -2,7 +2,6 @@
 
 #pragma once
 
-#include "Trace/Trace.h"
 #include "Trace/Analyzer.h"
 
 namespace Trace
diff --git a/Engine/Source/Developer/TraceServices/Private/Analyzers/MiscTraceAnalysis.h b/Engine/Source/Developer/TraceServices/Private/Analyzers/MiscTraceAnalysis.h
--- a/Engine/Source/Developer/TraceServices/Private/Analyzers/MiscTraceAnalysis.h
+++ b/Engine/Source/Developer/TraceServices/Private/Analyzers/MiscTraceAnalysis.h
@@ -2,7 +2,6 @@
 
 #pragma once
 
-#include "Trace/Trace.h"
 #include "Trace/Analyzer.h"
 #include "Templates/SharedPointer.h"
 #include "ProfilingDebugging/MiscTrace.h"
diff --git a/Engine/Source/Developer/TraceServices/Private/Analyzers/NetTraceAnalyzer.h b/Engine/Source/Developer/TraceServices/Private/Analyzers/NetTraceAnalyzer.h
--- a/Engine/Source/Developer/TraceServices/Private/Analyzers/NetTraceAnalyzer.h
+++ b/Engine/Source/Developer/TraceServices/Private/Analyzers/NetTraceAnalyzer.h
@@ -2,7 +2,6 @@
 
 #pragma once
 
-#include "Trace/Trace.h"
 #include "Trace/Analyzer.h"
 #include "Containers/UnrealString.h"
 #include "Model/NetProfilerProvider.h"
diff --git a/Engine/Source/Developer/TraceServices/Private/Analyzers/PlatformFileTraceAnalysis.h b/Engine/Source/Developer/TraceServices/Private/Analyzers/PlatformFileTraceAnalysis.h
--- a/Engine/Source/Developer/TraceServices/Private/Analyzers/PlatformFileTraceAnalysis.h
+++ b/Engine/Source/Developer/TraceServices/Private/Analyzers/PlatformFileTraceAnalysis.h
@@ -2,7 +2,6 @@
 
 #pragma once
 
-#include "Trace/Trace.h"
 #include "Trace/Analyzer.h"
 #include "Containers/Map.h"
 #include "TraceServices/Model/AnalysisSession.h"
diff --git a/Engine/Source/Developer/TraceServices/Private/Analyzers/StatsTraceAnalysis.h b/Engine/Source/Developer/TraceServices/Private/Analyzers/StatsTraceAnalysis.h
--- a/Engine/Source/Developer/TraceServices/Private/Analyzers/StatsTraceAnalysis.h
+++ b/Engine/Source/Developer/TraceServices/Private/Analyzers/StatsTraceAnalysis.h
@@ -2,7 +2,6 @@
 
 #pragma once
 
-#include "Trace/Trace.h"
 #include "Trace/Analyzer.h"
 #include "Templates/SharedPointer.h"
 
