Fixed Unreal Insights getting stuck waiting for executables from branches that won't send the SessionInfo trace event.

#rnx

diff --git a/Engine/Source/Developer/TraceServices/Private/SessionService.cpp b/Engine/Source/Developer/TraceServices/Private/SessionService.cpp
--- a/Engine/Source/Developer/TraceServices/Private/SessionService.cpp
+++ b/Engine/Source/Developer/TraceServices/Private/SessionService.cpp
@@ -314,7 +314,7 @@ bool FSessionService::Tick(float DeltaTime)
 
 void FSessionService::UpdateSessionContext(FStoreSessionHandle StoreHandle, FSessionInfoInternal& Info)
 {
-	if (Info.bIsUpdated)
+	if (Info.bIsUpdated || Info.bIsLive)
 	{
 		return;
 	}
