Only check a trace for session information if the trace is of a suitable master version.

#rnx

diff --git a/Engine/Source/Developer/TraceServices/Private/SessionService.cpp b/Engine/Source/Developer/TraceServices/Private/SessionService.cpp
--- a/Engine/Source/Developer/TraceServices/Private/SessionService.cpp
+++ b/Engine/Source/Developer/TraceServices/Private/SessionService.cpp
@@ -28,6 +28,11 @@ struct FDiagnosticsSessionAnalyzer
 {
 	virtual void OnAnalysisBegin(const FOnAnalysisContext& Context) override
 	{
+		if (Context.SessionContext.Version < 2)
+		{
+			return;
+		}
+
 		Context.InterfaceBuilder.RouteEvent(0, "Diagnostics", "Session");
 	}
 
@@ -314,7 +319,7 @@ bool FSessionService::Tick(float DeltaTime)
 
 void FSessionService::UpdateSessionContext(FStoreSessionHandle StoreHandle, FSessionInfoInternal& Info)
 {
-	if (Info.bIsUpdated || Info.bIsLive)
+	if (Info.bIsUpdated)
 	{
 		return;
 	}
@@ -326,9 +331,14 @@ void FSessionService::UpdateSessionContext(FStoreSessionHandle StoreHandle, FSes
 	{
 		virtual int32 Read(void* Data, uint32 Size) override
 		{
+			if (BytesRead >= 8192)
+			{
+				return 0;
+			}
+
 			int32 InnerBytesRead = Inner->Read(Data, Size);
 			BytesRead += InnerBytesRead;
-			return (BytesRead < (2 << 20)) ? InnerBytesRead : 0;
+			return InnerBytesRead;
 		}
 
 		int32 BytesRead = 0;
