This needs to be slightly larger to better accomodate the packet sizes that events are transported in.

#rnx

diff --git a/Engine/Source/Developer/TraceServices/Private/SessionService.cpp b/Engine/Source/Developer/TraceServices/Private/SessionService.cpp
--- a/Engine/Source/Developer/TraceServices/Private/SessionService.cpp
+++ b/Engine/Source/Developer/TraceServices/Private/SessionService.cpp
@@ -331,7 +331,7 @@ void FSessionService::UpdateSessionContext(FStoreSessionHandle StoreHandle, FSes
 	{
 		virtual int32 Read(void* Data, uint32 Size) override
 		{
-			if (BytesRead >= 8192)
+			if (BytesRead >= 32768)
 			{
 				return 0;
 			}
