UnrealInsights:
Support recursive session read/edit scopes
Support multiple concurrent session readers

diff --git a/Engine/Source/Developer/TraceServices/Private/AnalysisService.cpp b/Engine/Source/Developer/TraceServices/Private/AnalysisService.cpp
--- a/Engine/Source/Developer/TraceServices/Private/AnalysisService.cpp
+++ b/Engine/Source/Developer/TraceServices/Private/AnalysisService.cpp
@@ -20,46 +20,60 @@
 namespace Trace
 {
 
+thread_local FAnalysisSessionLock* GThreadCurrentSessionLock;
+thread_local int32 GThreadCurrentReadLockCount;
+thread_local int32 GThreadCurrentWriteLockCount;
+
 void FAnalysisSessionLock::ReadAccessCheck() const
 {
-	checkf(IsReadOnly || FPlatformTLS::GetCurrentThreadId() == OwnerThread, TEXT("Trying to read from session while someone else is writing"));
+	checkf(GThreadCurrentSessionLock == this && GThreadCurrentReadLockCount > 0, TEXT("Trying to read from session outside of a ReadScope"));
 }
 
 void FAnalysisSessionLock::WriteAccessCheck() const
 {
-	checkf(!IsReadOnly, TEXT("Trying to edit session while in read only mode"));
-	checkf(FPlatformTLS::GetCurrentThreadId() == OwnerThread, TEXT("Trying to edit session from thread without write access"));
+	checkf(GThreadCurrentSessionLock == this && GThreadCurrentWriteLockCount > 0, TEXT("Trying to write to session outside of an EditScope"));
 }
 
 void FAnalysisSessionLock::BeginRead()
 {
-	CriticalSection.Lock();
-	check(OwnerThread == 0);
-	OwnerThread = FPlatformTLS::GetCurrentThreadId();
-	IsReadOnly = true;
+	check(!GThreadCurrentSessionLock || GThreadCurrentSessionLock == this);
+	checkf(GThreadCurrentWriteLockCount == 0, TEXT("Trying to lock for read while holding write access"));
+	if (GThreadCurrentReadLockCount++ == 0)
+	{
+		GThreadCurrentSessionLock = this;
+		RWLock.ReadLock();
+	}
 }
 
 void FAnalysisSessionLock::EndRead()
 {
-	check(FPlatformTLS::GetCurrentThreadId() == OwnerThread);
-	OwnerThread = 0;
-	IsReadOnly = false;
-	CriticalSection.Unlock();
+	check(GThreadCurrentReadLockCount > 0);
+	if (--GThreadCurrentReadLockCount == 0)
+	{
+		RWLock.ReadUnlock();
+		GThreadCurrentSessionLock = nullptr;
+	}
 }
 
 void FAnalysisSessionLock::BeginEdit()
 {
-	CriticalSection.Lock();
-	check(OwnerThread == 0);
-	check(!IsReadOnly);
-	OwnerThread = FPlatformTLS::GetCurrentThreadId();
+	check(!GThreadCurrentSessionLock || GThreadCurrentSessionLock == this);
+	checkf(GThreadCurrentWriteLockCount == 0, TEXT("Trying to lock for edit while holding read access"));
+	if (GThreadCurrentWriteLockCount++ == 0)
+	{
+		GThreadCurrentSessionLock = this;
+		RWLock.WriteLock();
+	}
 }
 
 void FAnalysisSessionLock::EndEdit()
 {
-	check(OwnerThread == FPlatformTLS::GetCurrentThreadId());
-	OwnerThread = 0;
-	CriticalSection.Unlock();
+	check(GThreadCurrentWriteLockCount > 0);
+	if (--GThreadCurrentWriteLockCount == 0)
+	{
+		RWLock.WriteUnlock();
+		GThreadCurrentSessionLock = nullptr;
+	}
 }
 
 FAnalysisSession::FAnalysisSession(const TCHAR* SessionName)
diff --git a/Engine/Source/Developer/TraceServices/Private/AnalysisServicePrivate.h b/Engine/Source/Developer/TraceServices/Private/AnalysisServicePrivate.h
--- a/Engine/Source/Developer/TraceServices/Private/AnalysisServicePrivate.h
+++ b/Engine/Source/Developer/TraceServices/Private/AnalysisServicePrivate.h
@@ -26,9 +26,7 @@ public:
 	void EndEdit();
 
 private:
-	FCriticalSection CriticalSection;
-	uint32 OwnerThread = 0;
-	bool IsReadOnly = false;
+	FRWLock RWLock;
 };
 
 class FAnalysisSession
