Fixed bug when having both incoming and outgoing data

(peafour-cherrypick of //UE4/Dev-Networking-ReplicationRefactor/[REDACTED_HANDLE] by [REDACTED_PERSON])

diff --git a/Engine/Source/Developer/TraceServices/Private/Analyzers/NetTraceAnalyzer.cpp b/Engine/Source/Developer/TraceServices/Private/Analyzers/NetTraceAnalyzer.cpp
--- a/Engine/Source/Developer/TraceServices/Private/Analyzers/NetTraceAnalyzer.cpp
+++ b/Engine/Source/Developer/TraceServices/Private/Analyzers/NetTraceAnalyzer.cpp
@@ -118,14 +118,14 @@ bool FNetTraceAnalyzer::OnEvent(uint16 RouteId, const FOnEventContext& Context)
 			TPagedArray<Trace::FNetProfilerContentEvent>& Events = ConnectionData.ContentEvents;
 
 			// New Packet, update state
-			if (ConnectionState->CurrentPacketStartIndex == Events.Num())
+			if (ConnectionState->CurrentPacketStartIndex[ConnectionMode] == Events.Num())
 			{
 				ConnectionState->ConnectionMode = ConnectionMode;
 			}
 			else
 			{
 				// Corrupt data
-				check(ConnectionState->CurrentPacketStartIndex == Events.Num() || (ConnectionState->ConnectionMode == ConnectionMode));
+				check(ConnectionState->CurrentPacketStartIndex[ConnectionMode] == Events.Num() || (ConnectionState->ConnectionMode == ConnectionMode));
 			}
 
 			// Decode batched events
@@ -185,7 +185,7 @@ bool FNetTraceAnalyzer::OnEvent(uint16 RouteId, const FOnEventContext& Context)
 			const uint8 GameInstanceId = EventData.GetValue<uint8>("ReplicationSystemId");
 			const uint16 ConnectionId = EventData.GetValue<uint16>("ConnectionId");
 
-			// Update LastTimestamp, latar on we will be able to get timestamps piggybacked from other analyzers
+			// Update LastTimestamp, later on we will be able to get timestamps piggybacked from other analyzers
 			LastTimeStamp = Context.SessionContext.TimestampFromCycle(TimestampCycles);
 
 			TSharedRef<FNetTraceGameInstanceState> GameInstanceState = GetOrCreateActiveGameInstanceState(GameInstanceId);
@@ -197,7 +197,7 @@ bool FNetTraceAnalyzer::OnEvent(uint16 RouteId, const FOnEventContext& Context)
 			++ConnectionData.PacketChangeCount;
 
 			// Fill in packet data a packet has at least 1 event
-			Packet.StartEventIndex = ConnectionState->CurrentPacketStartIndex;
+			Packet.StartEventIndex = ConnectionState->CurrentPacketStartIndex[ConnectionState->ConnectionMode];
 			Packet.EventCount = ConnectionData.ContentEvents.Num() - Packet.StartEventIndex;
 			Packet.TimeStamp = GetLastTimestamp();
 			Packet.SequenceNumber = SequenceNumber;
@@ -210,7 +210,7 @@ bool FNetTraceAnalyzer::OnEvent(uint16 RouteId, const FOnEventContext& Context)
 			//NetProfilerProvider->FinalizePacket();
 
 			// Mark the beginning of a new packet
-			ConnectionState->CurrentPacketStartIndex = ConnectionData.ContentEvents.Num();
+			ConnectionState->CurrentPacketStartIndex[ConnectionState->ConnectionMode] = ConnectionData.ContentEvents.Num();
 		}
 		break;
 
@@ -252,7 +252,8 @@ bool FNetTraceAnalyzer::OnEvent(uint16 RouteId, const FOnEventContext& Context)
 			Connection.Connection.ConnectionId = ConnectionId;
 			Connection.Connection.LifeTime.Begin =  GetLastTimestamp();
 			ConnectionState->ConnectionIndex = Connection.Connection.ConnectionIndex;
-			ConnectionState->CurrentPacketStartIndex = 0;
+			ConnectionState->CurrentPacketStartIndex[Trace::ENetProfilerConnectionMode::Outgoing] = 0U;
+			ConnectionState->CurrentPacketStartIndex[Trace::ENetProfilerConnectionMode::Incoming] = 0U;
 		}
 		break;
 
@@ -290,20 +291,33 @@ bool FNetTraceAnalyzer::OnEvent(uint16 RouteId, const FOnEventContext& Context)
 			const uint16 ConnectionId = EventData.GetValue<uint16>("ConnectionId");
 
 			TSharedRef<FNetTraceGameInstanceState> GameInstanceState = GetOrCreateActiveGameInstanceState(GameInstanceId);
-			check(!GameInstanceState->ActiveObjects.Contains(ObjectId));
-
-			// Add persistent object representation
-			Trace::FNetProfilerObjectInstance& ObjectInstance = NetProfilerProvider.CreateObject(GameInstanceState->GameInstanceIndex);
-
-			// Fill in object data
-			ObjectInstance.LifeTime.Begin =  GetLastTimestamp();
-			const uint32* NetProfilerNameIndex = TracedNameIdToNetProfilerNameIdMap.Find(NameId);
-			ObjectInstance.NameIndex = NetProfilerNameIndex ? *NetProfilerNameIndex : 0u;
-			ObjectInstance.NetId = ObjectId;
-			ObjectInstance.TypeId = TypeId;
-
-			// Add to active objects
-			GameInstanceState->ActiveObjects.Add(ObjectId, { ObjectInstance.ObjectIndex, ObjectInstance.NameIndex });
+			if (GameInstanceState->ActiveObjects.Contains(ObjectId))
+			{
+				// Verify that this is a reference that has been updated to a replicated object
+				Trace::FNetProfilerObjectInstance& ObjectInstance = *NetProfilerProvider.EditObject(GameInstanceState->GameInstanceIndex, GameInstanceState->ActiveObjects[ObjectId].ObjectIndex);
+				
+				// Update data 
+				ObjectInstance.LifeTime.Begin =  GetLastTimestamp();
+				const uint32* NetProfilerNameIndex = TracedNameIdToNetProfilerNameIdMap.Find(NameId);
+				ObjectInstance.NameIndex = NetProfilerNameIndex ? *NetProfilerNameIndex : 0u;
+				ObjectInstance.NetId = ObjectId;
+				ObjectInstance.TypeId = TypeId;
+			}
+			else
+			{
+				// Add persistent object representation
+				Trace::FNetProfilerObjectInstance& ObjectInstance = NetProfilerProvider.CreateObject(GameInstanceState->GameInstanceIndex);
+
+				// Fill in object data
+				ObjectInstance.LifeTime.Begin =  GetLastTimestamp();
+				const uint32* NetProfilerNameIndex = TracedNameIdToNetProfilerNameIdMap.Find(NameId);
+				ObjectInstance.NameIndex = NetProfilerNameIndex ? *NetProfilerNameIndex : 0u;
+				ObjectInstance.NetId = ObjectId;
+				ObjectInstance.TypeId = TypeId;
+
+				// Add to active objects
+				GameInstanceState->ActiveObjects.Add(ObjectId, { ObjectInstance.ObjectIndex, ObjectInstance.NameIndex });
+			}
 		}
 		break;
 
diff --git a/Engine/Source/Developer/TraceServices/Private/Analyzers/NetTraceAnalyzer.h b/Engine/Source/Developer/TraceServices/Private/Analyzers/NetTraceAnalyzer.h
--- a/Engine/Source/Developer/TraceServices/Private/Analyzers/NetTraceAnalyzer.h
+++ b/Engine/Source/Developer/TraceServices/Private/Analyzers/NetTraceAnalyzer.h
@@ -52,7 +52,7 @@ private:
 		uint32 ConnectionIndex;
 
 		// Current packet data
-		uint32 CurrentPacketStartIndex;
+		uint32 CurrentPacketStartIndex[Trace::ENetProfilerConnectionMode::Count];
 		Trace::ENetProfilerConnectionMode ConnectionMode;
 	};
 
