[TraceServices] Fixed SessionService to read up to 48K (instead of 32K) from a trace file when trying to find the session metadata events. Also removed check by trace version.

diff --git a/Engine/Source/Developer/TraceServices/Private/SessionService.cpp b/Engine/Source/Developer/TraceServices/Private/SessionService.cpp
--- a/Engine/Source/Developer/TraceServices/Private/SessionService.cpp
+++ b/Engine/Source/Developer/TraceServices/Private/SessionService.cpp
@@ -28,10 +28,10 @@ struct FDiagnosticsSessionAnalyzer
 {
 	virtual void OnAnalysisBegin(const FOnAnalysisContext& Context) override
 	{
-		if (Context.SessionContext.Version < 2)
-		{
-			return;
-		}
+		//if (Context.SessionContext.Version < 2)
+		//{
+		//	return;
+		//}
 
 		Context.InterfaceBuilder.RouteEvent(0, "Diagnostics", "Session");
 	}
@@ -331,7 +331,7 @@ void FSessionService::UpdateSessionContext(FStoreSessionHandle StoreHandle, FSes
 	{
 		virtual int32 Read(void* Data, uint32 Size) override
 		{
-			if (BytesRead >= 32768)
+			if (BytesRead >= 48 * 1024)
 			{
 				return 0;
 			}
