UnrealInsights:
Allow reading from session while holding write lock

diff --git a/Engine/Source/Developer/TraceServices/Private/AnalysisService.cpp b/Engine/Source/Developer/TraceServices/Private/AnalysisService.cpp
--- a/Engine/Source/Developer/TraceServices/Private/AnalysisService.cpp
+++ b/Engine/Source/Developer/TraceServices/Private/AnalysisService.cpp
@@ -26,7 +26,7 @@ thread_local int32 GThreadCurrentWriteLockCount;
 
 void FAnalysisSessionLock::ReadAccessCheck() const
 {
-	checkf(GThreadCurrentSessionLock == this && GThreadCurrentReadLockCount > 0, TEXT("Trying to read from session outside of a ReadScope"));
+	checkf(GThreadCurrentSessionLock == this && (GThreadCurrentReadLockCount > 0 || GThreadCurrentWriteLockCount > 0) , TEXT("Trying to read from session outside of a ReadScope"));
 }
 
 void FAnalysisSessionLock::WriteAccessCheck() const
