UnrealInsights:
Gracefully handle missing bookmark specs

diff --git a/Engine/Source/Developer/TraceServices/Private/Model/Bookmarks.cpp b/Engine/Source/Developer/TraceServices/Private/Model/Bookmarks.cpp
--- a/Engine/Source/Developer/TraceServices/Private/Model/Bookmarks.cpp
+++ b/Engine/Source/Developer/TraceServices/Private/Model/Bookmarks.cpp
@@ -26,6 +26,8 @@ FBookmarkSpec& FBookmarkProvider::GetSpec(uint64 BookmarkPoint)
 	else
 	{
 		TSharedPtr<FBookmarkSpec> Spec = MakeShared<FBookmarkSpec>();
+		Spec->File = TEXT("<unknown>");
+		Spec->FormatString = TEXT("<unknown>");
 		SpecMap.Add(BookmarkPoint, Spec);
 		return *Spec.Get();
 	}
diff --git a/Engine/Source/Developer/TraceServices/Private/Model/BookmarksPrivate.h b/Engine/Source/Developer/TraceServices/Private/Model/BookmarksPrivate.h
--- a/Engine/Source/Developer/TraceServices/Private/Model/BookmarksPrivate.h
+++ b/Engine/Source/Developer/TraceServices/Private/Model/BookmarksPrivate.h
@@ -13,15 +13,15 @@ class FStringStore;
 
 struct FBookmarkSpec
 {
-	const TCHAR* File;
-	const TCHAR* FormatString;
-	int32 Line;
+	const TCHAR* File = nullptr;
+	const TCHAR* FormatString = nullptr;
+	int32 Line = 0;
 };
 
 struct FBookmarkInternal
 {
-	double Time;
-	const TCHAR* Text;
+	double Time = 0.0;
+	const TCHAR* Text = nullptr;
 };
 
 class FBookmarkProvider
