NetTraceAnalyzer now supports nested packet content events reported from bunches
Broke up larger case statements into separate handler rmethods

diff --git a/Engine/Source/Developer/TraceServices/Private/Analyzers/NetTraceAnalyzer.cpp b/Engine/Source/Developer/TraceServices/Private/Analyzers/NetTraceAnalyzer.cpp
--- a/Engine/Source/Developer/TraceServices/Private/Analyzers/NetTraceAnalyzer.cpp
+++ b/Engine/Source/Developer/TraceServices/Private/Analyzers/NetTraceAnalyzer.cpp
@@ -5,6 +5,9 @@
 #include "TraceServices/Model/Threads.h"
 #include "Logging/LogMacros.h"
 
+DECLARE_LOG_CATEGORY_EXTERN(LogNetTrace, Log, All);
+DEFINE_LOG_CATEGORY(LogNetTrace);
+
 FNetTraceAnalyzer::FNetTraceAnalyzer(Trace::IAnalysisSession& InSession, Trace::FNetProfilerProvider& InNetProfilerProvider)
 	: Session(InSession)
 	, NetProfilerProvider(InNetProfilerProvider)
@@ -29,6 +32,12 @@ void FNetTraceAnalyzer::OnAnalysisBegin(const FOnAnalysisContext& Context)
 	Builder.RouteEvent(RouteId_ConnectionClosedEvent, "NetTrace", "ConnectionClosedEvent");
 	Builder.RouteEvent(RouteId_ObjectCreatedEvent, "NetTrace", "ObjectCreatedEvent");
 	Builder.RouteEvent(RouteId_ObjectDestroyedEvent, "NetTrace", "ObjectDestroyedEvent");
+
+	// Default names
+	{
+		Trace::FAnalysisSessionEditScope _(Session);
+		BunchHeaderNameIndex = NetProfilerProvider.AddNetProfilerName(TEXT("BunchHeader"));
+	}
 }
 
 void FNetTraceAnalyzer::OnAnalysisEnd()
@@ -68,20 +77,20 @@ bool FNetTraceAnalyzer::OnEvent(uint16 RouteId, const FOnEventContext& Context)
 	{
 		case RouteId_InitEvent:
 		{
-			// we always trace the version so that we make sure that we are backwards compatible with older trace stream
-			NetTraceVersion = EventData.GetValue<uint32>("NetTraceVersion");
-			NetTraceReporterVersion = EventData.GetValue<uint32>("NetTraceReporterVersion");
 			const uint64 TimestampCycles = EventData.GetValue<uint64>("Timestamp");
-
 			LastTimeStamp = Context.SessionContext.TimestampFromCycle(TimestampCycles);
 
+			// we always trace the version so that we make sure that we are backwards compatible with older trace stream
+			NetTraceVersion = EventData.GetValue<uint32>("NetTraceVersion");
+			NetTraceReporterVersion = EventData.GetValue<uint32>("NetTraceReporterVersion");
+			
 			NetProfilerProvider.SetNetTraceVersion(NetTraceVersion);
 		}
 		break;
 
 		case RouteId_InstanceDestroyedEvent:
 		{
-			const uint8 GameInstanceId = EventData.GetValue<uint8>("ReplicationSystemId");
+			const uint8 GameInstanceId = EventData.GetValue<uint8>("GameInstanceId");
 			DestroyActiveGameInstanceState(GameInstanceId);
 		}
 		break;
@@ -103,57 +112,106 @@ bool FNetTraceAnalyzer::OnEvent(uint16 RouteId, const FOnEventContext& Context)
 
 		case RouteId_PacketContentEvent:
 		{
-			const uint16 ConnectionId = EventData.GetValue<uint16>("ConnectionId");
-			const uint8 GameInstanceId = EventData.GetValue<uint8>("ReplicationSystemId");
-			const uint8 PacketType =  EventData.GetValue<uint8>("PacketType");
+			HandlePacketContentEvent(Context, EventData);
+		}
+		break;
+
+		case RouteId_PacketEvent:
+		{
+			HandlePacketEvent(Context, EventData);
+		}
+		break;
+
+		case RouteId_PacketDroppedEvent:
+		{
+			HandlePacketDroppedEvent(Context, EventData);
+		}
+		break;
+
+		case RouteId_ConnectionCreatedEvent:
+		{
+			HandleConnectionCretedEvent(Context, EventData);
+		}
+		break;
 
-			TSharedRef<FNetTraceGameInstanceState> GameInstanceState = GetOrCreateActiveGameInstanceState(GameInstanceId);
-			TSharedRef<FNetTraceConnectionState> ConnectionState = GetActiveConnectionState(GameInstanceId, ConnectionId);
+		case RouteId_ConnectionClosedEvent:
+		{
+			HandleConnectionClosedEvent(Context, EventData);
+		}
+		break;
 
-			Trace::ENetProfilerConnectionMode ConnectionMode = Trace::ENetProfilerConnectionMode(PacketType);
+		case RouteId_ObjectCreatedEvent:
+		{
+			HandleObjectCreatedEvent(Context, EventData);
+		}
+		break;
 
-			Trace::FNetProfilerConnectionData& ConnectionData = NetProfilerProvider.EditConnectionData(ConnectionState->ConnectionIndex, ConnectionMode);
-			++ConnectionData.ContentEventChangeCount;
+		case RouteId_ObjectDestroyedEvent:
+		{
+			HandleObjectDestroyedEvent(Context, EventData);
+		}
+		break;
+	}
 
-			TPagedArray<Trace::FNetProfilerContentEvent>& Events = ConnectionData.ContentEvents;
+	return true;
+}
 
-			// New Packet, update state
-			if (ConnectionState->CurrentPacketStartIndex[ConnectionMode] == Events.Num())
-			{
-				ConnectionState->ConnectionMode = ConnectionMode;
-			}
-			else
-			{
-				// Corrupt data
-				check(ConnectionState->CurrentPacketStartIndex[ConnectionMode] == Events.Num() || (ConnectionState->ConnectionMode == ConnectionMode));
-			}
+void FNetTraceAnalyzer::HandlePacketContentEvent(const FOnEventContext& Context, const FEventData& EventData)
+{
+	const uint16 ConnectionId = EventData.GetValue<uint16>("ConnectionId");
+	const uint8 GameInstanceId = EventData.GetValue<uint8>("GameInstanceId");
+	const uint8 PacketType =  EventData.GetValue<uint8>("PacketType");
+
+	//UE_LOG(LogNetTrace, Display, TEXT("FNetTraceAnalyzer::HandlePacketContentEvent: GameInstanceId: %u, ConnectionId: %u, %s"), GameInstanceId, ConnectionId, PacketType ? TEXT("Incoming") : TEXT("Outgoing"));
 
-			// Decode batched events
-			uint64 BufferSize = EventData.GetAttachmentSize();
-			const uint8* BufferPtr = EventData.GetAttachment();
-			const uint8* BufferEnd = BufferPtr + BufferSize;
-			uint64 LastOffset = 0;
+	TSharedRef<FNetTraceGameInstanceState> GameInstanceState = GetOrCreateActiveGameInstanceState(GameInstanceId);
+	TSharedRef<FNetTraceConnectionState> ConnectionState = GetActiveConnectionState(GameInstanceId, ConnectionId);
+
+	const Trace::ENetProfilerConnectionMode ConnectionMode = Trace::ENetProfilerConnectionMode(PacketType);
+
+	Trace::FNetProfilerConnectionData& ConnectionData = NetProfilerProvider.EditConnectionData(ConnectionState->ConnectionIndex, ConnectionMode);
+	++ConnectionData.ContentEventChangeCount;
+
+	TArray<Trace::FNetProfilerContentEvent>& Events = (ConnectionState->BunchEvents)[ConnectionMode];
+	TArray<FBunchInfo>& BunchInfos = (ConnectionState->BunchInfos)[ConnectionMode];
 
-			while (BufferPtr < BufferEnd)
+	// Decode batched events
+	uint64 BufferSize = EventData.GetAttachmentSize();
+	const uint8* BufferPtr = EventData.GetAttachment();
+	const uint8* BufferEnd = BufferPtr + BufferSize;
+	uint64 LastOffset = 0;
+
+	uint64 CurrentBunchOffset = 0U;
+
+	while (BufferPtr < BufferEnd)
+	{
+		// Decode data
+		const uint8 DecodedEventType = *BufferPtr++;
+
+		switch (DecodedEventType)
+		{
+			case EContentEventType::Object:
+			case EContentEventType::NameId:
 			{
-				// Decode data
+				Trace::FNetProfilerContentEvent& Event = Events.Emplace_GetRef();
+
+				const uint8 DecodedNestingLevel = *BufferPtr++;
+
 				const uint64 DecodedNameOrObjectId = FTraceAnalyzerUtils::Decode7bit(BufferPtr);
 				const uint64 DecodedEventStartPos = FTraceAnalyzerUtils::Decode7bit(BufferPtr) + LastOffset;
 				LastOffset = DecodedEventStartPos;
 				const uint64 DecodedEventEndPos = FTraceAnalyzerUtils::Decode7bit(BufferPtr) + DecodedEventStartPos;
-				const uint8 DecodedEventType = *BufferPtr++;
-				const uint8 DecodedNestingLevel = *BufferPtr++;
 
-				// Store the event, should we store it immediately? probably.
-				Trace::FNetProfilerContentEvent& Event = Events.PushBack();
-
-				Event.StartPos = DecodedEventStartPos;
-				Event.EndPos = DecodedEventEndPos;
+				// Fill in event data
+				Event.StartPos = DecodedEventStartPos + CurrentBunchOffset;
+				Event.EndPos = DecodedEventEndPos + CurrentBunchOffset;
 				Event.Level = DecodedNestingLevel;
 
 				Event.ObjectInstanceIndex = 0;
 				Event.NameIndex = 0;
 
+				checkSlow(Event.EndPos > Event.StartPos);
+
 				if (DecodedEventType == 0)
 				{
 					// Object index, need to lookup name indirectly
@@ -169,180 +227,336 @@ bool FNetTraceAnalyzer::OnEvent(uint16 RouteId, const FOnEventContext& Context)
 					{
 						Event.NameIndex = *NetProfilerNameIndex;
 					}
+					else
+					{
+						UE_LOG(LogNetTrace, Warning, TEXT("PacketContentEvent GameInstanceId: %u, ConnectionId: %u %s, Missing NameIndex: %u"), GameInstanceId, ConnectionId, ConnectionMode ? TEXT("Incoming") : TEXT("Outgoing"), DecodedNameOrObjectId);	
+					}
 				}
 
 				// EventTypeIndex does not match NameIndex as we might see the same name on different levels
 				Event.EventTypeIndex = GetTracedEventTypeIndex(Event.NameIndex, Event.Level);
 			}
-			check(BufferPtr == BufferEnd);
-		}
-		break;
+			break;
 
-		case RouteId_PacketEvent:
-		{
-			const uint64 TimestampCycles = EventData.GetValue<uint64>("Timestamp");
-			const uint32 SequenceNumber = EventData.GetValue<uint32>("SequenceNumber");
-			const uint8 GameInstanceId = EventData.GetValue<uint8>("ReplicationSystemId");
-			const uint16 ConnectionId = EventData.GetValue<uint16>("ConnectionId");
-
-			// Update LastTimestamp, later on we will be able to get timestamps piggybacked from other analyzers
-			LastTimeStamp = Context.SessionContext.TimestampFromCycle(TimestampCycles);
+			case EContentEventType::BunchEvent:
+			{
+				// DebugName
+				const uint64 DecodedNameId = FTraceAnalyzerUtils::Decode7bit(BufferPtr);
 
-			TSharedRef<FNetTraceGameInstanceState> GameInstanceState = GetOrCreateActiveGameInstanceState(GameInstanceId);
-			TSharedRef<FNetTraceConnectionState> ConnectionState = GetActiveConnectionState(GameInstanceId, ConnectionId);
-
-			// Add the packet
-			Trace::FNetProfilerConnectionData& ConnectionData = NetProfilerProvider.EditConnectionData(ConnectionState->ConnectionIndex, ConnectionState->ConnectionMode);
-			Trace::FNetProfilerPacket& Packet = ConnectionData.Packets.PushBack();
-			++ConnectionData.PacketChangeCount;
-
-			// Fill in packet data a packet has at least 1 event
-			Packet.StartEventIndex = ConnectionState->CurrentPacketStartIndex[ConnectionState->ConnectionMode];
-			Packet.EventCount = ConnectionData.ContentEvents.Num() - Packet.StartEventIndex;
-			Packet.TimeStamp = GetLastTimestamp();
-			Packet.SequenceNumber = SequenceNumber;
-			Packet.DeliveryStatus = Trace::ENetProfilerDeliveryStatus::Unknown;
-			Packet.ContentSizeInBits = Packet.EventCount ? ConnectionData.ContentEvents[Packet.StartEventIndex + Packet.EventCount - 1].EndPos : 0u;
-			Packet.TotalPacketSizeInBytes = (Packet.ContentSizeInBits + 7u) >> 3u;
-			Packet.DeliveryStatus = Trace::ENetProfilerDeliveryStatus::Delivered;
-
-			// Finalize received packet data and update data in the NetProfilerProvider
-			//NetProfilerProvider->FinalizePacket();
-
-			// Mark the beginning of a new packet
-			ConnectionState->CurrentPacketStartIndex[ConnectionState->ConnectionMode] = ConnectionData.ContentEvents.Num();
-		}
-		break;
+				const uint64 DecodedEventStartPos = FTraceAnalyzerUtils::Decode7bit(BufferPtr);
+				const uint64 DecodedEventEndPos = FTraceAnalyzerUtils::Decode7bit(BufferPtr) + DecodedEventStartPos;
 
-		case RouteId_PacketDroppedEvent:
-		{
-			const uint64 TimestampCycles = EventData.GetValue<uint64>("Timestamp");
-			const uint32 SequenceNumber = EventData.GetValue<uint32>("SequenceNumber");
-			const uint8 GameInstanceId = EventData.GetValue<uint8>("ReplicationSystemId");
-			const uint16 ConnectionId = EventData.GetValue<uint16>("ConnectionId");
-			const uint8 PacketType = EventData.GetValue<uint8>("PacketType");
+				const uint32* NetProfilerNameIndex = TracedNameIdToNetProfilerNameIdMap.Find(DecodedNameId);
 
-			// Update LastTimestamp, later on we will be able to get timestamps piggybacked from other analyzers
-			LastTimeStamp = Context.SessionContext.TimestampFromCycle(TimestampCycles);
+				FBunchInfo BunchInfo;
 
-			TSharedRef<FNetTraceGameInstanceState> GameInstanceState = GetOrCreateActiveGameInstanceState(GameInstanceId);
-			TSharedRef<FNetTraceConnectionState> ConnectionState = GetActiveConnectionState(GameInstanceId, ConnectionId);
+				BunchInfo.ChannelIndex = 0;//ChannelIndex;
+				BunchInfo.HeaderBits = 0;
+				BunchInfo.BunchBits = DecodedEventEndPos;
+				BunchInfo.FirstBunchEventIndex = Events.Num();
+				BunchInfo.NameIndex = NetProfilerNameIndex ? *NetProfilerNameIndex : 0u;
 
-			Trace::FNetProfilerConnectionData& ConnectionData = NetProfilerProvider.EditConnectionData(ConnectionState->ConnectionIndex, Trace::ENetProfilerConnectionMode(PacketType));
+				BunchInfos.Add(BunchInfo);
 
-			// Update packet delivery status
-			NetProfilerProvider.EditPacketDeliveryStatus(ConnectionState->ConnectionIndex, Trace::ENetProfilerConnectionMode(PacketType), SequenceNumber, Trace::ENetProfilerDeliveryStatus::Dropped);
-		}
-		break;
+				// Must reset LastOffset after reading bunch data
+				LastOffset = 0U;
+			}
+			break;
 
-		case RouteId_ConnectionCreatedEvent:
-		{
-			const uint8 GameInstanceId = EventData.GetValue<uint8>("ReplicationSystemId");
-			const uint16 ConnectionId = EventData.GetValue<uint16>("ConnectionId");
-
-			TSharedRef<FNetTraceGameInstanceState> GameInstanceState = GetOrCreateActiveGameInstanceState(GameInstanceId);
-			check(!GameInstanceState->ActiveConnections.Contains(ConnectionId));
-
-			// Add to both active connections and to persistent connections
- 			Trace::FNetProfilerConnectionInternal& Connection = NetProfilerProvider.CreateConnection(GameInstanceState->GameInstanceIndex);
-			TSharedRef<FNetTraceConnectionState> ConnectionState = MakeShared<FNetTraceConnectionState>();
-			GameInstanceState->ActiveConnections.Add(ConnectionId, ConnectionState);
-
-			// Fill in Connection data
-			Connection.Connection.ConnectionId = ConnectionId;
-			Connection.Connection.LifeTime.Begin =  GetLastTimestamp();
-			ConnectionState->ConnectionIndex = Connection.Connection.ConnectionIndex;
-			ConnectionState->CurrentPacketStartIndex[Trace::ENetProfilerConnectionMode::Outgoing] = 0U;
-			ConnectionState->CurrentPacketStartIndex[Trace::ENetProfilerConnectionMode::Incoming] = 0U;
-		}
-		break;
+			case EContentEventType::BunchHeaderEvent:
+			{
+				// HeaderSize if any
+				const uint64 DecodedEventCount = FTraceAnalyzerUtils::Decode7bit(BufferPtr);
+				const uint64 DecodedHeaderBits = FTraceAnalyzerUtils::Decode7bit(BufferPtr);
 
-		case RouteId_ConnectionClosedEvent:
-		{
-			const uint8 GameInstanceId = EventData.GetValue<uint8>("ReplicationSystemId");
-			const uint16 ConnectionId = EventData.GetValue<uint16>("ConnectionId");
+				FBunchInfo& BunchInfo = BunchInfos.Last();
 
-			TSharedRef<FNetTraceGameInstanceState> GameInstanceState = GetOrCreateActiveGameInstanceState(GameInstanceId);
+				BunchInfo.EventCount = DecodedEventCount;
+				if (DecodedEventCount > 0)
+				{
+					BunchInfo.FirstBunchEventIndex = Events.Num() - DecodedEventCount;
+				}
+				else
+				{
+					BunchInfo.FirstBunchEventIndex = Events.Num();
+				}
 
-			if (TSharedRef<FNetTraceConnectionState>* ConnectionState = GameInstanceState->ActiveConnections.Find(ConnectionId))
-			{
-				if (Trace::FNetProfilerConnectionInternal* Connection = NetProfilerProvider.EditConnection((*ConnectionState)->ConnectionIndex))
+				if (DecodedHeaderBits)
 				{
-					// Update connection state
-					Connection->Connection.LifeTime.End =  GetLastTimestamp();
+					// A bunch with header bits set is an actual bunch
+					BunchInfo.HeaderBits = DecodedHeaderBits;
+					CurrentBunchOffset = 0U;
+				}
+				else
+				{
+					// Merged bunch, set offset for events
+					CurrentBunchOffset = BunchInfo.BunchBits;
 				}
-				GameInstanceState->ActiveConnections.Remove(ConnectionId);
-			}
-			else
-			{
-				// Incomplete trace?  Ignore?
-				check(false);
 			}
-		}
-		break;
+			break;
+		};
 
-		case RouteId_ObjectCreatedEvent:
+	}
+	check(BufferPtr == BufferEnd);
+}
+
+void FNetTraceAnalyzer::AddEvent(TPagedArray<Trace::FNetProfilerContentEvent>& Events, const Trace::FNetProfilerContentEvent& InEvent, uint32 Offset, uint32 LevelOffset)
+{
+	Trace::FNetProfilerContentEvent& Event = Events.PushBack();
+	
+	Event.EventTypeIndex = GetTracedEventTypeIndex(InEvent.NameIndex, InEvent.Level + LevelOffset);
+	Event.NameIndex =  InEvent.NameIndex;
+	Event.ObjectInstanceIndex = InEvent.ObjectInstanceIndex;
+	Event.StartPos = InEvent.StartPos + Offset;
+	Event.EndPos = InEvent.EndPos + Offset;
+	Event.Level = InEvent.Level + LevelOffset;	
+	Event.ParentIndex = InEvent.ParentIndex;
+}
+
+void FNetTraceAnalyzer::AddEvent(TPagedArray<Trace::FNetProfilerContentEvent>& Events, uint32 StartPos, uint32 EndPos, uint32 Level, uint32 NameIndex)
+{
+	Trace::FNetProfilerContentEvent& Event = Events.PushBack();
+
+	Event.EventTypeIndex = GetTracedEventTypeIndex(NameIndex, Level);
+	Event.NameIndex = NameIndex; 
+	Event.ObjectInstanceIndex = 0;
+	Event.StartPos = StartPos;
+	Event.EndPos = EndPos;			
+	Event.Level = Level;
+	Event.ParentIndex = 0;
+}
+
+void FNetTraceAnalyzer::FlushPacketEvents(FNetTraceConnectionState& ConnectionState, Trace::FNetProfilerConnectionData& ConnectionData, const Trace::ENetProfilerConnectionMode ConnectionMode)
+{
+	TPagedArray<Trace::FNetProfilerContentEvent>& Events = ConnectionData.ContentEvents;
+
+	TArray<Trace::FNetProfilerContentEvent>& BunchEvents = ConnectionState.BunchEvents[ConnectionMode];
+	const int32 NumPacketEvents = BunchEvents.Num();
+
+	int32 CurrentBunchEventIndex = 0;
+
+	// Track bunch offsets
+	uint32 NextBunchOffset = 0U;
+	uint32 NextEventOffset = 0U;
+
+	int32 NonBunchEventCount = ConnectionState.BunchInfos[ConnectionMode].Num() ? ConnectionState.BunchInfos[ConnectionMode][0].FirstBunchEventIndex : BunchEvents.Num();
+
+	// Inject any events reported before the first bunch
+	while (CurrentBunchEventIndex < NonBunchEventCount)
+	{
+		const Trace::FNetProfilerContentEvent& BunchEvent = BunchEvents[CurrentBunchEventIndex];
+
+		AddEvent(Events, BunchEvent, 0U, 0U);
+	
+		NextBunchOffset = FMath::Max<uint64>(BunchEvent.EndPos, NextBunchOffset);
+		++CurrentBunchEventIndex;
+	}
+
+	uint32 EventsToAdd = 0U;
+	for (const FBunchInfo& Bunch : ConnectionState.BunchInfos[ConnectionMode])
+	{
+		uint32 BunchOffset = NextBunchOffset + Bunch.HeaderBits;
+		EventsToAdd += Bunch.EventCount;
+
+		// Report events for committed bunches
+		if (Bunch.HeaderBits)
 		{
-			const uint64 TypeId = EventData.GetValue<uint64>("TypeId");
-			const uint32 ObjectId = EventData.GetValue<uint32>("ObjectId");
-			const uint32 OwnerId = EventData.GetValue<uint32>("OwnerId");
-			const uint16 NameId = EventData.GetValue<uint16>("NameId");
-			const uint8 GameInstanceId = EventData.GetValue<uint8>("ReplicationSystemId");
-			const uint16 ConnectionId = EventData.GetValue<uint16>("ConnectionId");
-
-			TSharedRef<FNetTraceGameInstanceState> GameInstanceState = GetOrCreateActiveGameInstanceState(GameInstanceId);
-			if (GameInstanceState->ActiveObjects.Contains(ObjectId))
+			// Bunch event
+			AddEvent(Events, NextBunchOffset, NextBunchOffset + Bunch.HeaderBits + Bunch.BunchBits, 0, Bunch.NameIndex);
+
+			// Bunch header event
+			AddEvent(Events, NextBunchOffset, NextBunchOffset + Bunch.HeaderBits, 1, BunchHeaderNameIndex);
+	
+			// Add events belonging to bunch, including the ones from merged bunches
+			for (uint32 EventIt = 0; EventIt < EventsToAdd; ++EventIt)
 			{
-				// Verify that this is a reference that has been updated to a replicated object
-				Trace::FNetProfilerObjectInstance& ObjectInstance = *NetProfilerProvider.EditObject(GameInstanceState->GameInstanceIndex, GameInstanceState->ActiveObjects[ObjectId].ObjectIndex);
-				
-				// Update data 
-				ObjectInstance.LifeTime.Begin =  GetLastTimestamp();
-				const uint32* NetProfilerNameIndex = TracedNameIdToNetProfilerNameIdMap.Find(NameId);
-				ObjectInstance.NameIndex = NetProfilerNameIndex ? *NetProfilerNameIndex : 0u;
-				ObjectInstance.NetId = ObjectId;
-				ObjectInstance.TypeId = TypeId;
-			}
-			else
-			{
-				// Add persistent object representation
-				Trace::FNetProfilerObjectInstance& ObjectInstance = NetProfilerProvider.CreateObject(GameInstanceState->GameInstanceIndex);
-
-				// Fill in object data
-				ObjectInstance.LifeTime.Begin =  GetLastTimestamp();
-				const uint32* NetProfilerNameIndex = TracedNameIdToNetProfilerNameIdMap.Find(NameId);
-				ObjectInstance.NameIndex = NetProfilerNameIndex ? *NetProfilerNameIndex : 0u;
-				ObjectInstance.NetId = ObjectId;
-				ObjectInstance.TypeId = TypeId;
-
-				// Add to active objects
-				GameInstanceState->ActiveObjects.Add(ObjectId, { ObjectInstance.ObjectIndex, ObjectInstance.NameIndex });
+				const Trace::FNetProfilerContentEvent& BunchEvent = BunchEvents[CurrentBunchEventIndex];
+
+				AddEvent(Events, BunchEvent, BunchOffset, 1U);
+				++CurrentBunchEventIndex;
 			}
+
+			// Accumulate offset
+			NextBunchOffset += Bunch.BunchBits + Bunch.HeaderBits;
+			NextEventOffset = NextBunchOffset;
+
+			// Reset event count
+			EventsToAdd = 0U;
 		}
-		break;
+		
+		++ConnectionData.ContentEventChangeCount;
+	}
 
-		case RouteId_ObjectDestroyedEvent:
-		{
-			// Remove from active instances and mark the end timestamp in the persistent instance list
-			const uint8 GameInstanceId = EventData.GetValue<uint8>("ReplicationSystemId");
-			const uint32 ObjectId = EventData.GetValue<uint32>("ObjectId");
+	ConnectionState.BunchEvents[ConnectionMode].Reset();
+	ConnectionState.BunchInfos[ConnectionMode].Reset();
+}
 
-			TSharedRef<FNetTraceGameInstanceState> GameInstanceState = GetOrCreateActiveGameInstanceState(GameInstanceId);
+void FNetTraceAnalyzer::HandlePacketEvent(const FOnEventContext& Context, const FEventData& EventData)
+{
+	const uint64 TimestampCycles = EventData.GetValue<uint64>("Timestamp");
+	const uint32 PacketBits = EventData.GetValue<uint32>("PacketBits");
+	const uint32 SequenceNumber = EventData.GetValue<uint32>("SequenceNumber");
+	const uint8 GameInstanceId = EventData.GetValue<uint8>("GameInstanceId");
+	const uint16 ConnectionId = EventData.GetValue<uint16>("ConnectionId");
+	const uint8 PacketType = EventData.GetValue<uint8>("PacketType");
 
-			FNetTraceActiveObjectState DestroyedObjectState;
-			if (GameInstanceState->ActiveObjects.RemoveAndCopyValue(ObjectId, DestroyedObjectState))
-			{
-				if (Trace::FNetProfilerObjectInstance* ObjectInstance = NetProfilerProvider.EditObject(GameInstanceState->GameInstanceIndex, DestroyedObjectState.ObjectIndex))
-				{
-					// Update object data
-					ObjectInstance->LifeTime.End = GetLastTimestamp();
-				}
-			}
+	const Trace::ENetProfilerConnectionMode ConnectionMode = Trace::ENetProfilerConnectionMode(PacketType);
+
+	// Update LastTimestamp, later on we will be able to get timestamps piggybacked from other analyzers
+	LastTimeStamp = Context.SessionContext.TimestampFromCycle(TimestampCycles);
+
+	TSharedRef<FNetTraceGameInstanceState> GameInstanceState = GetOrCreateActiveGameInstanceState(GameInstanceId);
+	TSharedRef<FNetTraceConnectionState> ConnectionState = GetActiveConnectionState(GameInstanceId, ConnectionId);
+
+	// Add the packet
+	Trace::FNetProfilerConnectionData& ConnectionData = NetProfilerProvider.EditConnectionData(ConnectionState->ConnectionIndex, ConnectionMode);
+	Trace::FNetProfilerPacket& Packet = ConnectionData.Packets.PushBack();
+	++ConnectionData.PacketChangeCount;
+
+	// Flush packet events
+	FlushPacketEvents(*ConnectionState, ConnectionData, ConnectionMode);
+
+	// Fill in packet data a packet must have at least 1 event?
+	Packet.StartEventIndex = ConnectionState->CurrentPacketStartIndex[ConnectionMode];
+	Packet.EventCount = ConnectionData.ContentEvents.Num() - Packet.StartEventIndex;
+	Packet.TimeStamp = GetLastTimestamp();
+	Packet.SequenceNumber = SequenceNumber;
+	Packet.DeliveryStatus = Trace::ENetProfilerDeliveryStatus::Unknown;
+
+	Packet.ContentSizeInBits = PacketBits;
+	Packet.TotalPacketSizeInBytes = (Packet.ContentSizeInBits + 7u) >> 3u;
+	Packet.DeliveryStatus = Trace::ENetProfilerDeliveryStatus::Delivered;
+
+	// Mark the beginning of a new packet
+	ConnectionState->CurrentPacketStartIndex[ConnectionMode] = ConnectionData.ContentEvents.Num();
+	ConnectionState->CurrentPacketBitOffset[ConnectionMode] = 0U;
+
+	//UE_LOG(LogNetTrace, Log, TEXT("PacketEvent GameInstanceId: %u, ConnectionId: %u, %s, Seq: %u PacketBits: %u"), GameInstanceId, ConnectionId, ConnectionMode ? TEXT("Incoming") : TEXT("Outgoing"), SequenceNumber, Packet.ContentSizeInBits);
+}
+
+void FNetTraceAnalyzer::HandlePacketDroppedEvent(const FOnEventContext& Context, const FEventData& EventData)
+{
+	const uint64 TimestampCycles = EventData.GetValue<uint64>("Timestamp");
+	const uint32 SequenceNumber = EventData.GetValue<uint32>("SequenceNumber");
+	const uint8 GameInstanceId = EventData.GetValue<uint8>("GameInstanceId");
+	const uint16 ConnectionId = EventData.GetValue<uint16>("ConnectionId");
+	const uint8 PacketType = EventData.GetValue<uint8>("PacketType");
+
+	// Update LastTimestamp, later on we will be able to get timestamps piggybacked from other analyzers
+	LastTimeStamp = Context.SessionContext.TimestampFromCycle(TimestampCycles);
+
+	TSharedRef<FNetTraceGameInstanceState> GameInstanceState = GetOrCreateActiveGameInstanceState(GameInstanceId);
+	TSharedRef<FNetTraceConnectionState> ConnectionState = GetActiveConnectionState(GameInstanceId, ConnectionId);
+
+	Trace::FNetProfilerConnectionData& ConnectionData = NetProfilerProvider.EditConnectionData(ConnectionState->ConnectionIndex, Trace::ENetProfilerConnectionMode(PacketType));
+
+	// Update packet delivery status
+	NetProfilerProvider.EditPacketDeliveryStatus(ConnectionState->ConnectionIndex, Trace::ENetProfilerConnectionMode(PacketType), SequenceNumber, Trace::ENetProfilerDeliveryStatus::Dropped);
+}
+
+void FNetTraceAnalyzer::HandleConnectionCretedEvent(const FOnEventContext& Context, const FEventData& EventData)
+{
+	const uint8 GameInstanceId = EventData.GetValue<uint8>("GameInstanceId");
+	const uint16 ConnectionId = EventData.GetValue<uint16>("ConnectionId");
+
+	TSharedRef<FNetTraceGameInstanceState> GameInstanceState = GetOrCreateActiveGameInstanceState(GameInstanceId);
+	check(!GameInstanceState->ActiveConnections.Contains(ConnectionId));
+
+	// Add to both active connections and to persistent connections
+ 	Trace::FNetProfilerConnectionInternal& Connection = NetProfilerProvider.CreateConnection(GameInstanceState->GameInstanceIndex);
+	TSharedRef<FNetTraceConnectionState> ConnectionState = MakeShared<FNetTraceConnectionState>();
+	GameInstanceState->ActiveConnections.Add(ConnectionId, ConnectionState);
+
+	// Fill in Connection data
+	Connection.Connection.ConnectionId = ConnectionId;
+	Connection.Connection.LifeTime.Begin =  GetLastTimestamp();
+	ConnectionState->ConnectionIndex = Connection.Connection.ConnectionIndex;
+	ConnectionState->CurrentPacketStartIndex[Trace::ENetProfilerConnectionMode::Outgoing] = 0U;
+	ConnectionState->CurrentPacketStartIndex[Trace::ENetProfilerConnectionMode::Incoming] = 0U;
+
+	ConnectionState->CurrentPacketBitOffset[Trace::ENetProfilerConnectionMode::Outgoing] = 0U;
+	ConnectionState->CurrentPacketBitOffset[Trace::ENetProfilerConnectionMode::Incoming] = 0U;
+}
+
+void FNetTraceAnalyzer::HandleConnectionClosedEvent(const FOnEventContext& Context, const FEventData& EventData)
+{
+	const uint8 GameInstanceId = EventData.GetValue<uint8>("GameInstanceId");
+	const uint16 ConnectionId = EventData.GetValue<uint16>("ConnectionId");
+
+	TSharedRef<FNetTraceGameInstanceState> GameInstanceState = GetOrCreateActiveGameInstanceState(GameInstanceId);
+
+	if (TSharedRef<FNetTraceConnectionState>* ConnectionState = GameInstanceState->ActiveConnections.Find(ConnectionId))
+	{
+		if (Trace::FNetProfilerConnectionInternal* Connection = NetProfilerProvider.EditConnection((*ConnectionState)->ConnectionIndex))
+		{
+			// Update connection state
+			Connection->Connection.LifeTime.End =  GetLastTimestamp();
 		}
-		break;
+		GameInstanceState->ActiveConnections.Remove(ConnectionId);
 	}
+	else
+	{
+		// Incomplete trace?  Ignore?
+		check(false);
+	}
+}
 
-	return true;
+void FNetTraceAnalyzer::HandleObjectCreatedEvent(const FOnEventContext& Context, const FEventData& EventData)
+{
+	const uint64 TypeId = EventData.GetValue<uint64>("TypeId");
+	const uint32 ObjectId = EventData.GetValue<uint32>("ObjectId");
+	const uint32 OwnerId = EventData.GetValue<uint32>("OwnerId");
+	const uint16 NameId = EventData.GetValue<uint16>("NameId");
+	const uint8 GameInstanceId = EventData.GetValue<uint8>("GameInstanceId");
+	const uint16 ConnectionId = EventData.GetValue<uint16>("ConnectionId");
+
+	TSharedRef<FNetTraceGameInstanceState> GameInstanceState = GetOrCreateActiveGameInstanceState(GameInstanceId);
+	if (GameInstanceState->ActiveObjects.Contains(ObjectId))
+	{
+		// Verify that this is a reference that has been updated to a replicated object
+		Trace::FNetProfilerObjectInstance& ObjectInstance = *NetProfilerProvider.EditObject(GameInstanceState->GameInstanceIndex, GameInstanceState->ActiveObjects[ObjectId].ObjectIndex);
+				
+		// Update data 
+		ObjectInstance.LifeTime.Begin =  GetLastTimestamp();
+		const uint32* NetProfilerNameIndex = TracedNameIdToNetProfilerNameIdMap.Find(NameId);
+		ObjectInstance.NameIndex = NetProfilerNameIndex ? *NetProfilerNameIndex : 0u;
+		ObjectInstance.NetId = ObjectId;
+		ObjectInstance.TypeId = TypeId;
+	}
+	else
+	{
+		// Add persistent object representation
+		Trace::FNetProfilerObjectInstance& ObjectInstance = NetProfilerProvider.CreateObject(GameInstanceState->GameInstanceIndex);
+
+		// Fill in object data
+		ObjectInstance.LifeTime.Begin =  GetLastTimestamp();
+		const uint32* NetProfilerNameIndex = TracedNameIdToNetProfilerNameIdMap.Find(NameId);
+		ObjectInstance.NameIndex = NetProfilerNameIndex ? *NetProfilerNameIndex : 0u;
+		ObjectInstance.NetId = ObjectId;
+		ObjectInstance.TypeId = TypeId;
+
+		// Add to active objects
+		GameInstanceState->ActiveObjects.Add(ObjectId, { ObjectInstance.ObjectIndex, ObjectInstance.NameIndex });
+	}
+}
+
+void FNetTraceAnalyzer::HandleObjectDestroyedEvent(const FOnEventContext& Context, const FEventData& EventData)
+{
+	// Remove from active instances and mark the end timestamp in the persistent instance list
+	const uint8 GameInstanceId = EventData.GetValue<uint8>("GameInstanceId");
+	const uint32 ObjectId = EventData.GetValue<uint32>("ObjectId");
+
+	TSharedRef<FNetTraceGameInstanceState> GameInstanceState = GetOrCreateActiveGameInstanceState(GameInstanceId);
+
+	FNetTraceActiveObjectState DestroyedObjectState;
+	if (GameInstanceState->ActiveObjects.RemoveAndCopyValue(ObjectId, DestroyedObjectState))
+	{
+		if (Trace::FNetProfilerObjectInstance* ObjectInstance = NetProfilerProvider.EditObject(GameInstanceState->GameInstanceIndex, DestroyedObjectState.ObjectIndex))
+		{
+			// Update object data
+			ObjectInstance->LifeTime.End = GetLastTimestamp();
+		}
+	}
 }
 
 TSharedRef<FNetTraceAnalyzer::FNetTraceGameInstanceState> FNetTraceAnalyzer::GetOrCreateActiveGameInstanceState(uint32 GameInstanceId)
diff --git a/Engine/Source/Developer/TraceServices/Private/Analyzers/NetTraceAnalyzer.h b/Engine/Source/Developer/TraceServices/Private/Analyzers/NetTraceAnalyzer.h
--- a/Engine/Source/Developer/TraceServices/Private/Analyzers/NetTraceAnalyzer.h
+++ b/Engine/Source/Developer/TraceServices/Private/Analyzers/NetTraceAnalyzer.h
@@ -36,16 +36,36 @@ private:
 		RouteId_ObjectDestroyedEvent,
 	};
 
+	// This must be kept in sync with Event types in NetTrace.h
+	enum class EContentEventType : uint8
+	{
+		Object = 0,
+		NameId = 1,
+		BunchEvent = 2,
+		BunchHeaderEvent = 3,
+	};
+
 	Trace::IAnalysisSession& Session;
 	Trace::FNetProfilerProvider& NetProfilerProvider;
 	uint32 NetTraceVersion;
 	uint32 NetTraceReporterVersion;
+	uint32 BunchHeaderNameIndex;
 
 	// Shared for trace
 	TMap<uint16, uint32> TracedNameIdToNetProfilerNameIdMap;
 
 	TMap<uint32, uint32> TraceEventTypeToNetProfilerEventTypeIndexMap;
 
+	struct FBunchInfo
+	{
+		int32 ChannelIndex;
+		uint32 BunchBits;
+		uint32 HeaderBits;
+		int32 FirstBunchEventIndex;
+		uint32 EventCount;
+		uint16 NameIndex;
+	};
+
 	struct FNetTraceConnectionState
 	{
 		// Index into persistent connection array
@@ -53,7 +73,12 @@ private:
 
 		// Current packet data
 		uint32 CurrentPacketStartIndex[Trace::ENetProfilerConnectionMode::Count];
-		Trace::ENetProfilerConnectionMode ConnectionMode;
+		uint32 CurrentPacketBitOffset[Trace::ENetProfilerConnectionMode::Count];
+
+		// Current bunch data
+		TArray<Trace::FNetProfilerContentEvent> BunchEvents[Trace::ENetProfilerConnectionMode::Count];
+
+		TArray<FBunchInfo> BunchInfos[Trace::ENetProfilerConnectionMode::Count];
 	};
 
 	struct FNetTraceActiveObjectState
@@ -66,11 +91,13 @@ private:
 	struct FNetTraceGameInstanceState
 	{
 		TMap<uint32, TSharedRef<FNetTraceConnectionState>> ActiveConnections;
-		TMap<uint32, FNetTraceActiveObjectState> ActiveObjects;
+		TMap<uint64, FNetTraceActiveObjectState> ActiveObjects;
 
 		uint32 GameInstanceIndex;
 	};
 
+private:
+
 	uint32 GetTracedEventTypeIndex(uint16 NameIndex, uint8 Level);
 
 	TSharedRef<FNetTraceGameInstanceState> GetOrCreateActiveGameInstanceState(uint32 GameInstanceId);
@@ -79,6 +106,21 @@ private:
 	TSharedRef<FNetTraceConnectionState> GetActiveConnectionState(uint32 GameInstanceId, uint32 ConnectionId);
 	Trace::FNetProfilerTimeStamp GetLastTimestamp() const { return LastTimeStamp; }
 
+	void FlushPacketEvents(FNetTraceConnectionState& ConnectionState, Trace::FNetProfilerConnectionData& ConnectionData, const Trace::ENetProfilerConnectionMode ConnectionMode);
+
+	void HandlePacketEvent(const FOnEventContext& Context, const FEventData& EventData);	
+	void HandlePacketContentEvent(const FOnEventContext& Context, const FEventData& EventData);
+	void HandlePacketDroppedEvent(const FOnEventContext& Context, const FEventData& EventData);
+	void HandleConnectionCretedEvent(const FOnEventContext& Context, const FEventData& EventData);
+	void HandleConnectionClosedEvent(const FOnEventContext& Context, const FEventData& EventData);
+	void HandleObjectCreatedEvent(const FOnEventContext& Context, const FEventData& EventData);
+	void HandleObjectDestroyedEvent(const FOnEventContext& Context, const FEventData& EventData);
+
+	void AddEvent(TPagedArray<Trace::FNetProfilerContentEvent>& Events, const Trace::FNetProfilerContentEvent& Event, uint32 Offset, uint32 LevelOffset);
+	void AddEvent(TPagedArray<Trace::FNetProfilerContentEvent>& Events, uint32 StartPos, uint32 EndPos, uint32 Level, uint32 NameIndex);
+
+private:
+
 	TMap<uint32, TSharedRef<FNetTraceGameInstanceState>> ActiveGameInstances;
 	Trace::FNetProfilerTimeStamp LastTimeStamp;
 };
diff --git a/Engine/Source/Developer/TraceServices/Private/Model/NetProfilerProvider.h b/Engine/Source/Developer/TraceServices/Private/Model/NetProfilerProvider.h
--- a/Engine/Source/Developer/TraceServices/Private/Model/NetProfilerProvider.h
+++ b/Engine/Source/Developer/TraceServices/Private/Model/NetProfilerProvider.h
@@ -114,7 +114,9 @@ public:
 	// Stats queries
 	virtual ITable<FNetProfilerAggregatedStats>* CreateAggregation(uint32 ConnectionIndex, ENetProfilerConnectionMode Mode, uint32 PacketIndexIntervalStart, uint32 PacketIndexIntervalEnd, uint32 StartPosition, uint32 EndPosition) const override;
 
+	const TCHAR* InternalGetNetProfilerName(uint32 NameIndex) const { return Names[NameIndex].Name; }
 private:
+	
 	const FNetProfilerName* GetNetProfilerName(uint32 ProfilerNameId) const;
 	const FNetProfilerEventType* GetNetProfilerEventType(uint32 ProfilerEventTypeId) const;
 
