Fixed compilation error with -strict

diff --git a/Engine/Source/Developer/TraceServices/Private/Analyzers/NetTraceAnalyzer.cpp b/Engine/Source/Developer/TraceServices/Private/Analyzers/NetTraceAnalyzer.cpp
--- a/Engine/Source/Developer/TraceServices/Private/Analyzers/NetTraceAnalyzer.cpp
+++ b/Engine/Source/Developer/TraceServices/Private/Analyzers/NetTraceAnalyzer.cpp
@@ -186,8 +186,7 @@ void FNetTraceAnalyzer::HandlePacketContentEvent(const FOnEventContext& Context,
 	while (BufferPtr < BufferEnd)
 	{
 		// Decode data
-		const uint8 DecodedEventType = *BufferPtr++;
-
+		const EContentEventType DecodedEventType = EContentEventType(*BufferPtr++);
 		switch (DecodedEventType)
 		{
 			case EContentEventType::Object:
@@ -212,7 +211,7 @@ void FNetTraceAnalyzer::HandlePacketContentEvent(const FOnEventContext& Context,
 
 				checkSlow(Event.EndPos > Event.StartPos);
 
-				if (DecodedEventType == 0)
+				if (DecodedEventType == EContentEventType::Object)
 				{
 					// Object index, need to lookup name indirectly
 					if (const FNetTraceActiveObjectState* ActiveObjectState = GameInstanceState->ActiveObjects.Find(DecodedNameOrObjectId))
@@ -221,7 +220,7 @@ void FNetTraceAnalyzer::HandlePacketContentEvent(const FOnEventContext& Context,
 						Event.ObjectInstanceIndex = ActiveObjectState->ObjectIndex;
 					}
 				}
-				else if (DecodedEventType == 1)
+				else if (DecodedEventType == EContentEventType::NameId)
 				{
 					if (const uint32* NetProfilerNameIndex = TracedNameIdToNetProfilerNameIdMap.Find(DecodedNameOrObjectId))
 					{
