No master fixes

diff --git a/Engine/Source/Developer/TraceServices/Private/Model/Channel.cpp b/Engine/Source/Developer/TraceServices/Private/Model/Channel.cpp
--- a/Engine/Source/Developer/TraceServices/Private/Model/Channel.cpp
+++ b/Engine/Source/Developer/TraceServices/Private/Model/Channel.cpp
@@ -1,4 +1,5 @@
 #include "Channel.h"
+#include "UObject/NameTypes.h"
 
 namespace Trace {
 
diff --git a/Engine/Source/Developer/TraceServices/Public/TraceServices/Model/Channel.h b/Engine/Source/Developer/TraceServices/Public/TraceServices/Model/Channel.h
--- a/Engine/Source/Developer/TraceServices/Public/TraceServices/Model/Channel.h
+++ b/Engine/Source/Developer/TraceServices/Public/TraceServices/Model/Channel.h
@@ -3,6 +3,7 @@
 #pragma once
 
 #include "TraceServices/Model/AnalysisSession.h"
+#include "Containers/UnrealString.h"
 
 namespace Trace {
 
