Copyright header fix.

diff --git a/Engine/Source/Developer/TraceServices/Private/Model/Channel.cpp b/Engine/Source/Developer/TraceServices/Private/Model/Channel.cpp
--- a/Engine/Source/Developer/TraceServices/Private/Model/Channel.cpp
+++ b/Engine/Source/Developer/TraceServices/Private/Model/Channel.cpp
@@ -1,3 +1,4 @@
+// Copyright Epic Games, Inc. All Rights Reserved.
 #include "Channel.h"
 #include "UObject/NameTypes.h"
 
diff --git a/Engine/Source/Developer/TraceServices/Public/TraceServices/Model/Channel.h b/Engine/Source/Developer/TraceServices/Public/TraceServices/Model/Channel.h
--- a/Engine/Source/Developer/TraceServices/Public/TraceServices/Model/Channel.h
+++ b/Engine/Source/Developer/TraceServices/Public/TraceServices/Model/Channel.h
@@ -1,4 +1,4 @@
-// Copyright 1998-2019 Epic Games, Inc. All Rights Reserved.
+// Copyright Epic Games, Inc. All Rights Reserved.
 
 #pragma once
 
diff --git a/Engine/Source/Runtime/TraceLog/Public/Trace/Detail/Channel.h b/Engine/Source/Runtime/TraceLog/Public/Trace/Detail/Channel.h
--- a/Engine/Source/Runtime/TraceLog/Public/Trace/Detail/Channel.h
+++ b/Engine/Source/Runtime/TraceLog/Public/Trace/Detail/Channel.h
@@ -1,4 +1,4 @@
-// Copyright 1998-2019 Epic Games, Inc. All Rights Reserved.
+// Copyright Epic Games, Inc. All Rights Reserved.
 
 #pragma once
 
diff --git a/Engine/Source/Runtime/TraceLog/Public/Trace/Detail/Channel.inl b/Engine/Source/Runtime/TraceLog/Public/Trace/Detail/Channel.inl
--- a/Engine/Source/Runtime/TraceLog/Public/Trace/Detail/Channel.inl
+++ b/Engine/Source/Runtime/TraceLog/Public/Trace/Detail/Channel.inl
@@ -1,4 +1,4 @@
-// Copyright 1998-2019 Epic Games, Inc. All Rights Reserved.
+// Copyright Epic Games, Inc. All Rights Reserved.
 
 #pragma once
 
