[NetTraceAnalyzer]
- Fixed issue with reusing object ids
- Channel names are now displayed properly

diff --git a/Engine/Source/Developer/TraceServices/Private/Analyzers/NetTraceAnalyzer.cpp b/Engine/Source/Developer/TraceServices/Private/Analyzers/NetTraceAnalyzer.cpp
--- a/Engine/Source/Developer/TraceServices/Private/Analyzers/NetTraceAnalyzer.cpp
+++ b/Engine/Source/Developer/TraceServices/Private/Analyzers/NetTraceAnalyzer.cpp
@@ -8,6 +8,13 @@
 DECLARE_LOG_CATEGORY_EXTERN(LogNetTrace, Log, All);
 DEFINE_LOG_CATEGORY(LogNetTrace);
 
+enum ENetTraceAnalyzerVersion
+{
+	ENetTraceAnalyzerVersion_Initial = 1,
+	ENetTraceAnalyzerVersion_BunchChannelIndex = 2,
+};
+
+
 FNetTraceAnalyzer::FNetTraceAnalyzer(Trace::IAnalysisSession& InSession, Trace::FNetProfilerProvider& InNetProfilerProvider)
 	: Session(InSession)
 	, NetProfilerProvider(InNetProfilerProvider)
@@ -239,21 +246,19 @@ void FNetTraceAnalyzer::HandlePacketContentEvent(const FOnEventContext& Context,
 
 			case EContentEventType::BunchEvent:
 			{
-				// DebugName
 				const uint64 DecodedNameId = FTraceAnalyzerUtils::Decode7bit(BufferPtr);
-
 				const uint64 DecodedEventStartPos = FTraceAnalyzerUtils::Decode7bit(BufferPtr);
 				const uint64 DecodedEventEndPos = FTraceAnalyzerUtils::Decode7bit(BufferPtr) + DecodedEventStartPos;
-
-				const uint32* NetProfilerNameIndex = TracedNameIdToNetProfilerNameIdMap.Find(DecodedNameId);
+				
+				const uint32* NetProfilerNameIndex = DecodedNameId ? TracedNameIdToNetProfilerNameIdMap.Find(DecodedNameId) : nullptr;
 
 				FBunchInfo BunchInfo;
 
-				BunchInfo.ChannelIndex = 0;//ChannelIndex;
-				BunchInfo.HeaderBits = 0;
+				BunchInfo.ChannelIndex = -1;
+				BunchInfo.HeaderBits = 0U;
 				BunchInfo.BunchBits = DecodedEventEndPos;
 				BunchInfo.FirstBunchEventIndex = Events.Num();
-				BunchInfo.NameIndex = NetProfilerNameIndex ? *NetProfilerNameIndex : 0u;
+				BunchInfo.NameIndex = NetProfilerNameIndex ? *NetProfilerNameIndex : 0U;
 
 				BunchInfos.Add(BunchInfo);
 
@@ -264,25 +269,38 @@ void FNetTraceAnalyzer::HandlePacketContentEvent(const FOnEventContext& Context,
 
 			case EContentEventType::BunchHeaderEvent:
 			{
-				// HeaderSize if any
 				const uint64 DecodedEventCount = FTraceAnalyzerUtils::Decode7bit(BufferPtr);
 				const uint64 DecodedHeaderBits = FTraceAnalyzerUtils::Decode7bit(BufferPtr);
 
 				FBunchInfo& BunchInfo = BunchInfos.Last();
 
 				BunchInfo.EventCount = DecodedEventCount;
-				if (DecodedEventCount > 0)
-				{
-					BunchInfo.FirstBunchEventIndex = Events.Num() - DecodedEventCount;
-				}
-				else
-				{
-					BunchInfo.FirstBunchEventIndex = Events.Num();
-				}
+				BunchInfo.FirstBunchEventIndex = Events.Num() - DecodedEventCount;
 
+				// A bunch with header bits set is an actual bunch
 				if (DecodedHeaderBits)
 				{
-					// A bunch with header bits set is an actual bunch
+					if (NetTraceVersion >= ENetTraceAnalyzerVersion_BunchChannelIndex)
+					{
+						const int32 DecodedChannelId = DecodedHeaderBits ? (int32)FTraceAnalyzerUtils::Decode7bit(BufferPtr) : -1;
+
+						BunchInfo.ChannelIndex = DecodedChannelId;
+
+						// Try to find bunch name using channelId if not included in the Bunch
+						if (DecodedChannelId != -1)
+						{					
+							if (BunchInfo.NameIndex)
+							{
+								GameInstanceState->ChannelNames.FindOrAdd(DecodedChannelId) = BunchInfo.NameIndex;
+							}
+							else
+							{
+								const uint32* ExistingChannelNameIndex = GameInstanceState->ChannelNames.Find(DecodedChannelId);						
+								BunchInfo.NameIndex = ExistingChannelNameIndex ? *ExistingChannelNameIndex : 0U;
+							}
+						}					
+					}
+
 					BunchInfo.HeaderBits = DecodedHeaderBits;
 					CurrentBunchOffset = 0U;
 				}
@@ -507,36 +525,42 @@ void FNetTraceAnalyzer::HandleObjectCreatedEvent(const FOnEventContext& Context,
 	const uint32 OwnerId = EventData.GetValue<uint32>("OwnerId");
 	const uint16 NameId = EventData.GetValue<uint16>("NameId");
 	const uint8 GameInstanceId = EventData.GetValue<uint8>("GameInstanceId");
-	const uint16 ConnectionId = EventData.GetValue<uint16>("ConnectionId");
 
 	TSharedRef<FNetTraceGameInstanceState> GameInstanceState = GetOrCreateActiveGameInstanceState(GameInstanceId);
+	const uint32* NetProfilerNameIndex = TracedNameIdToNetProfilerNameIdMap.Find(NameId);
+	const uint32 NameIndex = NetProfilerNameIndex ? *NetProfilerNameIndex : 0u;
+	
 	if (GameInstanceState->ActiveObjects.Contains(ObjectId))
 	{
-		// Verify that this is a reference that has been updated to a replicated object
-		Trace::FNetProfilerObjectInstance& ObjectInstance = *NetProfilerProvider.EditObject(GameInstanceState->GameInstanceIndex, GameInstanceState->ActiveObjects[ObjectId].ObjectIndex);
-				
-		// Update data 
-		ObjectInstance.LifeTime.Begin =  GetLastTimestamp();
-		const uint32* NetProfilerNameIndex = TracedNameIdToNetProfilerNameIdMap.Find(NameId);
-		ObjectInstance.NameIndex = NetProfilerNameIndex ? *NetProfilerNameIndex : 0u;
-		ObjectInstance.NetId = ObjectId;
-		ObjectInstance.TypeId = TypeId;
-	}
-	else
-	{
-		// Add persistent object representation
-		Trace::FNetProfilerObjectInstance& ObjectInstance = NetProfilerProvider.CreateObject(GameInstanceState->GameInstanceIndex);
-
-		// Fill in object data
-		ObjectInstance.LifeTime.Begin =  GetLastTimestamp();
-		const uint32* NetProfilerNameIndex = TracedNameIdToNetProfilerNameIdMap.Find(NameId);
-		ObjectInstance.NameIndex = NetProfilerNameIndex ? *NetProfilerNameIndex : 0u;
-		ObjectInstance.NetId = ObjectId;
-		ObjectInstance.TypeId = TypeId;
-
-		// Add to active objects
-		GameInstanceState->ActiveObjects.Add(ObjectId, { ObjectInstance.ObjectIndex, ObjectInstance.NameIndex });
+		if (Trace::FNetProfilerObjectInstance* ExistingInstance = NetProfilerProvider.EditObject(GameInstanceState->GameInstanceIndex, GameInstanceState->ActiveObjects[ObjectId].ObjectIndex))
+		{
+			if (ExistingInstance->NameIndex == NameIndex)
+			{
+				// Update existing object instance
+				ExistingInstance->LifeTime.Begin = GetLastTimestamp();
+				ExistingInstance->NetId = ObjectId;
+				ExistingInstance->TypeId = TypeId;
+
+				return;
+			}
+
+			// End instance and remove it from ActiveObjects
+			ExistingInstance->LifeTime.End = GetLastTimestamp();
+			GameInstanceState->ActiveObjects.Remove(ObjectId);
+		}
 	}
+
+	// Add persistent object representation
+	Trace::FNetProfilerObjectInstance& ObjectInstance = NetProfilerProvider.CreateObject(GameInstanceState->GameInstanceIndex);
+
+	// Fill in object data
+	ObjectInstance.LifeTime.Begin =  GetLastTimestamp();
+	ObjectInstance.NameIndex = NameIndex;
+	ObjectInstance.NetId = ObjectId;
+	ObjectInstance.TypeId = TypeId;
+
+	// Add to active objects
+	GameInstanceState->ActiveObjects.Add(ObjectId, { ObjectInstance.ObjectIndex, ObjectInstance.NameIndex });
 }
 
 void FNetTraceAnalyzer::HandleObjectDestroyedEvent(const FOnEventContext& Context, const FEventData& EventData)
diff --git a/Engine/Source/Developer/TraceServices/Private/Analyzers/NetTraceAnalyzer.h b/Engine/Source/Developer/TraceServices/Private/Analyzers/NetTraceAnalyzer.h
--- a/Engine/Source/Developer/TraceServices/Private/Analyzers/NetTraceAnalyzer.h
+++ b/Engine/Source/Developer/TraceServices/Private/Analyzers/NetTraceAnalyzer.h
@@ -92,6 +92,7 @@ private:
 	{
 		TMap<uint32, TSharedRef<FNetTraceConnectionState>> ActiveConnections;
 		TMap<uint64, FNetTraceActiveObjectState> ActiveObjects;
+		TMap<int32, uint32> ChannelNames;
 
 		uint32 GameInstanceIndex;
 	};
