NetTraceAnalyzer
- Fixed check triggering when processing PacketContentEvent after connection being reported as closed

diff --git a/Engine/Source/Developer/TraceServices/Private/Analyzers/NetTraceAnalyzer.cpp b/Engine/Source/Developer/TraceServices/Private/Analyzers/NetTraceAnalyzer.cpp
--- a/Engine/Source/Developer/TraceServices/Private/Analyzers/NetTraceAnalyzer.cpp
+++ b/Engine/Source/Developer/TraceServices/Private/Analyzers/NetTraceAnalyzer.cpp
@@ -172,10 +172,13 @@ void FNetTraceAnalyzer::HandlePacketContentEvent(const FOnEventContext& Context,
 	//UE_LOG(LogNetTrace, Display, TEXT("FNetTraceAnalyzer::HandlePacketContentEvent: GameInstanceId: %u, ConnectionId: %u, %s"), GameInstanceId, ConnectionId, PacketType ? TEXT("Incoming") : TEXT("Outgoing"));
 
 	TSharedRef<FNetTraceGameInstanceState> GameInstanceState = GetOrCreateActiveGameInstanceState(GameInstanceId);
-	TSharedRef<FNetTraceConnectionState> ConnectionState = GetActiveConnectionState(GameInstanceId, ConnectionId);
+	FNetTraceConnectionState* ConnectionState = GetActiveConnectionState(GameInstanceId, ConnectionId);
+	if (!ConnectionState)
+	{
+		return;
+	}
 
 	const Trace::ENetProfilerConnectionMode ConnectionMode = Trace::ENetProfilerConnectionMode(PacketType);
-
 	Trace::FNetProfilerConnectionData& ConnectionData = NetProfilerProvider.EditConnectionData(ConnectionState->ConnectionIndex, ConnectionMode);
 	++ConnectionData.ContentEventChangeCount;
 
@@ -422,8 +425,11 @@ void FNetTraceAnalyzer::HandlePacketEvent(const FOnEventContext& Context, const
 	// Update LastTimestamp, later on we will be able to get timestamps piggybacked from other analyzers
 	LastTimeStamp = Context.SessionContext.TimestampFromCycle(TimestampCycles);
 
-	TSharedRef<FNetTraceGameInstanceState> GameInstanceState = GetOrCreateActiveGameInstanceState(GameInstanceId);
-	TSharedRef<FNetTraceConnectionState> ConnectionState = GetActiveConnectionState(GameInstanceId, ConnectionId);
+	FNetTraceConnectionState* ConnectionState = GetActiveConnectionState(GameInstanceId, ConnectionId);
+	if (!ConnectionState)
+	{
+		return;
+	}
 
 	// Add the packet
 	Trace::FNetProfilerConnectionData& ConnectionData = NetProfilerProvider.EditConnectionData(ConnectionState->ConnectionIndex, ConnectionMode);
@@ -462,8 +468,11 @@ void FNetTraceAnalyzer::HandlePacketDroppedEvent(const FOnEventContext& Context,
 	// Update LastTimestamp, later on we will be able to get timestamps piggybacked from other analyzers
 	LastTimeStamp = Context.SessionContext.TimestampFromCycle(TimestampCycles);
 
-	TSharedRef<FNetTraceGameInstanceState> GameInstanceState = GetOrCreateActiveGameInstanceState(GameInstanceId);
-	TSharedRef<FNetTraceConnectionState> ConnectionState = GetActiveConnectionState(GameInstanceId, ConnectionId);
+	FNetTraceConnectionState* ConnectionState = GetActiveConnectionState(GameInstanceId, ConnectionId);
+	if (!ConnectionState)
+	{
+		return;
+	}
 
 	Trace::FNetProfilerConnectionData& ConnectionData = NetProfilerProvider.EditConnectionData(ConnectionState->ConnectionIndex, Trace::ENetProfilerConnectionMode(PacketType));
 
@@ -617,9 +626,15 @@ void FNetTraceAnalyzer::DestroyActiveGameInstanceState(uint32 GameInstanceId)
 	}
 }
 
-TSharedRef<FNetTraceAnalyzer::FNetTraceConnectionState> FNetTraceAnalyzer::GetActiveConnectionState(uint32 GameInstanceId, uint32 ConnectionId)
+FNetTraceAnalyzer::FNetTraceConnectionState* FNetTraceAnalyzer::GetActiveConnectionState(uint32 GameInstanceId, uint32 ConnectionId)
 {
-	TSharedRef<FNetTraceGameInstanceState> GameInstanceState = GetOrCreateActiveGameInstanceState(GameInstanceId);
-
-	return GameInstanceState->ActiveConnections.FindChecked(ConnectionId);
+	if (TSharedRef<FNetTraceAnalyzer::FNetTraceGameInstanceState>* FoundState = ActiveGameInstances.Find(GameInstanceId))
+	{
+		if (TSharedRef<FNetTraceConnectionState>* ConnectionState =  (*FoundState)->ActiveConnections.Find(ConnectionId))
+		{
+			return &(*ConnectionState).Get();
+		}
+	}
+		
+	return nullptr;
 }
diff --git a/Engine/Source/Developer/TraceServices/Private/Analyzers/NetTraceAnalyzer.h b/Engine/Source/Developer/TraceServices/Private/Analyzers/NetTraceAnalyzer.h
--- a/Engine/Source/Developer/TraceServices/Private/Analyzers/NetTraceAnalyzer.h
+++ b/Engine/Source/Developer/TraceServices/Private/Analyzers/NetTraceAnalyzer.h
@@ -104,7 +104,7 @@ private:
 	TSharedRef<FNetTraceGameInstanceState> GetOrCreateActiveGameInstanceState(uint32 GameInstanceId);
 	void DestroyActiveGameInstanceState(uint32 GameInstanceId);
 
-	TSharedRef<FNetTraceConnectionState> GetActiveConnectionState(uint32 GameInstanceId, uint32 ConnectionId);
+	FNetTraceConnectionState* GetActiveConnectionState(uint32 GameInstanceId, uint32 ConnectionId);
 	Trace::FNetProfilerTimeStamp GetLastTimestamp() const { return LastTimeStamp; }
 
 	void FlushPacketEvents(FNetTraceConnectionState& ConnectionState, Trace::FNetProfilerConnectionData& ConnectionData, const Trace::ENetProfilerConnectionMode ConnectionMode);
