NetTraceAnalyzer
- Bumped number of bits to store offsets in PacketContentEvents  from 16 to 24 bits
- Fixed issue with displaying large split bunches

diff --git a/Engine/Source/Developer/TraceServices/Public/TraceServices/Model/NetProfiler.h b/Engine/Source/Developer/TraceServices/Public/TraceServices/Model/NetProfiler.h
--- a/Engine/Source/Developer/TraceServices/Public/TraceServices/Model/NetProfiler.h
+++ b/Engine/Source/Developer/TraceServices/Public/TraceServices/Model/NetProfiler.h
@@ -66,14 +66,16 @@ struct FNetProfilerObjectInstance
 
 struct FNetProfilerContentEvent
 {
+	uint64 StartPos : 24;		// Start position in the packet
+	uint64 EndPos : 24;			// End position in the packet
+	uint64 Level : 4;			// Level
+	uint64 Padding : 12;		// Padding
+
 	uint32 EventTypeIndex;		// Will replace name index
-	uint32 NameIndex ;			// identify the name / type, should we store the actual Name as well
+	uint32 NameIndex;			// identify the name / type, should we store the actual Name as well
 	uint32 ObjectInstanceIndex;	// object instance, Non zero if this is a NetObject, we can then look up data by indexing into ObjectInstances
 
-	uint64 StartPos : 16;		// Start position in the packet
-	uint64 EndPos : 16;			// End position in the packet
-	uint64 Level : 4;			// Level
-	uint64 ParentIndex : 28;	// Parent to be able to build a tree of nested events?
+	uint32 ParentIndex;			// Parent to be able to build a tree of nested events?
 };
 
 struct FNetProfilerPacket
