NetTraceAnalyzer
FNetProfilerProvider::EnumeratePacketContentEventsByPosition now includes sub events of largest included event when doing range selection.

diff --git a/Engine/Source/Developer/TraceServices/Private/Model/NetProfilerProvider.cpp b/Engine/Source/Developer/TraceServices/Private/Model/NetProfilerProvider.cpp
--- a/Engine/Source/Developer/TraceServices/Private/Model/NetProfilerProvider.cpp
+++ b/Engine/Source/Developer/TraceServices/Private/Model/NetProfilerProvider.cpp
@@ -443,6 +443,7 @@ void FNetProfilerProvider::EnumeratePacketContentEventsByPosition(uint32 Connect
 	while (EventIt <= EndEventIndex && ContentEvents[EventIt].StartPos < EndPos)
 	{
 		Callback(ContentEvents[EventIt]);
+		EndPos = FMath::Max(EndPos, (uint32)ContentEvents[EventIt].EndPos);
 		++EventIt;
 	}
 }
