Closer use of TraceAnalysis' Processor object instead of tasks so it's possible to stop and/or wait on analysis via the IAnalysisSession interface.

#rnx

diff --git a/Engine/Source/Developer/TraceInsights/Private/Insights/InsightsManager.cpp b/Engine/Source/Developer/TraceInsights/Private/Insights/InsightsManager.cpp
--- a/Engine/Source/Developer/TraceInsights/Private/Insights/InsightsManager.cpp
+++ b/Engine/Source/Developer/TraceInsights/Private/Insights/InsightsManager.cpp
@@ -77,6 +77,8 @@ void FInsightsManager::BindCommands()
 
 FInsightsManager::~FInsightsManager()
 {
+	ResetSession();
+
 	FInsightsCommands::Unregister();
 
 	// Unregister tick function.
@@ -173,6 +175,7 @@ void FInsightsManager::ResetSession()
 {
 	if (Session.IsValid())
 	{
+		Session->Stop(true);
 		Session.Reset();
 		CurrentTraceId = 0;
 		bIsNetworkingProfilerAvailable = false;
diff --git a/Engine/Source/Developer/TraceServices/Private/AnalysisService.cpp b/Engine/Source/Developer/TraceServices/Private/AnalysisService.cpp
--- a/Engine/Source/Developer/TraceServices/Private/AnalysisService.cpp
+++ b/Engine/Source/Developer/TraceServices/Private/AnalysisService.cpp
@@ -77,13 +77,13 @@ void FAnalysisSessionLock::EndEdit()
 	}
 }
 
-FAnalysisSession::FAnalysisSession(const TCHAR* SessionName)
+FAnalysisSession::FAnalysisSession(const TCHAR* SessionName, TUniquePtr<Trace::IInDataStream>&& InDataStream)
 	: Name(SessionName)
 	, DurationSeconds(0.0)
 	, Allocator(32 << 20)
 	, StringStore(Allocator)
+	, DataStream(MoveTemp(InDataStream))
 {
-
 }
 
 FAnalysisSession::~FAnalysisSession()
@@ -98,6 +98,30 @@ FAnalysisSession::~FAnalysisSession()
 	}
 }
 
+void FAnalysisSession::Start()
+{
+	FAnalysisContext Context;
+	for (Trace::IAnalyzer* Analyzer : ReadAnalyzers())
+	{
+		Context.AddAnalyzer(*Analyzer);
+	}
+	Processor = Context.Process(*DataStream);
+}
+
+void FAnalysisSession::Stop(bool bAndWait) const
+{
+	DataStream->Close();
+	if (bAndWait)
+	{
+		Wait();
+	}
+}
+
+void FAnalysisSession::Wait() const
+{
+	Processor.Wait();
+}
+
 void FAnalysisSession::AddAnalyzer(IAnalyzer* Analyzer)
 {
 	Analyzers.Add(Analyzer);
@@ -138,114 +162,69 @@ IProvider* FAnalysisSession::EditProviderPrivate(const FName& InName)
 FAnalysisService::FAnalysisService(FModuleService& InModuleService)
 	: ModuleService(InModuleService)
 {
-
-}
-
-FAnalysisService::FAnalysisWorker::FAnalysisWorker(FAnalysisService& InOuter, TUniquePtr<Trace::IInDataStream>&& InDataStream, TSharedRef<FAnalysisSession> InAnalysisSession)
-	: Outer(InOuter)
-	, DataStream(MoveTemp(InDataStream))
-	, AnalysisSession(InAnalysisSession)
-{
-	
-}
-
-void FAnalysisService::AnalyzeInternal(TSharedRef<FAnalysisSession> AnalysisSession, Trace::IInDataStream* DataStream)
-{
-	OnAnalysisStarted().Broadcast(AnalysisSession);
-
-	{
-		IAnalysisSession& Session = AnalysisSession.Get();
-		Trace::FAnalysisSessionEditScope _(Session);
-
-		FBookmarkProvider* BookmarkProvider = new FBookmarkProvider(Session);
-		Session.AddProvider(FBookmarkProvider::ProviderName, BookmarkProvider);
-
-		FLogProvider* LogProvider = new FLogProvider(Session);
-		Session.AddProvider(FLogProvider::ProviderName, LogProvider);
-
-		FThreadProvider* ThreadProvider = new FThreadProvider(Session);
-		Session.AddProvider(FThreadProvider::ProviderName, ThreadProvider);
-
-		FFrameProvider* FrameProvider = new FFrameProvider(Session);
-		Session.AddProvider(FFrameProvider::ProviderName, FrameProvider);
-
-		FCounterProvider* CounterProvider = new FCounterProvider(Session, *FrameProvider);
-		Session.AddProvider(FCounterProvider::ProviderName, CounterProvider);
-
-		FNetProfilerProvider* NetProfilerProvider = new FNetProfilerProvider(Session);
-		Session.AddProvider(FNetProfilerProvider::ProviderName, NetProfilerProvider);
-
-		FChannelProvider* ChannelProvider = new FChannelProvider();
-		Session.AddProvider(FChannelProvider::ProviderName, ChannelProvider);
-
-		Session.AddAnalyzer(new FMiscTraceAnalyzer(Session, *ThreadProvider, *BookmarkProvider, *LogProvider, *FrameProvider, *ChannelProvider));
-		Session.AddAnalyzer(new FLogTraceAnalyzer(Session, *LogProvider));
-
-		ModuleService.OnAnalysisBegin(Session);
-	}
-	
-	FAnalysisContext Context;
-	for (Trace::IAnalyzer* Analyzer : AnalysisSession->ReadAnalyzers())
-	{
-		Context.AddAnalyzer(*Analyzer);
-	}
-	Trace::FAnalysisProcessor Processor = Context.Process(*DataStream);
-	Processor.Wait();
-
-	delete DataStream;
-
-	AnalysisSession->SetComplete();
-
-	OnAnalysisFinished().Broadcast(AnalysisSession);
-}
-
-void FAnalysisService::FAnalysisWorker::DoWork()
-{
-	Outer.AnalyzeInternal(AnalysisSession.ToSharedRef(), DataStream.Release());
-	AnalysisSession = nullptr;
 }
 
 FAnalysisService::~FAnalysisService()
 {
-	for (TSharedPtr<FAsyncTask<FAnalysisWorker>> Task : Tasks)
-	{
-		Task->EnsureCompletion();
-	}
 }
 
 TSharedPtr<const IAnalysisSession> FAnalysisService::Analyze(const TCHAR* SessionUri)
 {
-	IInDataStream* DataStream = Trace::DataStream_ReadFile(SessionUri);
+	TUniquePtr<Trace::IInDataStream> DataStream(Trace::DataStream_ReadFile(SessionUri));
 	if (!DataStream)
 	{
 		return nullptr;
 	}
-	TSharedRef<FAnalysisSession> AnalysisSession = MakeShared<FAnalysisSession>(SessionUri);
-	AnalyzeInternal(AnalysisSession, DataStream);
+	
+	TSharedPtr<const IAnalysisSession> AnalysisSession = StartAnalysis(SessionUri, MoveTemp(DataStream));
+	AnalysisSession->Wait();
 	return AnalysisSession;
 }
 
 TSharedPtr<const IAnalysisSession> FAnalysisService::StartAnalysis(const TCHAR* SessionUri)
 {
-	TUniquePtr<IInDataStream> DataStream(Trace::DataStream_ReadFile(SessionUri));
+	TUniquePtr<Trace::IInDataStream> DataStream(Trace::DataStream_ReadFile(SessionUri));
 	if (!DataStream)
 	{
 		return nullptr;
 	}
-	TSharedRef<FAnalysisSession> AnalysisSession = MakeShared<FAnalysisSession>(SessionUri);
-	TSharedPtr<FAsyncTask<FAnalysisWorker>> Task = MakeShared<FAsyncTask<FAnalysisWorker>>(*this, MoveTemp(DataStream), AnalysisSession);
-	Tasks.Add(Task);
-	Task->StartBackgroundTask();
-	return AnalysisSession;
+	return StartAnalysis(SessionUri, MoveTemp(DataStream));
 }
 
 TSharedPtr<const IAnalysisSession> FAnalysisService::StartAnalysis(const TCHAR* SessionName, TUniquePtr<Trace::IInDataStream>&& DataStream)
 {
-	TSharedRef<FAnalysisSession> AnalysisSession = MakeShared<FAnalysisSession>(SessionName);
-	TSharedPtr<FAsyncTask<FAnalysisWorker>> Task = MakeShared<FAsyncTask<FAnalysisWorker>>(*this, MoveTemp(DataStream), AnalysisSession);
-	Tasks.Add(Task);
-	Task->StartBackgroundTask();
-	return AnalysisSession;
+	TSharedRef<FAnalysisSession> Session = MakeShared<FAnalysisSession>(SessionName, MoveTemp(DataStream));
+
+	Trace::FAnalysisSessionEditScope _(*Session);
+
+	FBookmarkProvider* BookmarkProvider = new FBookmarkProvider(*Session);
+	Session->AddProvider(FBookmarkProvider::ProviderName, BookmarkProvider);
+
+	FLogProvider* LogProvider = new FLogProvider(*Session);
+	Session->AddProvider(FLogProvider::ProviderName, LogProvider);
+
+	FThreadProvider* ThreadProvider = new FThreadProvider(*Session);
+	Session->AddProvider(FThreadProvider::ProviderName, ThreadProvider);
+
+	FFrameProvider* FrameProvider = new FFrameProvider(*Session);
+	Session->AddProvider(FFrameProvider::ProviderName, FrameProvider);
+
+	FCounterProvider* CounterProvider = new FCounterProvider(*Session, *FrameProvider);
+	Session->AddProvider(FCounterProvider::ProviderName, CounterProvider);
+
+	FNetProfilerProvider* NetProfilerProvider = new FNetProfilerProvider(*Session);
+	Session->AddProvider(FNetProfilerProvider::ProviderName, NetProfilerProvider);
+
+	FChannelProvider* ChannelProvider = new FChannelProvider();
+	Session->AddProvider(FChannelProvider::ProviderName, ChannelProvider);
+
+	Session->AddAnalyzer(new FMiscTraceAnalyzer(*Session, *ThreadProvider, *BookmarkProvider, *LogProvider, *FrameProvider, *ChannelProvider));
+	Session->AddAnalyzer(new FLogTraceAnalyzer(*Session, *LogProvider));
+
+	ModuleService.OnAnalysisBegin(*Session);
+
+	Session->Start();
+	return Session;
 }
 
 }
diff --git a/Engine/Source/Developer/TraceServices/Private/AnalysisServicePrivate.h b/Engine/Source/Developer/TraceServices/Private/AnalysisServicePrivate.h
--- a/Engine/Source/Developer/TraceServices/Private/AnalysisServicePrivate.h
+++ b/Engine/Source/Developer/TraceServices/Private/AnalysisServicePrivate.h
@@ -7,11 +7,13 @@
 #include "Containers/HashTable.h"
 #include "Misc/ScopeLock.h"
 #include "Common/StringStore.h"
+#include "Trace/Analysis.h"
 
 namespace Trace
 {
 
 class FModuleService;
+class IInDataStream;
 
 class FAnalysisSessionLock
 {
@@ -33,14 +35,16 @@ class FAnalysisSession
 	: public IAnalysisSession
 {
 public:
-	FAnalysisSession(const TCHAR* SessionName);
+	FAnalysisSession(const TCHAR* SessionName, TUniquePtr<Trace::IInDataStream>&& InDataStream);
 	virtual ~FAnalysisSession();
+	void Start();
+	virtual void Stop(bool bAndWait) const override;
+	virtual void Wait() const override;
 
 	virtual const TCHAR* GetName() const { return *Name; }
-	virtual bool IsAnalysisComplete() const override { return IsComplete; }
+	virtual bool IsAnalysisComplete() const override { return !Processor.IsActive(); }
 	virtual double GetDurationSeconds() const { Lock.ReadAccessCheck(); return DurationSeconds; }
 	virtual void UpdateDurationSeconds(double Duration) override { Lock.WriteAccessCheck(); DurationSeconds = FMath::Max(Duration, DurationSeconds); }
-	void SetComplete() { IsComplete = true; }
 
 	virtual ILinearAllocator& GetLinearAllocator() override { return Allocator; }
 	virtual const TCHAR* StoreString(const TCHAR* String) override { return StringStore.Store(String); }
@@ -66,13 +70,14 @@ private:
 	mutable FAnalysisSessionLock Lock;
 
 	FString Name;
-	bool IsComplete = false;
 	double DurationSeconds = 0.0;
 	FSlabAllocator Allocator;
 	FStringStore StringStore;
 	TArray<IAnalyzer*> Analyzers;
 	TArray<IProvider*> Providers;
 	TMap<FName, IProvider*> ProvidersMap;
+	mutable TUniquePtr<Trace::IInDataStream> DataStream;
+	mutable Trace::FAnalysisProcessor Processor;
 };
 
 class FAnalysisService
@@ -84,32 +89,9 @@ public:
 	virtual TSharedPtr<const IAnalysisSession> Analyze(const TCHAR* SessionUri) override;
 	virtual TSharedPtr<const IAnalysisSession> StartAnalysis(const TCHAR* SessionUri) override;
 	virtual TSharedPtr<const IAnalysisSession> StartAnalysis(const TCHAR* SessionName, TUniquePtr<Trace::IInDataStream>&& DataStream) override;
-	virtual FAnalysisStartedEvent& OnAnalysisStarted() override { return AnalysisStartedEvent; }
-	virtual FAnalysisFinishedEvent& OnAnalysisFinished() override { return AnalysisFinishedEvent; }
 
 private:
-	class FAnalysisWorker : public FNonAbandonableTask
-	{
-	public:
-		FAnalysisWorker(FAnalysisService& Outer, TUniquePtr<Trace::IInDataStream>&& InDataStream, TSharedRef<FAnalysisSession> InAnalysisSession);
-		void DoWork();
-		FORCEINLINE TStatId GetStatId() const
-		{
-			RETURN_QUICK_DECLARE_CYCLE_STAT(FAnalysisWorker, STATGROUP_ThreadPoolAsyncTasks);
-		}
-
-	private:
-		FAnalysisService& Outer;
-		TUniquePtr<Trace::IInDataStream> DataStream;
-		TSharedPtr<FAnalysisSession> AnalysisSession;
-	};
-
-	void AnalyzeInternal(TSharedRef<FAnalysisSession> AnalysisSession, Trace::IInDataStream* DataStream);
-
 	FModuleService& ModuleService;
-	FAnalysisStartedEvent AnalysisStartedEvent;
-	FAnalysisFinishedEvent AnalysisFinishedEvent;
-	TArray<TSharedPtr<FAsyncTask<FAnalysisWorker>>> Tasks;
 };
 
 }
diff --git a/Engine/Source/Developer/TraceServices/Public/TraceServices/AnalysisService.h b/Engine/Source/Developer/TraceServices/Public/TraceServices/AnalysisService.h
--- a/Engine/Source/Developer/TraceServices/Public/TraceServices/AnalysisService.h
+++ b/Engine/Source/Developer/TraceServices/Public/TraceServices/AnalysisService.h
@@ -25,12 +25,6 @@ public:
 	virtual TSharedPtr<const IAnalysisSession> Analyze(const TCHAR* SessionUri) = 0;
 	virtual TSharedPtr<const IAnalysisSession> StartAnalysis(const TCHAR* SessionUri) = 0;
 	virtual TSharedPtr<const IAnalysisSession> StartAnalysis(const TCHAR* SessionName, TUniquePtr<Trace::IInDataStream>&& DataStream) = 0;
-
-	DECLARE_EVENT_OneParam(IAnalysisService, FAnalysisStartedEvent, TSharedRef<const IAnalysisSession>)
-	virtual FAnalysisStartedEvent& OnAnalysisStarted() = 0;
-
-	DECLARE_EVENT_OneParam(IAnalysisService, FAnalysisFinishedEvent, TSharedRef<const IAnalysisSession>)
-	virtual FAnalysisFinishedEvent& OnAnalysisFinished() = 0;
 };
 
 }
diff --git a/Engine/Source/Developer/TraceServices/Public/TraceServices/Model/AnalysisSession.h b/Engine/Source/Developer/TraceServices/Public/TraceServices/Model/AnalysisSession.h
--- a/Engine/Source/Developer/TraceServices/Public/TraceServices/Model/AnalysisSession.h
+++ b/Engine/Source/Developer/TraceServices/Public/TraceServices/Model/AnalysisSession.h
@@ -22,6 +22,8 @@ class IAnalysisSession
 {
 public:
 	virtual ~IAnalysisSession() = default;
+	virtual void Stop(bool bAndWait=false) const = 0;
+	virtual void Wait() const = 0;
 	
 	virtual const TCHAR* GetName() const = 0;
 	virtual bool IsAnalysisComplete() const = 0;
