Fixed previous frame button on Gameplay insights transport controls

Fixes off-by-one error in frames provider
This caused frames to be mis-reported when using FFrameProvider::GetFrameFromTime.

diff --git a/Engine/Source/Developer/TraceServices/Private/Model/Frames.cpp b/Engine/Source/Developer/TraceServices/Private/Model/Frames.cpp
--- a/Engine/Source/Developer/TraceServices/Private/Model/Frames.cpp
+++ b/Engine/Source/Developer/TraceServices/Private/Model/Frames.cpp
@@ -68,12 +68,9 @@ static int64 LowerBound64(RangeValueType* First, const int64 Num, const Predicat
 bool FFrameProvider::GetFrameFromTime(ETraceFrameType FrameType, double Time, FFrame& OutFrame) const
 {
 	int64 LowerBound = LowerBound64(FrameStartTimes[FrameType].GetData(), FrameStartTimes[FrameType].Num(), Time, FIdentityFunctor(), TLess<double>());
-	if(FrameStartTimes[FrameType].IsValidIndex(LowerBound) && LowerBound > 0)
+	if(LowerBound > 0 && LowerBound - 1 < (int64)Frames[FrameType].Num())
 	{
-		OutFrame.Index = LowerBound;
-		OutFrame.StartTime = FrameStartTimes[FrameType][LowerBound - 1];
-		OutFrame.EndTime = FrameStartTimes[FrameType][LowerBound];
-		OutFrame.FrameType = FrameType;
+		OutFrame = Frames[FrameType][LowerBound - 1];
 		return true;
 	}
 
