[Insights]
  - NetProfiler: Fixed computation of Exclusive size aggregated stats (UE-91087).
  - Net / Packets view: Fixed Ctrl+A not triggering the update of aggregated stats if no packet was previously selected (UE-91067).

diff --git a/Engine/Source/Developer/TraceInsights/Private/Insights/NetworkingProfiler/Widgets/SNetStatsView.cpp b/Engine/Source/Developer/TraceInsights/Private/Insights/NetworkingProfiler/Widgets/SNetStatsView.cpp
--- a/Engine/Source/Developer/TraceInsights/Private/Insights/NetworkingProfiler/Widgets/SNetStatsView.cpp
+++ b/Engine/Source/Developer/TraceInsights/Private/Insights/NetworkingProfiler/Widgets/SNetStatsView.cpp
@@ -1667,8 +1667,8 @@ void SNetStatsView::UpdateStats(uint32 InGameInstanceIndex, uint32 InConnectionI
 
 void SNetStatsView::UpdateStatsInternal()
 {
-	if (StatsPacketStartIndex >= StatsPacketEndIndex ||
-		StatsStartPosition >= StatsEndPosition)
+	if (StatsPacketStartIndex >= StatsPacketEndIndex || // no packet selected?
+		(StatsPacketEndIndex - StatsPacketStartIndex == 1 && StatsStartPosition >= StatsEndPosition)) // single packet, but invalid bit range?
 	{
 		// keep previous aggregated stats
 		return;
@@ -1686,6 +1686,7 @@ void SNetStatsView::UpdateStatsInternal()
 		{
 			Trace::FAnalysisSessionReadScope SessionReadScope(*Session.Get());
 			const Trace::INetProfilerProvider& NetProfilerProvider = Trace::ReadNetProfilerProvider(*Session.Get());
+			// CreateAggregation requires [PacketStartIndex, PacketEndIndex] as inclusive interval and [StartPos, EndPos) as exclusive interval.
 			AggregationResultTable.Reset(NetProfilerProvider.CreateAggregation(ConnectionIndex, ConnectionMode, StatsPacketStartIndex, StatsPacketEndIndex - 1, StatsStartPosition, StatsEndPosition));
 		}
 
diff --git a/Engine/Source/Developer/TraceInsights/Private/Insights/NetworkingProfiler/Widgets/SNetworkingProfilerWindow.cpp b/Engine/Source/Developer/TraceInsights/Private/Insights/NetworkingProfiler/Widgets/SNetworkingProfilerWindow.cpp
--- a/Engine/Source/Developer/TraceInsights/Private/Insights/NetworkingProfiler/Widgets/SNetworkingProfilerWindow.cpp
+++ b/Engine/Source/Developer/TraceInsights/Private/Insights/NetworkingProfiler/Widgets/SNetworkingProfilerWindow.cpp
@@ -979,8 +979,8 @@ void SNetworkingProfilerWindow::UpdateAggregatedNetStats()
 		if (SelectedGameInstance.IsValid() &&
 			SelectedConnection.IsValid() &&
 			SelectedConnectionMode.IsValid() &&
-			SelectedPacketStartIndex < SelectedPacketEndIndex &&
-			SelectionStartPosition < SelectionEndPosition)
+			(SelectedPacketStartIndex < SelectedPacketEndIndex) &&
+			(SelectedPacketStartIndex + 1 != SelectedPacketEndIndex || SelectionStartPosition < SelectionEndPosition))
 		{
 			NetStatsView->UpdateStats(SelectedGameInstance->GetIndex(),
 									  SelectedConnection->GetIndex(),
diff --git a/Engine/Source/Developer/TraceInsights/Private/Insights/NetworkingProfiler/Widgets/SNetworkingProfilerWindow.h b/Engine/Source/Developer/TraceInsights/Private/Insights/NetworkingProfiler/Widgets/SNetworkingProfilerWindow.h
--- a/Engine/Source/Developer/TraceInsights/Private/Insights/NetworkingProfiler/Widgets/SNetworkingProfilerWindow.h
+++ b/Engine/Source/Developer/TraceInsights/Private/Insights/NetworkingProfiler/Widgets/SNetworkingProfilerWindow.h
@@ -298,9 +298,14 @@ private:
 	TArray<TSharedPtr<FConnectionModeItem>> AvailableConnectionModes;
 	TSharedPtr<FConnectionModeItem> SelectedConnectionMode;
 
+	// [SelectedPacketStartIndex, SelectedPacketEndIndex) is the exclusive interval of selected packages.
+	// NumSelectedPackets == SelectedPacketEndIndex - SelectedPacketStartIndex.
 	uint32 SelectedPacketStartIndex;
 	uint32 SelectedPacketEndIndex;
 
+	// [SelectionStartPosition, SelectionEndPosition) is the exclusive selected bit range inside a single selected package.
+	// Used only when NumSelectedPackets == SelectedPacketEndIndex - SelectedPacketStartIndex == 1.
+	// SelectedBitSize == SelectionEndPosition - SelectionStartPosition.
 	uint32 SelectionStartPosition;
 	uint32 SelectionEndPosition;
 
diff --git a/Engine/Source/Developer/TraceServices/Private/Model/NetProfilerProvider.cpp b/Engine/Source/Developer/TraceServices/Private/Model/NetProfilerProvider.cpp
--- a/Engine/Source/Developer/TraceServices/Private/Model/NetProfilerProvider.cpp
+++ b/Engine/Source/Developer/TraceServices/Private/Model/NetProfilerProvider.cpp
@@ -392,7 +392,7 @@ int32 FNetProfilerProvider::FindPacketIndexFromPacketSequence(uint32 ConnectionI
 		}
 	}
 
-	return -1;	
+	return -1;
 }
 
 uint32 FNetProfilerProvider::GetPacketCount(uint32 ConnectionIndex, ENetProfilerConnectionMode Mode) const
@@ -416,6 +416,7 @@ void FNetProfilerProvider::EnumeratePackets(uint32 ConnectionIndex, ENetProfiler
 
 	const uint32 PacketCount = Packets.Num();
 
+	// [PacketIndexIntervalStart, PacketIndexIntervalEnd] is an inclusive interval.
 	if (PacketCount == 0 || PacketIndexIntervalStart > PacketIndexIntervalEnd)
 	{
 		return;
@@ -438,6 +439,7 @@ void FNetProfilerProvider::EnumeratePacketContentEventsByIndex(uint32 Connection
 
 	const uint32 EventCount = ContentEvents.Num();
 
+	// [StartEventIndex, EndEventIndex] is an inclusive interval.
 	if (EventCount == 0 || StartEventIndex > EndEventIndex)
 	{
 		return;
@@ -468,9 +470,12 @@ void FNetProfilerProvider::EnumeratePacketContentEventsByPosition(uint32 Connect
 
 	const auto& ContentEvents = ConnectionData.ContentEvents;
 
+	// The input [StartPos, EndPos) is an exclusive bit range.
+	// Also, the [StartPos, EndPos) for ContentEvents[EventIt] is an exclusive bit range.
+
 	uint32 EventIt = StartEventIndex;
 	// Skip all Events outside of the scope
-	while (EventIt <= EndEventIndex &&	ContentEvents[EventIt].EndPos < StartPos)
+	while (EventIt <= EndEventIndex && ContentEvents[EventIt].EndPos <= StartPos)
 	{
 		++EventIt;
 	}
@@ -509,6 +514,7 @@ ITable<FNetProfilerAggregatedStats>* FNetProfilerProvider::CreateAggregation(uin
 		return nullptr;
 	}
 
+	// [PacketIndexIntervalStart, PacketIndexIntervalEnd] is an inclusive interval.
 	if (!ensure(PacketIndexIntervalStart <= PacketIndexIntervalEnd))
 	{
 		return nullptr;
@@ -522,10 +528,6 @@ ITable<FNetProfilerAggregatedStats>* FNetProfilerProvider::CreateAggregation(uin
 		return nullptr;
 	}
 
-	// Aggregate stats for events in selection
-	TMap<uint32, FNetProfilerAggregatedStats> AggregatedStats;
-	AggregatedStats.Reserve(EventTypes.Num());
-
 	struct FStackEntry
 	{
 		uint32 EventTypeIndex = 0U;
@@ -534,53 +536,83 @@ ITable<FNetProfilerAggregatedStats>* FNetProfilerProvider::CreateAggregation(uin
 		uint32 ExclusiveAccumulator = 0U;
 	};
 
-	TArray<FStackEntry> Stack;
-	Stack.SetNum(256);
+	struct FStatsHelper
+	{
+		TMap<uint32, FNetProfilerAggregatedStats> AggregatedStats;
+		uint64 StackSize;
+		TArray<FStackEntry> Stack;
+	};
 
-	auto GetStatsFunction = [&AggregatedStats, &Stack, this](const FNetProfilerContentEvent& ContentEvent)
+	FStatsHelper Helper;
+	Helper.AggregatedStats.Reserve(EventTypes.Num());
+	Helper.StackSize = 0; // stack is empty
+	Helper.Stack.SetNum(256);
+
+	auto GetStatsFunction = [&Helper, this](const FNetProfilerContentEvent& ContentEvent)
 	{
-		FNetProfilerAggregatedStats* StatsEntry = AggregatedStats.Find(ContentEvent.EventTypeIndex);
+		FNetProfilerAggregatedStats* StatsEntry = Helper.AggregatedStats.Find(ContentEvent.EventTypeIndex);
 		if (!StatsEntry)
 		{
-			StatsEntry = &AggregatedStats.Add(ContentEvent.EventTypeIndex);
+			StatsEntry = &Helper.AggregatedStats.Add(ContentEvent.EventTypeIndex);
 			StatsEntry->EventTypeIndex = ContentEvent.EventTypeIndex;
 		}
 
-		// Fill in basics
-		uint32 InclusiveSize = ContentEvent.EndPos - ContentEvent.StartPos;
-
+		// Fill in basics.
+		const uint32 InclusiveSize = ContentEvent.EndPos - ContentEvent.StartPos;
 		++StatsEntry->InstanceCount;
 		StatsEntry->TotalInclusive += InclusiveSize;
 		StatsEntry->MaxInclusive = FMath::Max(InclusiveSize, StatsEntry->MaxInclusive);
 
-		// We track what we have visited to be able to update exclusive bits for our parent
-		// We accumulate during the first pass and finalize the values during the second
-		Stack[ContentEvent.Level].EventTypeIndex = ContentEvent.EventTypeIndex;
-		Stack[ContentEvent.Level].StartPos = ContentEvent.StartPos;
-		Stack[ContentEvent.Level].EndPos = ContentEvent.EndPos;
-		Stack[ContentEvent.Level].ExclusiveAccumulator = 0U;
+		// Pops events from the stack. Keeps only the parent hierarchy of the current event.
+		while (Helper.StackSize > ContentEvent.Level)
+		{
+			FStackEntry& StackEntry = Helper.Stack[--Helper.StackSize]; // pop
+
+			// Finalize exclusive for each poped event (all its children were already processed).
+			FNetProfilerAggregatedStats& Stats = Helper.AggregatedStats.FindChecked(StackEntry.EventTypeIndex);
+			const uint32 ExclusiveSize = (StackEntry.EndPos - StackEntry.StartPos) - StackEntry.ExclusiveAccumulator;
+			Stats.TotalExclusive += ExclusiveSize;
+			Stats.MaxExclusive = FMath::Max(Stats.MaxExclusive, ExclusiveSize);
+		}
+
+		// Pushes the new event on the stack.
+		ensure(ContentEvent.Level == Helper.StackSize);
+		Helper.StackSize = ContentEvent.Level + 1; // push
+		FStackEntry& CurrentEventLevelStackEntry = Helper.Stack[ContentEvent.Level];
+
+		// We track what we have visited to be able to update exclusive bits for our parent.
+		// We accumulate during the first pass and finalize the values during the second.
+		CurrentEventLevelStackEntry.EventTypeIndex = ContentEvent.EventTypeIndex;
+		CurrentEventLevelStackEntry.StartPos = ContentEvent.StartPos;
+		CurrentEventLevelStackEntry.EndPos = ContentEvent.EndPos;
+		CurrentEventLevelStackEntry.ExclusiveAccumulator = 0U;
 
 		if (ContentEvent.Level > 0U)
 		{
-			FStackEntry& ParentStackEntry = Stack[ContentEvent.Level - 1U];
-
+			// Update parent event with the contribution from current event.
+			FStackEntry& ParentStackEntry = Helper.Stack[ContentEvent.Level - 1U];
 			ParentStackEntry.ExclusiveAccumulator += InclusiveSize;
-
-			if (ContentEvent.EndPos == ParentStackEntry.EndPos)
-			{
-				// Finalize exclusive
-				FNetProfilerAggregatedStats& ParentStateEntry = AggregatedStats.FindChecked(ContentEvent.EventTypeIndex);
-				uint32 ExclusiveTime = (ParentStackEntry.EndPos - ParentStackEntry.StartPos) - ParentStackEntry.ExclusiveAccumulator;
-				ParentStateEntry.TotalExclusive += ExclusiveTime;
-				ParentStateEntry.MaxExclusive = FMath::Max(ParentStateEntry.MaxExclusive, ExclusiveTime);
-			}
 		}
 	};
 
 	// Iterate over content events
 	if (PacketIndexIntervalStart == PacketIndexIntervalEnd)
 	{
+		ensure(Helper.StackSize == 0); // stack should be empty
+
 		EnumeratePacketContentEventsByPosition(ConnectionIndex, Mode, PacketIndexIntervalStart, StartPosition, EndPosition, GetStatsFunction);
+
+		// Pops the remaining events from the stack.
+		while (Helper.StackSize > 0)
+		{
+			FStackEntry& StackEntry = Helper.Stack[--Helper.StackSize]; // pop
+
+			// Finalize exclusive for each poped event (all its children were already processed).
+			FNetProfilerAggregatedStats& Stats = Helper.AggregatedStats.FindChecked(StackEntry.EventTypeIndex);
+			const uint32 ExclusiveSize = (StackEntry.EndPos - StackEntry.StartPos) - StackEntry.ExclusiveAccumulator;
+			Stats.TotalExclusive += ExclusiveSize;
+			Stats.MaxExclusive = FMath::Max(Stats.MaxExclusive, ExclusiveSize);
+		}
 	}
 	else
 	{
@@ -593,18 +625,32 @@ ITable<FNetProfilerAggregatedStats>* FNetProfilerProvider::CreateAggregation(uin
 			const FNetProfilerPacket& Packet = Packets[PacketIt];
 			if (Packet.EventCount > 0U)
 			{
+				ensure(Helper.StackSize == 0); // stack should be empty (before each packet)
+
 				const uint32 StartEventIndex = Packet.StartEventIndex;
 				for (uint32 It = 0U; It < Packet.EventCount; ++It)
 				{
 					GetStatsFunction(ContentEvents[StartEventIndex + It]);
 				}
+
+				// Pops the remaining events from the stack, for each packet.
+				while (Helper.StackSize > 0)
+				{
+					FStackEntry& StackEntry = Helper.Stack[--Helper.StackSize]; // pop
+
+					// Finalize exclusive for each poped event (all its children were already processed).
+					FNetProfilerAggregatedStats& Stats = Helper.AggregatedStats.FindChecked(StackEntry.EventTypeIndex);
+					const uint32 ExclusiveSize = (StackEntry.EndPos - StackEntry.StartPos) - StackEntry.ExclusiveAccumulator;
+					Stats.TotalExclusive += ExclusiveSize;
+					Stats.MaxExclusive = FMath::Max(Stats.MaxExclusive, ExclusiveSize);
+				}
 			}
 		}
 	}
 
 	// Calculate averages and populate table
 	TTable<FNetProfilerAggregatedStats>* Table = new TTable<FNetProfilerAggregatedStats>(AggregatedStatsTableLayout);
-	for (const auto& KV : AggregatedStats)
+	for (const auto& KV : Helper.AggregatedStats)
 	{
 		FNetProfilerAggregatedStats& Row = Table->AddRow();
 		Row = KV.Value;
diff --git a/Engine/Source/Developer/TraceServices/Public/TraceServices/Model/NetProfiler.h b/Engine/Source/Developer/TraceServices/Public/TraceServices/Model/NetProfiler.h
--- a/Engine/Source/Developer/TraceServices/Public/TraceServices/Model/NetProfiler.h
+++ b/Engine/Source/Developer/TraceServices/Public/TraceServices/Model/NetProfiler.h
@@ -66,8 +66,8 @@ struct FNetProfilerObjectInstance
 
 struct FNetProfilerContentEvent
 {
-	uint64 StartPos : 24;		// Start position in the packet
-	uint64 EndPos : 24;			// End position in the packet
+	uint64 StartPos : 24;		// Inclusive start position in the packet
+	uint64 EndPos : 24;			// Exclusive end position in the packet; BitSize = EndPos - StartPos
 	uint64 Level : 4;			// Level
 	uint64 Padding : 12;		// Padding
 
@@ -158,17 +158,24 @@ public:
 	// Find Packet Index from SequenceNumber
 	virtual int32 FindPacketIndexFromPacketSequence(uint32 ConnectionIndex, ENetProfilerConnectionMode Mode, uint32 SequenceNumber) const = 0;
 
-	// Enumerate packets in the provided packet interval
+	// Gets the number of packets for the specified connection and connection mode.
 	virtual uint32 GetPacketCount(uint32 ConnectionIndex, ENetProfilerConnectionMode Mode) const = 0;
+
+	// Enumerates packets in the provided inclusive packet interval, [PacketIndexIntervalStart, PacketIndexIntervalEnd].
 	virtual void EnumeratePackets(uint32 ConnectionIndex, ENetProfilerConnectionMode Mode, uint32 PacketIndexIntervalStart, uint32 PacketIndexIntervalEnd, TFunctionRef<void(const FNetProfilerPacket&)> Callback) const = 0;
+	// Returns a change number incremented each time a change occurs in the packets for the specified connection and connection mode.
 	virtual uint32 GetPacketChangeCount(uint32 ConnectionIndex, ENetProfilerConnectionMode Mode) const = 0;
 
-	// Enumerate packet content events by range
+	// Enumerates packet content events in the inclusive event index interval [StartEventIndex, EndEventIndex]. */
 	virtual void EnumeratePacketContentEventsByIndex(uint32 ConnectionIndex, ENetProfilerConnectionMode Mode, uint32 StartEventIndex, uint32 EndEventIndex, TFunctionRef<void(const FNetProfilerContentEvent&)> Callback) const = 0;
+	// Enumerates packet content events for a packet, in the exclusive bit range [StartPosition, EndPosition). */
 	virtual void EnumeratePacketContentEventsByPosition(uint32 ConnectionIndex, ENetProfilerConnectionMode Mode, uint32 PacketIndex, uint32 StartPosition, uint32 EndPosition, TFunctionRef<void(const FNetProfilerContentEvent&)> Callback) const = 0;
+	// Returns a change number incremented each time a change occurs in the packet content events for the specified connection and connection mode. */
 	virtual uint32 GetPacketContentEventChangeCount(uint32 ConnectionIndex, ENetProfilerConnectionMode Mode) const = 0;
 
-	// Stats queries
+	// Computes aggregated stats for a packet interval or for a range of content events in a single packet.
+	// [PacketIndexIntervalStart, PacketIndexIntervalEnd] is the inclusive packet interval.
+	// [StartPosition, EndPosition) is the exclusive bit range interval; only used when PacketIndexIntervalStart == PacketIndexIntervalEnd.
 	virtual ITable<FNetProfilerAggregatedStats>* CreateAggregation(uint32 ConnectionIndex, ENetProfilerConnectionMode Mode, uint32 PacketIndexIntervalStart, uint32 PacketIndexIntervalEnd, uint32 StartPosition, uint32 EndPosition) const = 0;
 };
 
