Fixed static analysis warning (UE-91127).

diff --git a/Engine/Source/Developer/TraceServices/Private/Model/NetProfilerProvider.cpp b/Engine/Source/Developer/TraceServices/Private/Model/NetProfilerProvider.cpp
--- a/Engine/Source/Developer/TraceServices/Private/Model/NetProfilerProvider.cpp
+++ b/Engine/Source/Developer/TraceServices/Private/Model/NetProfilerProvider.cpp
@@ -598,8 +598,6 @@ ITable<FNetProfilerAggregatedStats>* FNetProfilerProvider::CreateAggregation(uin
 	// Iterate over content events
 	if (PacketIndexIntervalStart == PacketIndexIntervalEnd)
 	{
-		ensure(Helper.StackSize == 0); // stack should be empty
-
 		EnumeratePacketContentEventsByPosition(ConnectionIndex, Mode, PacketIndexIntervalStart, StartPosition, EndPosition, GetStatsFunction);
 
 		// Pops the remaining events from the stack.
