Added an "end thread" event. This is used to as an indication that CPU trace analysis should flush its state and drop the thread.

diff --git a/Engine/Source/Developer/TraceServices/Private/Analyzers/CpuProfilerTraceAnalysis.cpp b/Engine/Source/Developer/TraceServices/Private/Analyzers/CpuProfilerTraceAnalysis.cpp
--- a/Engine/Source/Developer/TraceServices/Private/Analyzers/CpuProfilerTraceAnalysis.cpp
+++ b/Engine/Source/Developer/TraceServices/Private/Analyzers/CpuProfilerTraceAnalysis.cpp
@@ -26,6 +26,7 @@ void FCpuProfilerAnalyzer::OnAnalysisBegin(const FOnAnalysisContext& Context)
 
 	Builder.RouteEvent(RouteId_EventSpec, "CpuProfiler", "EventSpec");
 	Builder.RouteEvent(RouteId_EventBatch, "CpuProfiler", "EventBatch");
+	Builder.RouteEvent(RouteId_EndThread, "CpuProfiler", "EndThread");
 	Builder.RouteEvent(RouteId_EndCapture, "CpuProfiler", "EndCapture");
 	Builder.RouteEvent(RouteId_ChannelAnnounce, "$Trace", "ChannelAnnounce");
 	Builder.RouteEvent(RouteId_ChannelToggle, "$Trace", "ChannelToggle");
@@ -52,6 +53,19 @@ bool FCpuProfilerAnalyzer::OnEvent(uint16 RouteId, const FOnEventContext& Contex
 		}
 		break;
 	}
+	case RouteId_EndThread:
+	{
+		uint32 ThreadId = EventData.GetValue<uint32>("ThreadId");
+		FThreadState& ThreadState = GetThreadState(ThreadId);
+		const double Timestamp = Context.SessionContext.TimestampFromCycle(ThreadState.LastCycle);
+		Session.UpdateDurationSeconds(Timestamp);
+		while (ThreadState.ScopeStack.Num())
+		{
+			ThreadState.ScopeStack.Pop();
+			ThreadState.Timeline->AppendEndEvent(Timestamp);
+		}
+		ThreadStatesMap.Remove(ThreadId);
+	}
 	case RouteId_EventBatch:
 	case RouteId_EndCapture:
 	{
diff --git a/Engine/Source/Developer/TraceServices/Private/Analyzers/CpuProfilerTraceAnalysis.h b/Engine/Source/Developer/TraceServices/Private/Analyzers/CpuProfilerTraceAnalysis.h
--- a/Engine/Source/Developer/TraceServices/Private/Analyzers/CpuProfilerTraceAnalysis.h
+++ b/Engine/Source/Developer/TraceServices/Private/Analyzers/CpuProfilerTraceAnalysis.h
@@ -41,6 +41,7 @@ private:
 	{
 		RouteId_EventSpec,
 		RouteId_EventBatch,
+		RouteId_EndThread,
 		RouteId_EndCapture,
 		RouteId_ChannelAnnounce,
 		RouteId_ChannelToggle,
diff --git a/Engine/Source/Runtime/Core/Private/ProfilingDebugging/CpuProfilerTrace.cpp b/Engine/Source/Runtime/Core/Private/ProfilingDebugging/CpuProfilerTrace.cpp
--- a/Engine/Source/Runtime/Core/Private/ProfilingDebugging/CpuProfilerTrace.cpp
+++ b/Engine/Source/Runtime/Core/Private/ProfilingDebugging/CpuProfilerTrace.cpp
@@ -27,6 +27,10 @@ UE_TRACE_EVENT_BEGIN(CpuProfiler, EndCapture, Important)
 	UE_TRACE_EVENT_FIELD(uint32, ThreadId)
 UE_TRACE_EVENT_END()
 
+UE_TRACE_EVENT_BEGIN(CpuProfiler, EndThread, NoSync)
+	UE_TRACE_EVENT_FIELD(uint32, ThreadId)
+UE_TRACE_EVENT_END()
+
 struct FCpuProfilerTraceInternal
 {
 	enum
@@ -69,6 +73,12 @@ struct FCpuProfilerTraceInternal
 		{
 		}
 
+		virtual ~FThreadBuffer()
+		{
+			UE_TRACE_LOG(CpuProfiler, EndThread, CpuChannel)
+				<< EndThread.ThreadId(FPlatformTLS::GetCurrentThreadId());
+		}
+
 		uint64 LastCycle = 0;
 		uint16 BufferSize = 0;
 		uint8 Buffer[MaxBufferSize];
