Fixed out-of-bounds array access and mixed up CPU scope names in Unreal Insights.

diff --git a/Engine/Source/Developer/TraceServices/Private/Analyzers/CpuProfilerTraceAnalysis.cpp b/Engine/Source/Developer/TraceServices/Private/Analyzers/CpuProfilerTraceAnalysis.cpp
--- a/Engine/Source/Developer/TraceServices/Private/Analyzers/CpuProfilerTraceAnalysis.cpp
+++ b/Engine/Source/Developer/TraceServices/Private/Analyzers/CpuProfilerTraceAnalysis.cpp
@@ -169,8 +169,9 @@ void FCpuProfilerAnalyzer::DefineScope(uint32 Id, const TCHAR* Name)
 {
 	if (ScopeIdToEventIdMap.Contains(Id))
 	{
-		TimingProfilerProvider.SetTimerName(ScopeIdToEventIdMap[Id], Name);
-		ScopeNameToEventIdMap.Add(Name, Id);
+		uint32 EventId = ScopeIdToEventIdMap[Id];
+		TimingProfilerProvider.SetTimerName(EventId, Name);
+		ScopeNameToEventIdMap.Add(Name, EventId);
 	}
 	else
 	{
