Fix crash in MonotonicTimeline caused by incorrently cached pointer #rb Ionut.Matasaru #jira UE-101108

diff --git a/Engine/Source/Developer/TraceServices/Public/Model/MonotonicTimeline.h b/Engine/Source/Developer/TraceServices/Public/Model/MonotonicTimeline.h
--- a/Engine/Source/Developer/TraceServices/Public/Model/MonotonicTimeline.h
+++ b/Engine/Source/Developer/TraceServices/Public/Model/MonotonicTimeline.h
@@ -802,12 +802,14 @@ private:
 		FEventScopeEntry& ScopeEntry = DetailLevel.ScopeEntries.PushBack();
 		ScopeEntry.Time = IsEnter ? -Time : Time;
 		FEventScopeEntryPage* LastPage = DetailLevel.ScopeEntries.GetLastPage();
-		if (LastPage != DetailLevel.InsertionState.CurrentScopeEntryPage)
+		uint64 LastPageIndex = DetailLevel.ScopeEntries.NumPages() - 1;
+		if (LastPageIndex != DetailLevel.InsertionState.CurrentScopeEntryPageIndex)
 		{
-			//// At the very first call, CurrentScopeEntryPage will be nullptr
-			if (DetailLevel.InsertionState.CurrentScopeEntryPage != nullptr)
+			// At the very first call, CurrentScopeEntryPage will be nullptr
+			if (DetailLevel.InsertionState.CurrentScopeEntryPageIndex != (uint64) -1)
 			{
-				int32 PreviousPageInitialStackCount = DetailLevel.InsertionState.CurrentScopeEntryPage->InitialStackCount;
+				FEventScopeEntryPage* CurrentScopeEntryPage = DetailLevel.ScopeEntries.GetPage(DetailLevel.InsertionState.CurrentScopeEntryPageIndex);
+				int32 PreviousPageInitialStackCount = CurrentScopeEntryPage->InitialStackCount;
 				int32 CurrentDepth = DetailLevel.InsertionState.CurrentDepth;
 				// Update the open scopes that were also open at the beginning of the last page so the values 
 				// represent stats up to and including the current page
@@ -815,7 +817,7 @@ private:
 				for (; ii < PreviousPageInitialStackCount && ii < CurrentDepth; ++ii)
 				{
 					FEventStackEntry& InsertionStateStackEntry = DetailLevel.InsertionState.EventStack[ii];
-					FEventStackEntry& PreviousPageStackEntry = DetailLevel.InsertionState.CurrentScopeEntryPage->InitialStack[ii];
+					FEventStackEntry& PreviousPageStackEntry = CurrentScopeEntryPage->InitialStack[ii];
 
 					if (InsertionStateStackEntry.EnterScopeIndex == PreviousPageStackEntry.EnterScopeIndex
 						&& InsertionStateStackEntry.EventIndex == PreviousPageStackEntry.EventIndex)
@@ -828,7 +830,8 @@ private:
 					}
 				}
 			}
-			DetailLevel.InsertionState.CurrentScopeEntryPage = LastPage;
+
+			DetailLevel.InsertionState.CurrentScopeEntryPageIndex = LastPageIndex;
 			LastPage->BeginTime = Time;
 			LastPage->BeginEventIndex = DetailLevel.Events.Num();
 			LastPage->EndEventIndex = LastPage->BeginEventIndex;
@@ -865,9 +868,11 @@ private:
 
 			FEventStackEntry& StackEntry = DetailLevel.InsertionState.EventStack[DetailLevel.InsertionState.CurrentDepth];
 			StackEntry.ExclTime += Time - DetailLevel.InsertionState.LastTime;
-			if (DetailLevel.InsertionState.CurrentDepth < DetailLevel.InsertionState.CurrentScopeEntryPage->InitialStackCount)
+
+			FEventScopeEntryPage* CurrentScopeEntryPage = DetailLevel.ScopeEntries.GetPage(DetailLevel.InsertionState.CurrentScopeEntryPageIndex);
+			if (DetailLevel.InsertionState.CurrentDepth < CurrentScopeEntryPage->InitialStackCount)
 			{
-				FEventStackEntry& PreviousPageStackEntry = DetailLevel.InsertionState.CurrentScopeEntryPage->InitialStack[DetailLevel.InsertionState.CurrentDepth];
+				FEventStackEntry& PreviousPageStackEntry = CurrentScopeEntryPage->InitialStack[DetailLevel.InsertionState.CurrentDepth];
 
 				if (StackEntry.EnterScopeIndex == PreviousPageStackEntry.EnterScopeIndex
 					&& StackEntry.EventIndex == PreviousPageStackEntry.EventIndex)
@@ -885,7 +890,8 @@ private:
 
 	void AddEvent(FDetailLevel& DetailLevel, const EventType& Event)
 	{
-		++DetailLevel.InsertionState.CurrentScopeEntryPage->EndEventIndex;
+		FEventScopeEntryPage* CurrentScopeEntryPage = DetailLevel.ScopeEntries.GetPage(DetailLevel.InsertionState.CurrentScopeEntryPageIndex);
+		++CurrentScopeEntryPage->EndEventIndex;
 		DetailLevel.Events.PushBack() = Event;
 
 		//Event.DebugType = Owner.EventTypes[TypeIndex];
diff --git a/Engine/Source/Developer/TraceServices/Public/Model/MonotonicTimelineData.h b/Engine/Source/Developer/TraceServices/Public/Model/MonotonicTimelineData.h
--- a/Engine/Source/Developer/TraceServices/Public/Model/MonotonicTimelineData.h
+++ b/Engine/Source/Developer/TraceServices/Public/Model/MonotonicTimelineData.h
@@ -54,7 +54,7 @@ namespace Trace
 		int32 PendingDepth = -1;
 		FDetailLevelDepthState<EventType> DepthStates[SettingsType::MaxDepth];
 		FEventStackEntry EventStack[SettingsType::MaxDepth];
-		FEventScopeEntryPage* CurrentScopeEntryPage = nullptr;
+		uint64 CurrentScopeEntryPageIndex = (uint64) -1;
 	};
 
 	template<typename EventType, typename SettingsType>
