Added callstack analysis in TraceServices

Add CallstacksAnalyzer and CallbacksProvider which enables users to query traced callstack ids.

diff --git a/Engine/Source/Developer/TraceServices/Private/Analyzers/CallstacksAnalysis.cpp b/Engine/Source/Developer/TraceServices/Private/Analyzers/CallstacksAnalysis.cpp
new file mode 100644
--- /dev/null
+++ b/Engine/Source/Developer/TraceServices/Private/Analyzers/CallstacksAnalysis.cpp
@@ -0,0 +1,34 @@
+// Copyright Epic Games, Inc. All Rights Reserved.
+
+#include "CallstacksAnalysis.h"
+#include "TraceServices/Model/AnalysisSession.h"
+#include "Model/CallstacksProvider.h"
+
+////////////////////////////////////////////////////////////////////////////////
+FCallstacksAnalyzer::FCallstacksAnalyzer(Trace::IAnalysisSession& Session, Trace::FCallstacksProvider* InProvider)
+	: Provider(InProvider)
+{
+}
+
+////////////////////////////////////////////////////////////////////////////////
+void FCallstacksAnalyzer::OnAnalysisBegin(const FOnAnalysisContext& Context)
+{
+	FInterfaceBuilder& Builder = Context.InterfaceBuilder;
+	Builder.RouteEvent(RouteId_Callstack, "Memory", "CallstackSpec");
+}
+
+
+////////////////////////////////////////////////////////////////////////////////
+bool FCallstacksAnalyzer::OnEvent(uint16 RouteId, EStyle Style, const FOnEventContext& Context)
+{
+	switch(RouteId)
+	{
+		case RouteId_Callstack:
+			const uint64 Id = Context.EventData.GetValue<uint64>("Id");
+			const TArrayReader<uint64>& Frames = Context.EventData.GetArray<uint64>("Frames");
+			Provider->AddCallstack(Id, Frames.GetData(), uint8(Frames.Num()));
+			break;
+	}
+	return true;
+}
+
diff --git a/Engine/Source/Developer/TraceServices/Private/Analyzers/CallstacksAnalysis.h b/Engine/Source/Developer/TraceServices/Private/Analyzers/CallstacksAnalysis.h
new file mode 100644
--- /dev/null
+++ b/Engine/Source/Developer/TraceServices/Private/Analyzers/CallstacksAnalysis.h
@@ -0,0 +1,31 @@
+// Copyright Epic Games, Inc. All Rights Reserved.
+
+#pragma once
+
+#include "Trace/Analyzer.h"
+
+////////////////////////////////////////////////////////////////////////////////
+namespace Trace
+{
+	class IAnalysisSession;
+	class FCallstacksProvider;
+};
+
+////////////////////////////////////////////////////////////////////////////////
+class FCallstacksAnalyzer
+	: public Trace::IAnalyzer
+{
+public:
+					FCallstacksAnalyzer(Trace::IAnalysisSession& Session, Trace::FCallstacksProvider* Provider);
+	virtual void	OnAnalysisBegin(const FOnAnalysisContext& Context) override;
+	virtual bool 	OnEvent(uint16 RouteId, EStyle Style, const FOnEventContext& Context) override;
+
+private:
+	enum Routes 
+	{
+		RouteId_Callstack,
+	};
+
+	Trace::FCallstacksProvider* Provider;
+};
+
diff --git a/Engine/Source/Developer/TraceServices/Private/Model/CallstacksProvider.cpp b/Engine/Source/Developer/TraceServices/Private/Model/CallstacksProvider.cpp
new file mode 100644
--- /dev/null
+++ b/Engine/Source/Developer/TraceServices/Private/Model/CallstacksProvider.cpp
@@ -0,0 +1,111 @@
+// Copyright Epic Games, Inc. All Rights Reserved.
+
+#include "CallstacksProvider.h"
+#include "Misc/ScopeRWLock.h"
+#include "TraceServices/Model/AnalysisSession.h"
+#include "Algo/Unique.h"
+
+namespace Trace {
+
+/////////////////////////////////////////////////////////////////////
+#ifdef TRACE_CALLSTACK_STATS
+static struct FCallstackProviderStats
+{
+	uint64		Callstacks;
+	uint64		Frames;
+	uint64		FrameCountHistogram[256];
+} GCallstackStats;
+#endif
+
+/////////////////////////////////////////////////////////////////////
+FCallstacksProvider::FCallstacksProvider(IAnalysisSession& Session)
+	: Frames(Session.GetLinearAllocator(), FramesPerPage) 
+{
+}
+
+/////////////////////////////////////////////////////////////////////
+void FCallstacksProvider::AddCallstack(uint64 InCallstackId, const uint64* InFrames, uint8 InFrameCount)
+{
+	if (InFrameCount == 0)
+	{
+		return;
+	}
+
+#ifdef TRACE_CALLSTACK_STATS
+	GCallstackStats.Callstacks++;
+	GCallstackStats.Frames += InFrameCount;
+	GCallstackStats.FrameCountHistogram[InFrameCount]++;
+#endif
+
+	// Make sure all the frames fit on one page by appending dummy entries.
+	const uint64 PageHeadroom = Frames.GetPageSize() - (Frames.Num() % Frames.GetPageSize());
+	if (PageHeadroom < InFrameCount)
+	{
+		FRWScopeLock WriteLock(EntriesLock, SLT_Write);
+		uint64 EntriesToAdd = PageHeadroom + 1; // Fill page and allocate one on next
+		do { Frames.PushBack(); } while (--EntriesToAdd);
+	}
+
+	// Append the incoming frames
+	const uint64 FirstFrame = Frames.Num();
+	for (uint32 FrameIdx = 0; FrameIdx < InFrameCount; ++FrameIdx)
+	{
+		FStackFrame& F = Frames.PushBack();
+		F.Addr = InFrames[FrameIdx];
+		F.Symbol = nullptr;
+	}
+
+	uint64 FrameIdxAndLen = (uint64(InFrameCount) << EntryLenShift) | (~EntryLenMask & FirstFrame);
+
+	{
+		FRWScopeLock WriteLock(EntriesLock, SLT_Write);
+		CallstackEntries.Add(InCallstackId, FCallstackEntry{ FrameIdxAndLen });
+	}
+}
+
+/////////////////////////////////////////////////////////////////////
+FCallstack FCallstacksProvider::GetCallstack(uint64 CallstackId) 
+{
+	FRWScopeLock ReadLock(EntriesLock, SLT_ReadOnly);
+	FCallstackEntry* Entry = CallstackEntries.Find(CallstackId);
+	if (Entry)
+	{
+		const uint64 FirstFrameIdx = ~EntryLenMask & Entry->CallstackLenIndex;
+		const uint8 FrameCount = Entry->CallstackLenIndex >> EntryLenShift;
+		return FCallstack(&Frames[FirstFrameIdx], FrameCount);
+	}
+	return FCallstack(nullptr, 0);	
+}
+
+/////////////////////////////////////////////////////////////////////
+void FCallstacksProvider::GetCallstacks(const TArrayView<uint64>& CallstackIds, FCallstack* OutCallstacks)
+{
+	uint64 OutIdx(0);
+	check(OutCallstacks != nullptr);
+
+	FRWScopeLock ReadLock(EntriesLock, SLT_ReadOnly);
+	for (uint64 CallstackId : CallstackIds)
+	{
+		FCallstackEntry* Entry = CallstackEntries.Find(CallstackId);
+		if (Entry)
+		{
+			const uint64 FirstFrameIdx = ~EntryLenMask & Entry->CallstackLenIndex;
+			const uint8 FrameCount = Entry->CallstackLenIndex >> EntryLenShift;
+			new(&OutCallstacks[OutIdx]) FCallstack(&Frames[FirstFrameIdx], FrameCount);
+		}
+		else
+		{
+			new(&OutCallstacks[OutIdx]) FCallstack(nullptr, 0);	
+		}
+		OutIdx++;
+	}
+}
+
+/////////////////////////////////////////////////////////////////////
+FName FCallstacksProvider::GetName() const
+{
+	static FName Name(TEXT("CallstacksProvider"));
+	return Name;
+}
+
+} // namespace Trace
diff --git a/Engine/Source/Developer/TraceServices/Private/Model/CallstacksProvider.h b/Engine/Source/Developer/TraceServices/Private/Model/CallstacksProvider.h
new file mode 100644
--- /dev/null
+++ b/Engine/Source/Developer/TraceServices/Private/Model/CallstacksProvider.h
@@ -0,0 +1,46 @@
+#pragma once
+
+#include "TraceServices/Model/Callstack.h"
+#include "HAL/CriticalSection.h"
+#include "Common/PagedArray.h"
+
+namespace Trace {
+
+class IAnalysisSession;
+
+/////////////////////////////////////////////////////////////////////
+class FCallstacksProvider : public ICallstacksProvider 
+{
+public:
+					FCallstacksProvider(IAnalysisSession& Session);
+	virtual 		~FCallstacksProvider() {}
+
+	FCallstack	 	GetCallstack(uint64 CallstackId) override;
+	void			GetCallstacks(const TArrayView<uint64>& CallstackIds, FCallstack* OutCallstacks) override;
+	void 			AddCallstack(uint64 CallstackId, const uint64* Frames, uint8 FrameCount);
+	FName			GetName() const;
+
+private:
+	struct FCallstackEntry
+	{
+		uint64 CallstackLenIndex;
+	};
+
+	enum 
+	{
+		FramesPerPage	= 65536, // 16 bytes/entry -> 1 Mb per page
+	};
+
+	enum : uint64 
+	{
+		EntryLenMask	= 0xff00000000000000,
+		EntryLenShift	= 56,
+	};
+
+
+	FRWLock 						EntriesLock;
+	TMap<uint64, FCallstackEntry>	CallstackEntries;
+	TPagedArray<FStackFrame> 		Frames;
+};
+
+} // namespace Trace
diff --git a/Engine/Source/Developer/TraceServices/Private/Modules/MemoryModule.cpp b/Engine/Source/Developer/TraceServices/Private/Modules/MemoryModule.cpp
--- a/Engine/Source/Developer/TraceServices/Private/Modules/MemoryModule.cpp
+++ b/Engine/Source/Developer/TraceServices/Private/Modules/MemoryModule.cpp
@@ -2,8 +2,10 @@
 
 #include "MemoryModule.h"
 #include "Analyzers/AllocationsAnalysis.h"
+#include "Analyzers/CallstacksAnalysis.h"
 #include "Analyzers/MemoryAnalysis.h"
 #include "Model/AllocationsProvider.h"
+#include "Model/CallstacksProvider.h"
 #include "TraceServices/Model/AnalysisSession.h"
 
 #if defined(UE_USE_ALLOCATIONS_PROVIDER)
@@ -100,6 +102,10 @@ void FMemoryModule::OnAnalysisBegin(IAnalysisSession& Session)
 	TEMP_TestAllocationsProvider();
 #endif // UE_USE_ALLOCATIONS_PROVIDER
 
+	FCallstacksProvider* CallstacksProvider = new FCallstacksProvider(Session);
+	Session.AddProvider(CallstacksProvider->GetName(), CallstacksProvider);
+	Session.AddAnalyzer(new FCallstacksAnalyzer(Session, CallstacksProvider));
+
 	Session.AddAnalyzer(new FMemoryAnalyzer(Session));
 }
 
diff --git a/Engine/Source/Developer/TraceServices/Public/TraceServices/Model/Callstack.h b/Engine/Source/Developer/TraceServices/Public/TraceServices/Model/Callstack.h
new file mode 100644
--- /dev/null
+++ b/Engine/Source/Developer/TraceServices/Public/TraceServices/Model/Callstack.h
@@ -0,0 +1,88 @@
+#pragma once
+
+#include "TraceServices/Model/AnalysisSession.h"
+#include "TraceServices/Model/Memory.h"
+
+namespace Trace
+{
+
+/////////////////////////////////////////////////////////////////////
+struct SymbolEntry
+{
+	TCHAR*			Name;
+	TCHAR*			Filename;
+	TCHAR*			Module;
+	uint32			Line;
+};
+
+/////////////////////////////////////////////////////////////////////
+struct FStackFrame
+{
+	uint64 			Addr;
+	SymbolEntry*	Symbol;
+};
+
+/////////////////////////////////////////////////////////////////////
+struct FCallstack
+{
+					FCallstack(FStackFrame* FirstEntry, uint8 FrameCount);
+	/** Get the number of stack frames in callstack. */
+	uint32 			Num(); 
+	/** Gets the address at a given stack depth. */
+	uint64 			Addr(uint8 Depth);
+	/** Gets the cached symbol name at a given stack depth. */
+	TCHAR* 			SymbolName(uint8 Depth);
+
+private:
+	FStackFrame* 	FirstEntry;
+	uint8			FrameCount;
+};
+
+/////////////////////////////////////////////////////////////////////
+class ICallstacksProvider
+	: public IProvider
+{
+public:
+	virtual				~ICallstacksProvider() = default;
+
+	/**
+	  * Queries a callstack id.
+	  * @param CallstackId		Callstack id to query
+	  * @return					Callstack information. If id is not found a callstack with zero stack depth is returned.
+	  */
+	virtual FCallstack	GetCallstack(uint64 CallstackId) = 0;
+	
+	/**
+	  * Queries a set of callstack ids.
+	  * @param CallstackIds 	List of callstack ids to query
+	  * @param OutCallstacks	Output list of callstacks. Caller is responsible for allocating space according to CallstackIds length.
+	  */
+	virtual void		GetCallstacks(const TArrayView<uint64>& CallstackIds, FCallstack* OutCallstacks) = 0;
+};
+
+/////////////////////////////////////////////////////////////////////
+inline FCallstack::FCallstack(FStackFrame* InFirstEntry, uint8 InFrameCount)
+	: FirstEntry(InFirstEntry)
+	, FrameCount(InFrameCount)
+{
+}
+
+/////////////////////////////////////////////////////////////////////
+inline uint32 FCallstack::Num()
+{
+	return FrameCount;
+}
+
+/////////////////////////////////////////////////////////////////////
+inline uint64 FCallstack::Addr(uint8 Depth)
+{
+	return Depth < FrameCount ? FirstEntry[Depth].Addr : 0;
+}
+
+/////////////////////////////////////////////////////////////////////
+inline TCHAR* FCallstack::SymbolName(uint8 Depth)
+{
+	return Depth < FrameCount ? FirstEntry[Depth].Symbol->Name : 0;
+}
+
+} // namespace Trace
