Fixed copyright header.

diff --git a/Engine/Source/Developer/TraceServices/Private/Model/CallstacksProvider.h b/Engine/Source/Developer/TraceServices/Private/Model/CallstacksProvider.h
--- a/Engine/Source/Developer/TraceServices/Private/Model/CallstacksProvider.h
+++ b/Engine/Source/Developer/TraceServices/Private/Model/CallstacksProvider.h
@@ -1,3 +1,5 @@
+// Copyright Epic Games, Inc. All Rights Reserved.
+
 #pragma once
 
 #include "TraceServices/Model/Callstack.h"
@@ -9,15 +11,15 @@ namespace Trace {
 class IAnalysisSession;
 
 /////////////////////////////////////////////////////////////////////
-class FCallstacksProvider : public ICallstacksProvider 
+class FCallstacksProvider : public ICallstacksProvider
 {
 public:
 					FCallstacksProvider(IAnalysisSession& Session);
-	virtual 		~FCallstacksProvider() {}
+	virtual			~FCallstacksProvider() {}
 
-	FCallstack	 	GetCallstack(uint64 CallstackId) override;
+	FCallstack		GetCallstack(uint64 CallstackId) override;
 	void			GetCallstacks(const TArrayView<uint64>& CallstackIds, FCallstack* OutCallstacks) override;
-	void 			AddCallstack(uint64 CallstackId, const uint64* Frames, uint8 FrameCount);
+	void			AddCallstack(uint64 CallstackId, const uint64* Frames, uint8 FrameCount);
 	FName			GetName() const;
 
 private:
@@ -26,21 +28,21 @@ private:
 		uint64 CallstackLenIndex;
 	};
 
-	enum 
+	enum
 	{
 		FramesPerPage	= 65536, // 16 bytes/entry -> 1 Mb per page
 	};
 
-	enum : uint64 
+	enum : uint64
 	{
 		EntryLenMask	= 0xff00000000000000,
 		EntryLenShift	= 56,
 	};
 
 
-	FRWLock 						EntriesLock;
+	FRWLock							EntriesLock;
 	TMap<uint64, FCallstackEntry>	CallstackEntries;
-	TPagedArray<FStackFrame> 		Frames;
+	TPagedArray<FStackFrame>		Frames;
 };
 
 } // namespace Trace
diff --git a/Engine/Source/Developer/TraceServices/Public/TraceServices/Model/Callstack.h b/Engine/Source/Developer/TraceServices/Public/TraceServices/Model/Callstack.h
--- a/Engine/Source/Developer/TraceServices/Public/TraceServices/Model/Callstack.h
+++ b/Engine/Source/Developer/TraceServices/Public/TraceServices/Model/Callstack.h
@@ -1,3 +1,5 @@
+// Copyright Epic Games, Inc. All Rights Reserved.
+
 #pragma once
 
 #include "TraceServices/Model/AnalysisSession.h"
@@ -18,7 +20,7 @@ struct SymbolEntry
 /////////////////////////////////////////////////////////////////////
 struct FStackFrame
 {
-	uint64 			Addr;
+	uint64			Addr;
 	SymbolEntry*	Symbol;
 };
 
@@ -27,14 +29,14 @@ struct FCallstack
 {
 					FCallstack(FStackFrame* FirstEntry, uint8 FrameCount);
 	/** Get the number of stack frames in callstack. */
-	uint32 			Num(); 
+	uint32			Num();
 	/** Gets the address at a given stack depth. */
-	uint64 			Addr(uint8 Depth);
+	uint64			Addr(uint8 Depth);
 	/** Gets the cached symbol name at a given stack depth. */
-	TCHAR* 			SymbolName(uint8 Depth);
+	TCHAR*			SymbolName(uint8 Depth);
 
 private:
-	FStackFrame* 	FirstEntry;
+	FStackFrame*	FirstEntry;
 	uint8			FrameCount;
 };
 
@@ -51,10 +53,10 @@ public:
 	  * @return					Callstack information. If id is not found a callstack with zero stack depth is returned.
 	  */
 	virtual FCallstack	GetCallstack(uint64 CallstackId) = 0;
-	
+
 	/**
 	  * Queries a set of callstack ids.
-	  * @param CallstackIds 	List of callstack ids to query
+	  * @param CallstackIds		List of callstack ids to query
 	  * @param OutCallstacks	Output list of callstacks. Caller is responsible for allocating space according to CallstackIds length.
 	  */
 	virtual void		GetCallstacks(const TArrayView<uint64>& CallstackIds, FCallstack* OutCallstacks) = 0;
