Moved the TraceServices module into its own namespace.

diff --git a/Engine/Source/Developer/TraceServices/Private/Analyzers/CsvProfilerTraceAnalysis.cpp b/Engine/Source/Developer/TraceServices/Private/Analyzers/CsvProfilerTraceAnalysis.cpp
--- a/Engine/Source/Developer/TraceServices/Private/Analyzers/CsvProfilerTraceAnalysis.cpp
+++ b/Engine/Source/Developer/TraceServices/Private/Analyzers/CsvProfilerTraceAnalysis.cpp
@@ -5,7 +5,10 @@
 #include "TraceServices/Model/Counters.h"
 #include "TraceServices/Model/Frames.h"
 
-FCsvProfilerAnalyzer::FCsvProfilerAnalyzer(Trace::IAnalysisSession& InSession, Trace::FCsvProfilerProvider& InCsvProfilerProvider, Trace::ICounterProvider& InCounterProvider, const Trace::IFrameProvider& InFrameProvider, const Trace::IThreadProvider& InThreadProvider)
+namespace TraceServices
+{
+
+FCsvProfilerAnalyzer::FCsvProfilerAnalyzer(IAnalysisSession& InSession, FCsvProfilerProvider& InCsvProfilerProvider, ICounterProvider& InCounterProvider, const IFrameProvider& InFrameProvider, const IThreadProvider& InThreadProvider)
 	: Session(InSession)
 	, CsvProfilerProvider(InCsvProfilerProvider)
 	, CounterProvider(InCounterProvider)
@@ -62,7 +65,7 @@ void FCsvProfilerAnalyzer::OnAnalysisEnd()
 
 bool FCsvProfilerAnalyzer::OnEvent(uint16 RouteId, EStyle Style, const FOnEventContext& Context)
 {
-	Trace::FAnalysisSessionEditScope _(Session);
+	FAnalysisSessionEditScope _(Session);
 
 	const auto& EventData = Context.EventData;
 	switch (RouteId)
@@ -208,10 +211,10 @@ void FCsvProfilerAnalyzer::DefineStatSeries(uint64 StatId, const TCHAR* Name, in
 	}
 }
 
-const TCHAR* FCsvProfilerAnalyzer::GetStatSeriesName(const FStatSeriesDefinition* Definition, Trace::ECsvStatSeriesType Type, FThreadState& ThreadState, bool bIsCount)
+const TCHAR* FCsvProfilerAnalyzer::GetStatSeriesName(const FStatSeriesDefinition* Definition, ECsvStatSeriesType Type, FThreadState& ThreadState, bool bIsCount)
 {
 	FString Name = Definition->Name;
-	if (Type == Trace::CsvStatSeriesType_Timer || bIsCount)
+	if (Type == CsvStatSeriesType_Timer || bIsCount)
 	{
 		// Add a /<Threadname> prefix
 		Name = ThreadState.ThreadName + TEXT("/") + Name;
@@ -232,7 +235,7 @@ const TCHAR* FCsvProfilerAnalyzer::GetStatSeriesName(const FStatSeriesDefinition
 	return Session.StoreString(*Name);
 }
 
-FCsvProfilerAnalyzer::FStatSeriesInstance& FCsvProfilerAnalyzer::GetStatSeries(uint64 StatId, Trace::ECsvStatSeriesType Type, FThreadState& ThreadState)
+FCsvProfilerAnalyzer::FStatSeriesInstance& FCsvProfilerAnalyzer::GetStatSeries(uint64 StatId, ECsvStatSeriesType Type, FThreadState& ThreadState)
 {
 	FStatSeriesDefinition* Definition;
 	FStatSeriesDefinition** FindIt = StatSeriesMap.Find(StatId);
@@ -261,10 +264,10 @@ FCsvProfilerAnalyzer::FStatSeriesInstance& FCsvProfilerAnalyzer::GetStatSeries(u
 	ThreadState.StatSeries[Definition->ColumnIndex] = Instance;
 	const TCHAR* StatSeriesName = GetStatSeriesName(Definition, Type, ThreadState, false);
 	Instance->ProviderHandle = CsvProfilerProvider.AddSeries(StatSeriesName, Type);
-	Instance->ProviderCountHandle = CsvProfilerProvider.AddSeries(GetStatSeriesName(Definition, Type, ThreadState, true), Trace::CsvStatSeriesType_CustomStatInt);
+	Instance->ProviderCountHandle = CsvProfilerProvider.AddSeries(GetStatSeriesName(Definition, Type, ThreadState, true), CsvStatSeriesType_CustomStatInt);
 	Instance->Counter = CounterProvider.CreateCounter();
 	Instance->Counter->SetName(StatSeriesName);
-	Instance->Counter->SetIsFloatingPoint(Type != Trace::CsvStatSeriesType_CustomStatInt);
+	Instance->Counter->SetIsFloatingPoint(Type != CsvStatSeriesType_CustomStatInt);
 	Instance->Type = Type;
 	Instance->FrameType = ThreadState.FrameType;
 
@@ -365,7 +368,7 @@ void FCsvProfilerAnalyzer::HandleMarker(const FOnEventContext& Context, FThreadS
 				{
 					const FEventTime& EventTime = Context.EventTime;
 					double Elapsed = EventTime.AsSeconds(Marker.Cycle) - EventTime.AsSeconds(StartMarker.Cycle);
-					FStatSeriesInstance& StatSeries = GetStatSeries(Marker.StatId, Trace::CsvStatSeriesType_Timer, ThreadState);
+					FStatSeriesInstance& StatSeries = GetStatSeries(Marker.StatId, CsvStatSeriesType_Timer, ThreadState);
 					SetTimerValue(StatSeries, FrameNumber, Elapsed * 1000.0, !Marker.bIsExclusiveInsertedMarker);
 				}
 			}
@@ -377,7 +380,7 @@ void FCsvProfilerAnalyzer::HandleCustomStatEvent(const FOnEventContext& Context,
 {
 	uint32 ThreadId = FTraceAnalyzerUtils::GetThreadIdField(Context);
 	FThreadState& ThreadState = GetThreadState(ThreadId);
-	FStatSeriesInstance& StatSeries = GetStatSeries(Context.EventData.GetValue<uint64>("StatId"), bIsFloat ? Trace::CsvStatSeriesType_CustomStatFloat : Trace::CsvStatSeriesType_CustomStatInt, ThreadState);
+	FStatSeriesInstance& StatSeries = GetStatSeries(Context.EventData.GetValue<uint64>("StatId"), bIsFloat ? CsvStatSeriesType_CustomStatFloat : CsvStatSeriesType_CustomStatInt, ThreadState);
 	ECsvOpType OpType = static_cast<ECsvOpType>(Context.EventData.GetValue<uint8>("OpType"));
 	uint32 FrameNumber = GetFrameNumberForTimestamp(ThreadState.FrameType, Context.EventTime.AsSeconds(Context.EventData.GetValue<uint64>("Cycle")));
 	if (bIsFloat)
@@ -431,18 +434,18 @@ void FCsvProfilerAnalyzer::Flush(FStatSeriesInstance& StatSeries)
 	double CounterTimestamp;
 	if (StatSeries.CurrentFrame == 0)
 	{
-		const Trace::FFrame* Frame = FrameProvider.GetFrame(StatSeries.FrameType, 0);
+		const FFrame* Frame = FrameProvider.GetFrame(StatSeries.FrameType, 0);
 		check(Frame);
 		CounterTimestamp = Frame->StartTime;
 	}
 	else
 	{
-		const Trace::FFrame* Frame = FrameProvider.GetFrame(StatSeries.FrameType, StatSeries.CurrentFrame - 1);
-		const Trace::FFrame* NextFrame = FrameProvider.GetFrame(StatSeries.FrameType, StatSeries.CurrentFrame);
+		const FFrame* Frame = FrameProvider.GetFrame(StatSeries.FrameType, StatSeries.CurrentFrame - 1);
+		const FFrame* NextFrame = FrameProvider.GetFrame(StatSeries.FrameType, StatSeries.CurrentFrame);
 		check(NextFrame);
 		CounterTimestamp = Frame->EndTime;
 	}
-	if (StatSeries.Type == Trace::CsvStatSeriesType_CustomStatInt)
+	if (StatSeries.Type == CsvStatSeriesType_CustomStatInt)
 	{
 		CsvProfilerProvider.SetValue(StatSeries.ProviderHandle, StatSeries.CurrentFrame, StatSeries.CurrentValue.Value.AsInt);
 		StatSeries.Counter->SetValue(CounterTimestamp, StatSeries.CurrentValue.Value.AsInt);
@@ -548,3 +551,5 @@ void FCsvProfilerAnalyzer::SetCustomStatValue(FStatSeriesInstance& StatSeries, u
 	StatSeries.CurrentValue.bIsValid = true;
 	++StatSeries.CurrentCount;
 }
+
+} // namespace TraceServices
diff --git a/Engine/Source/Developer/TraceServices/Private/Analyzers/CsvProfilerTraceAnalysis.h b/Engine/Source/Developer/TraceServices/Private/Analyzers/CsvProfilerTraceAnalysis.h
--- a/Engine/Source/Developer/TraceServices/Private/Analyzers/CsvProfilerTraceAnalysis.h
+++ b/Engine/Source/Developer/TraceServices/Private/Analyzers/CsvProfilerTraceAnalysis.h
@@ -7,20 +7,20 @@
 #include "Model/CsvProfilerPrivate.h"
 #include "ProfilingDebugging/MiscTrace.h"
 
-namespace Trace
+namespace TraceServices
 {
-	class IAnalysisSession;
-	class IFrameProvider;
-	class IThreadProvider;
-	class ICounterProvider;
-	class IEditableCounter;
-}
+
+class IAnalysisSession;
+class IFrameProvider;
+class IThreadProvider;
+class ICounterProvider;
+class IEditableCounter;
 
 class FCsvProfilerAnalyzer
 	: public Trace::IAnalyzer
 {
 public:
-	FCsvProfilerAnalyzer(Trace::IAnalysisSession& Session, Trace::FCsvProfilerProvider& CsvProfilerProvider, Trace::ICounterProvider& CounterProvider, const Trace::IFrameProvider& FrameProvider, const Trace::IThreadProvider& ThreadProvider);
+	FCsvProfilerAnalyzer(IAnalysisSession& Session, FCsvProfilerProvider& CsvProfilerProvider, ICounterProvider& CounterProvider, const IFrameProvider& FrameProvider, const IThreadProvider& ThreadProvider);
 	~FCsvProfilerAnalyzer();
 	virtual void OnAnalysisBegin(const FOnAnalysisContext& Context) override;
 	virtual bool OnEvent(uint16 RouteId, EStyle Style, const FOnEventContext& Context) override;
@@ -86,11 +86,11 @@ private:
 	{
 		uint64 ProviderHandle = uint64(-1);
 		uint64 ProviderCountHandle = uint64(-1);
-		Trace::IEditableCounter* Counter = nullptr;
+		IEditableCounter* Counter = nullptr;
 		int64 CurrentFrame = -1;
 		FStatSeriesValue CurrentValue;
 		int64 CurrentCount = 0;
-		Trace::ECsvStatSeriesType Type = Trace::CsvStatSeriesType_CustomStatInt;
+		ECsvStatSeriesType Type = CsvStatSeriesType_CustomStatInt;
 		ETraceFrameType FrameType = TraceFrameType_Game;
 	};
 
@@ -115,8 +115,8 @@ private:
 	FThreadState& GetThreadState(uint32 ThreadId);
 	FStatSeriesDefinition* CreateStatSeries(const TCHAR* Name, int32 CategoryIndex);
 	void DefineStatSeries(uint64 StatId, const TCHAR* Name, int32 CategoryIndex, bool bIsInline);
-	const TCHAR* GetStatSeriesName(const FStatSeriesDefinition* Definition, Trace::ECsvStatSeriesType Type, FThreadState& ThreadState, bool bIsCount);
-	FStatSeriesInstance& GetStatSeries(uint64 StatId, Trace::ECsvStatSeriesType Type, FThreadState& ThreadState);
+	const TCHAR* GetStatSeriesName(const FStatSeriesDefinition* Definition, ECsvStatSeriesType Type, FThreadState& ThreadState, bool bIsCount);
+	FStatSeriesInstance& GetStatSeries(uint64 StatId, ECsvStatSeriesType Type, FThreadState& ThreadState);
 	void HandleMarkerEvent(const FOnEventContext& Context, bool bIsExclusive, bool bIsBegin);
 	void HandleMarker(const FOnEventContext& Context, FThreadState& ThreadState, const FTimingMarker& Marker);
 	void HandleCustomStatEvent(const FOnEventContext& Context, bool bIsFloat);
@@ -129,11 +129,11 @@ private:
 	void SetCustomStatValue(FStatSeriesInstance& StatSeries, uint32 FrameNumber, ECsvOpType OpType, int32 Value);
 	void SetCustomStatValue(FStatSeriesInstance& StatSeries, uint32 FrameNumber, ECsvOpType OpType, float Value);
 	
-	Trace::IAnalysisSession& Session;
-	Trace::FCsvProfilerProvider& CsvProfilerProvider;
-	Trace::ICounterProvider& CounterProvider;
-	const Trace::IFrameProvider& FrameProvider;
-	const Trace::IThreadProvider& ThreadProvider;
+	IAnalysisSession& Session;
+	FCsvProfilerProvider& CsvProfilerProvider;
+	ICounterProvider& CounterProvider;
+	const IFrameProvider& FrameProvider;
+	const IThreadProvider& ThreadProvider;
 
 	TMap<uint32, FThreadState*> ThreadStatesMap;
 	TMap<int32, const TCHAR*> CategoryMap;
@@ -148,3 +148,5 @@ private:
 
 	TArray<uint64> FrameBoundaries[TraceFrameType_Count];
 };
+
+} // namespace TraceServices
diff --git a/Engine/Source/Developer/TraceServices/Private/Analyzers/DiagnosticsAnalysis.cpp b/Engine/Source/Developer/TraceServices/Private/Analyzers/DiagnosticsAnalysis.cpp
--- a/Engine/Source/Developer/TraceServices/Private/Analyzers/DiagnosticsAnalysis.cpp
+++ b/Engine/Source/Developer/TraceServices/Private/Analyzers/DiagnosticsAnalysis.cpp
@@ -3,10 +3,13 @@
 #include "AnalysisServicePrivate.h"
 #include "Common/Utils.h"
 
-FDiagnosticsAnalyzer::FDiagnosticsAnalyzer(Trace::IAnalysisSession& InSession)
+namespace TraceServices
+{
+
+FDiagnosticsAnalyzer::FDiagnosticsAnalyzer(IAnalysisSession& InSession)
 	: Session(InSession)
 {
-	Provider = Session.EditProvider<Trace::FDiagnosticsProvider>(Trace::FDiagnosticsProvider::ProviderName);
+	Provider = Session.EditProvider<FDiagnosticsProvider>(FDiagnosticsProvider::ProviderName);
 }
 
 FDiagnosticsAnalyzer::~FDiagnosticsAnalyzer()
@@ -28,7 +31,7 @@ bool FDiagnosticsAnalyzer::OnEvent(uint16 RouteId, EStyle Style, const FOnEventC
 		return false;
 	}
 
-	Trace::FAnalysisSessionEditScope _(Session);
+	FAnalysisSessionEditScope _(Session);
 
 	const auto& EventData = Context.EventData;
 	switch (RouteId)
@@ -41,7 +44,7 @@ bool FDiagnosticsAnalyzer::OnEvent(uint16 RouteId, EStyle Style, const FOnEventC
 			return false;
 		}
 
-		Trace::FSessionInfo SessionInfo;
+		FSessionInfo SessionInfo;
 		uint8 AppNameOffset = EventData.GetValue<uint8>("AppNameOffset");
 		uint8 CommandLineOffset = EventData.GetValue<uint8>("CommandLineOffset");
 
@@ -65,7 +68,7 @@ bool FDiagnosticsAnalyzer::OnEvent(uint16 RouteId, EStyle Style, const FOnEventC
 	break;
 	case RouteId_Session2:
 	{
-		Trace::FSessionInfo SessionInfo;
+		FSessionInfo SessionInfo;
 
 		EventData.GetString("Platform", SessionInfo.Platform);
 		EventData.GetString("AppName", SessionInfo.AppName);
@@ -83,3 +86,5 @@ bool FDiagnosticsAnalyzer::OnEvent(uint16 RouteId, EStyle Style, const FOnEventC
 
 	return true;
 }
+
+} // namespace TraceServices
diff --git a/Engine/Source/Developer/TraceServices/Private/Analyzers/DiagnosticsAnalysis.h b/Engine/Source/Developer/TraceServices/Private/Analyzers/DiagnosticsAnalysis.h
--- a/Engine/Source/Developer/TraceServices/Private/Analyzers/DiagnosticsAnalysis.h
+++ b/Engine/Source/Developer/TraceServices/Private/Analyzers/DiagnosticsAnalysis.h
@@ -7,16 +7,16 @@
 #include "Model/DiagnosticsPrivate.h"
 #include "Containers/UnrealString.h"
 
-namespace Trace
+namespace TraceServices
 {
-	class IAnalysisSession;
-}
+
+class IAnalysisSession;
 
 class FDiagnosticsAnalyzer
 	: public Trace::IAnalyzer
 {
 public:
-	FDiagnosticsAnalyzer(Trace::IAnalysisSession& Session);
+	FDiagnosticsAnalyzer(IAnalysisSession& Session);
 	~FDiagnosticsAnalyzer();
 	virtual void OnAnalysisBegin(const FOnAnalysisContext& Context) override;
 	virtual bool OnEvent(uint16 RouteId, EStyle Style, const FOnEventContext& Context) override;
@@ -27,6 +27,8 @@ private:
 		RouteId_Session,
 		RouteId_Session2,
 	};
-	Trace::FDiagnosticsProvider* Provider;
-	Trace::IAnalysisSession& Session;
+	FDiagnosticsProvider* Provider;
+	IAnalysisSession& Session;
 };
+
+} // namespace TraceServices
diff --git a/Engine/Source/Developer/TraceServices/Private/Analyzers/GpuProfilerTraceAnalysis.cpp b/Engine/Source/Developer/TraceServices/Private/Analyzers/GpuProfilerTraceAnalysis.cpp
--- a/Engine/Source/Developer/TraceServices/Private/Analyzers/GpuProfilerTraceAnalysis.cpp
+++ b/Engine/Source/Developer/TraceServices/Private/Analyzers/GpuProfilerTraceAnalysis.cpp
@@ -3,7 +3,10 @@
 #include "AnalysisServicePrivate.h"
 #include "Common/Utils.h"
 
-FGpuProfilerAnalyzer::FGpuProfilerAnalyzer(Trace::FAnalysisSession& InSession, Trace::FTimingProfilerProvider& InTimingProfilerProvider)
+namespace TraceServices
+{
+
+FGpuProfilerAnalyzer::FGpuProfilerAnalyzer(FAnalysisSession& InSession, FTimingProfilerProvider& InTimingProfilerProvider)
 	: Session(InSession)
 	, TimingProfilerProvider(InTimingProfilerProvider)
 	, Timeline(TimingProfilerProvider.EditGpuTimeline())
@@ -22,7 +25,7 @@ void FGpuProfilerAnalyzer::OnAnalysisBegin(const FOnAnalysisContext& Context)
 
 bool FGpuProfilerAnalyzer::OnEvent(uint16 RouteId, EStyle Style, const FOnEventContext& Context)
 {
-	Trace::FAnalysisSessionEditScope _(Session);
+	FAnalysisSessionEditScope _(Session);
 
 	const auto& EventData = Context.EventData;
 
@@ -71,7 +74,7 @@ bool FGpuProfilerAnalyzer::OnEvent(uint16 RouteId, EStyle Style, const FOnEventC
 				BufferPtr += sizeof(uint32);
 				if (EventTypeMap.Contains(EventType))
 				{
-					Trace::FTimingProfilerEvent Event;
+					FTimingProfilerEvent Event;
 					Event.TimerIndex = EventTypeMap[EventType];
 					Timeline.AppendBeginEvent(LastTime, Event);
 				}
@@ -107,3 +110,5 @@ double FGpuProfilerAnalyzer::GpuTimestampToSessionTime(uint64 GpuMicroseconds)
 	}
 	return (GpuMicroseconds - GpuTimeOffset) / 1000000.0;
 }
+
+} // namespace TraceServices
diff --git a/Engine/Source/Developer/TraceServices/Private/Analyzers/GpuProfilerTraceAnalysis.h b/Engine/Source/Developer/TraceServices/Private/Analyzers/GpuProfilerTraceAnalysis.h
--- a/Engine/Source/Developer/TraceServices/Private/Analyzers/GpuProfilerTraceAnalysis.h
+++ b/Engine/Source/Developer/TraceServices/Private/Analyzers/GpuProfilerTraceAnalysis.h
@@ -5,16 +5,16 @@
 #include "Trace/Analyzer.h"
 #include "Model/TimingProfilerPrivate.h"
 
-namespace Trace
+namespace TraceServices
 {
-	class FAnalysisSession;
-}
+
+class FAnalysisSession;
 
 class FGpuProfilerAnalyzer
 	: public Trace::IAnalyzer
 {
 public:
-	FGpuProfilerAnalyzer(Trace::FAnalysisSession& Session, Trace::FTimingProfilerProvider& TimingProfilerProvider);
+	FGpuProfilerAnalyzer(FAnalysisSession& Session, FTimingProfilerProvider& TimingProfilerProvider);
 	virtual void OnAnalysisBegin(const FOnAnalysisContext& Context) override;
 	virtual bool OnEvent(uint16 RouteId, EStyle Style, const FOnEventContext& Context) override;
 
@@ -27,11 +27,13 @@ private:
 
 	double GpuTimestampToSessionTime(uint64 GpuMicroseconds);
 
-	Trace::FAnalysisSession& Session;
-	Trace::FTimingProfilerProvider& TimingProfilerProvider;
-	Trace::FTimingProfilerProvider::TimelineInternal& Timeline;
+	FAnalysisSession& Session;
+	FTimingProfilerProvider& TimingProfilerProvider;
+	FTimingProfilerProvider::TimelineInternal& Timeline;
 	TMap<uint64, uint32> EventTypeMap;
 	uint64 GpuTimeOffset;
 	double MinTime = 0.0f;
 	bool Calibrated;
 };
+
+} // namespace TraceServices
diff --git a/Engine/Source/Developer/TraceServices/Private/Analyzers/LoadTimeTraceAnalysis.cpp b/Engine/Source/Developer/TraceServices/Private/Analyzers/LoadTimeTraceAnalysis.cpp
--- a/Engine/Source/Developer/TraceServices/Private/Analyzers/LoadTimeTraceAnalysis.cpp
+++ b/Engine/Source/Developer/TraceServices/Private/Analyzers/LoadTimeTraceAnalysis.cpp
@@ -8,7 +8,10 @@
 #include "Common/Utils.h"
 #include <limits>
 
-FAsyncLoadingTraceAnalyzer::FAsyncLoadingTraceAnalyzer(Trace::IAnalysisSession& InSession, Trace::FLoadTimeProfilerProvider& InLoadTimeProfilerProvider)
+namespace TraceServices
+{
+
+FAsyncLoadingTraceAnalyzer::FAsyncLoadingTraceAnalyzer(IAnalysisSession& InSession, FLoadTimeProfilerProvider& InLoadTimeProfilerProvider)
 	: Session(InSession)
 	, LoadTimeProfilerProvider(InLoadTimeProfilerProvider)
 {
@@ -43,9 +46,9 @@ FAsyncLoadingTraceAnalyzer::FThreadState& FAsyncLoadingTraceAnalyzer::GetThreadS
 	return *ThreadState;
 }
 
-const Trace::FClassInfo* FAsyncLoadingTraceAnalyzer::GetClassInfo(uint64 ClassPtr) const
+const FClassInfo* FAsyncLoadingTraceAnalyzer::GetClassInfo(uint64 ClassPtr) const
 {
-	const Trace::FClassInfo* const* ClassInfo = ClassInfosMap.Find(ClassPtr);
+	const FClassInfo* const* ClassInfo = ClassInfosMap.Find(ClassPtr);
 	if (ClassInfo)
 	{
 		return *ClassInfo;
@@ -110,8 +113,8 @@ bool FAsyncLoadingTraceAnalyzer::OnEvent(uint16 RouteId, EStyle Style, const FOn
 		}
 		if (AsyncPackage)
 		{
-			Trace::FAnalysisSessionEditScope _(Session);
-			Trace::FPackageSummaryInfo& Summary = AsyncPackage->PackageInfo->Summary;
+			FAnalysisSessionEditScope _(Session);
+			FPackageSummaryInfo& Summary = AsyncPackage->PackageInfo->Summary;
 			Summary.TotalHeaderSize = EventData.GetValue<uint32>("TotalHeaderSize");
 			Summary.ImportCount = EventData.GetValue<uint32>("ImportCount");
 			Summary.ExportCount = EventData.GetValue<uint32>("ExportCount");
@@ -120,8 +123,8 @@ bool FAsyncLoadingTraceAnalyzer::OnEvent(uint16 RouteId, EStyle Style, const FOn
 	}
 	case RouteId_BeginCreateExport:
 	{
-		Trace::FAnalysisSessionEditScope _(Session);
-		Trace::FPackageExportInfo& Export = LoadTimeProfilerProvider.CreateExport();
+		FAnalysisSessionEditScope _(Session);
+		FPackageExportInfo& Export = LoadTimeProfilerProvider.CreateExport();
 		Export.SerialSize = EventData.GetValue<uint64>("SerialSize"); // Backwards compatibility
 		FAsyncPackageState* AsyncPackage = nullptr;
 		uint64 AsyncPackagePtr = EventData.GetValue<uint64>("AsyncPackage");
@@ -143,22 +146,22 @@ bool FAsyncLoadingTraceAnalyzer::OnEvent(uint16 RouteId, EStyle Style, const FOn
 		FThreadState& ThreadState = GetThreadState(ThreadId);
 		uint64 Cycle = EventData.GetValue<uint64>("Cycle");
 		double Time = Context.EventTime.AsSeconds(Cycle);
-		ThreadState.EnterExportScope(Time, &Export, Trace::LoadTimeProfilerObjectEventType_Create);
+		ThreadState.EnterExportScope(Time, &Export, LoadTimeProfilerObjectEventType_Create);
 		break;
 	}
 	case RouteId_EndCreateExport:
 	{
-		Trace::FAnalysisSessionEditScope _(Session);
+		FAnalysisSessionEditScope _(Session);
 		uint32 ThreadId = FTraceAnalyzerUtils::GetThreadIdField(Context);
 		FThreadState& ThreadState = GetThreadState(ThreadId);
-		if (ensure(ThreadState.GetCurrentExportScopeEventType() == Trace::LoadTimeProfilerObjectEventType_Create))
+		if (ensure(ThreadState.GetCurrentExportScopeEventType() == LoadTimeProfilerObjectEventType_Create))
 		{
-			Trace::FPackageExportInfo* Export = ThreadState.GetCurrentExportScope();
+			FPackageExportInfo* Export = ThreadState.GetCurrentExportScope();
 			if (Export)
 			{
 				uint64 ObjectPtr = EventData.GetValue<uint64>("Object");
 				ExportsMap.Add(ObjectPtr, Export);
-				const Trace::FClassInfo* ObjectClass = GetClassInfo(EventData.GetValue<uint64>("Class"));
+				const FClassInfo* ObjectClass = GetClassInfo(EventData.GetValue<uint64>("Class"));
 				Export->Class = ObjectClass;
 			}
 
@@ -170,9 +173,9 @@ bool FAsyncLoadingTraceAnalyzer::OnEvent(uint16 RouteId, EStyle Style, const FOn
 	}
 	case RouteId_BeginSerializeExport:
 	{
-		Trace::FAnalysisSessionEditScope _(Session);
+		FAnalysisSessionEditScope _(Session);
 		uint64 ObjectPtr = EventData.GetValue<uint64>("Object");
-		Trace::FPackageExportInfo* Export = ExportsMap.FindRef(ObjectPtr);
+		FPackageExportInfo* Export = ExportsMap.FindRef(ObjectPtr);
 		if (Export)
 		{
 			uint64 SerialSize = EventData.GetValue<uint64>("SerialSize");
@@ -182,22 +185,22 @@ bool FAsyncLoadingTraceAnalyzer::OnEvent(uint16 RouteId, EStyle Style, const FOn
 			}
 			if (Export->Package)
 			{
-				const_cast<Trace::FPackageInfo*>(Export->Package)->TotalExportsSerialSize += SerialSize;
+				const_cast<FPackageInfo*>(Export->Package)->TotalExportsSerialSize += SerialSize;
 			}
 		}
 		uint32 ThreadId = FTraceAnalyzerUtils::GetThreadIdField(Context);
 		FThreadState& ThreadState = GetThreadState(ThreadId);
 		uint64 Cycle = EventData.GetValue<uint64>("Cycle");
 		double Time = Context.EventTime.AsSeconds(Cycle);
-		ThreadState.EnterExportScope(Time, Export, Trace::LoadTimeProfilerObjectEventType_Serialize);
+		ThreadState.EnterExportScope(Time, Export, LoadTimeProfilerObjectEventType_Serialize);
 		break;
 	}
 	case RouteId_EndSerializeExport:
 	{
-		Trace::FAnalysisSessionEditScope _(Session);
+		FAnalysisSessionEditScope _(Session);
 		uint32 ThreadId = FTraceAnalyzerUtils::GetThreadIdField(Context);
 		FThreadState& ThreadState = GetThreadState(ThreadId);
-		if (ensure(ThreadState.GetCurrentExportScopeEventType() == Trace::LoadTimeProfilerObjectEventType_Serialize))
+		if (ensure(ThreadState.GetCurrentExportScopeEventType() == LoadTimeProfilerObjectEventType_Serialize))
 		{
 			uint64 Cycle = EventData.GetValue<uint64>("Cycle");
 			double Time = Context.EventTime.AsSeconds(Cycle);
@@ -207,22 +210,22 @@ bool FAsyncLoadingTraceAnalyzer::OnEvent(uint16 RouteId, EStyle Style, const FOn
 	}
 	case RouteId_BeginPostLoadExport:
 	{
-		Trace::FAnalysisSessionEditScope _(Session);
+		FAnalysisSessionEditScope _(Session);
 		uint64 ObjectPtr = EventData.GetValue<uint64>("Object");
-		Trace::FPackageExportInfo* Export = ExportsMap.FindRef(ObjectPtr);
+		FPackageExportInfo* Export = ExportsMap.FindRef(ObjectPtr);
 		uint32 ThreadId = FTraceAnalyzerUtils::GetThreadIdField(Context);
 		FThreadState& ThreadState = GetThreadState(ThreadId);
 		uint64 Cycle = EventData.GetValue<uint64>("Cycle");
 		double Time = Context.EventTime.AsSeconds(Cycle);
-		ThreadState.EnterExportScope(Time, Export, Trace::LoadTimeProfilerObjectEventType_PostLoad);
+		ThreadState.EnterExportScope(Time, Export, LoadTimeProfilerObjectEventType_PostLoad);
 		break;
 	}
 	case RouteId_EndPostLoadExport:
 	{
-		Trace::FAnalysisSessionEditScope _(Session);
+		FAnalysisSessionEditScope _(Session);
 		uint32 ThreadId = FTraceAnalyzerUtils::GetThreadIdField(Context);
 		FThreadState& ThreadState = GetThreadState(ThreadId);
-		if (ensure(ThreadState.GetCurrentExportScopeEventType() == Trace::LoadTimeProfilerObjectEventType_PostLoad))
+		if (ensure(ThreadState.GetCurrentExportScopeEventType() == LoadTimeProfilerObjectEventType_PostLoad))
 		{
 			uint64 Cycle = EventData.GetValue<uint64>("Cycle");
 			double Time = Context.EventTime.AsSeconds(Cycle);
@@ -232,7 +235,7 @@ bool FAsyncLoadingTraceAnalyzer::OnEvent(uint16 RouteId, EStyle Style, const FOn
 	}
 	case RouteId_BeginRequest:
 	{
-		Trace::FAnalysisSessionEditScope _(Session);
+		FAnalysisSessionEditScope _(Session);
 		uint64 RequestId = EventData.GetValue<uint64>("RequestId");
 		uint32 ThreadId = FTraceAnalyzerUtils::GetThreadIdField(Context);
 		//check(!ActiveRequestsMap.Contains(RequestId));
@@ -265,7 +268,7 @@ bool FAsyncLoadingTraceAnalyzer::OnEvent(uint16 RouteId, EStyle Style, const FOn
 			--RequestState->Group->ActiveRequestsCount;
 			if (RequestState->Group->LoadRequest && RequestState->Group->bIsClosed && RequestState->Group->ActiveRequestsCount == 0)
 			{
-				Trace::FAnalysisSessionEditScope _(Session);
+				FAnalysisSessionEditScope _(Session);
 				RequestState->Group->LoadRequest->EndTime = Context.EventTime.AsSeconds(RequestState->Group->LatestEndCycle);
 			}
 		}
@@ -273,11 +276,11 @@ bool FAsyncLoadingTraceAnalyzer::OnEvent(uint16 RouteId, EStyle Style, const FOn
 	}
 	case RouteId_BeginRequestGroup:
 	{
-		Trace::FAnalysisSessionEditScope _(Session);
+		FAnalysisSessionEditScope _(Session);
 		TSharedRef<FRequestGroupState> GroupState = MakeShared<FRequestGroupState>();
 		const TCHAR* FormatString = reinterpret_cast<const TCHAR*>(EventData.GetAttachment());
 		const uint8* FormatArgs = EventData.GetAttachment() + (FCString::Strlen(FormatString) + 1) * sizeof(TCHAR);
-		Trace::FFormatArgsHelper::Format(FormatBuffer, FormatBufferSize - 1, TempBuffer, FormatBufferSize - 1, FormatString, FormatArgs);
+		FFormatArgsHelper::Format(FormatBuffer, FormatBufferSize - 1, TempBuffer, FormatBufferSize - 1, FormatString, FormatArgs);
 		GroupState->Name = Session.StoreString(FormatBuffer);
 		uint32 ThreadId = FTraceAnalyzerUtils::GetThreadIdField(Context);
 		FThreadState& ThreadState = GetThreadState(ThreadId);
@@ -286,7 +289,7 @@ bool FAsyncLoadingTraceAnalyzer::OnEvent(uint16 RouteId, EStyle Style, const FOn
 	}
 	case RouteId_EndRequestGroup:
 	{
-		Trace::FAnalysisSessionEditScope _(Session);
+		FAnalysisSessionEditScope _(Session);
 		uint32 ThreadId = FTraceAnalyzerUtils::GetThreadIdField(Context);
 		FThreadState& ThreadState = GetThreadState(ThreadId);
 		if (ThreadState.RequestGroupStack.Num())
@@ -302,7 +305,7 @@ bool FAsyncLoadingTraceAnalyzer::OnEvent(uint16 RouteId, EStyle Style, const FOn
 	}
 	case RouteId_NewAsyncPackage:
 	{
-		Trace::FAnalysisSessionEditScope _(Session);
+		FAnalysisSessionEditScope _(Session);
 		uint64 AsyncPackagePtr = EventData.GetValue<uint64>("AsyncPackage");
 		//check(!ActivePackagesMap.Contains(AsyncPackagePtr));
 		FAsyncPackageState* AsyncPackageState = new FAsyncPackageState();
@@ -319,7 +322,7 @@ bool FAsyncLoadingTraceAnalyzer::OnEvent(uint16 RouteId, EStyle Style, const FOn
 			AsyncPackage->LoadStartCycle = EventData.GetValue<uint64>("Cycle");
 			if (AsyncPackage->PackageInfo)
 			{
-				Trace::FAnalysisSessionEditScope _(Session);
+				FAnalysisSessionEditScope _(Session);
 				double Time = Context.EventTime.AsSeconds(AsyncPackage->LoadStartCycle);
 				AsyncPackage->LoadHandle = LoadTimeProfilerProvider.BeginLoadPackage(*AsyncPackage->PackageInfo, Time);
 			}
@@ -335,7 +338,7 @@ bool FAsyncLoadingTraceAnalyzer::OnEvent(uint16 RouteId, EStyle Style, const FOn
 			AsyncPackage->LoadEndCycle = EventData.GetValue<uint64>("Cycle");
 			if (AsyncPackage->PackageInfo && AsyncPackage->LoadHandle != uint64(-1))
 			{
-				Trace::FAnalysisSessionEditScope _(Session);
+				FAnalysisSessionEditScope _(Session);
 				double Time = Context.EventTime.AsSeconds(AsyncPackage->LoadEndCycle);
 				LoadTimeProfilerProvider.EndLoadPackage(AsyncPackage->LoadHandle, Time);
 			}
@@ -378,16 +381,16 @@ bool FAsyncLoadingTraceAnalyzer::OnEvent(uint16 RouteId, EStyle Style, const FOn
 	}
 	case RouteId_ClassInfo:
 	{
-		Trace::FAnalysisSessionEditScope _(Session);
+		FAnalysisSessionEditScope _(Session);
 		uint64 ClassPtr = EventData.GetValue<uint64>("Class");
 		FString Name = FTraceAnalyzerUtils::LegacyAttachmentString<TCHAR>("Name", Context);
-		const Trace::FClassInfo& ClassInfo = LoadTimeProfilerProvider.AddClassInfo(*Name);
+		const FClassInfo& ClassInfo = LoadTimeProfilerProvider.AddClassInfo(*Name);
 		ClassInfosMap.Add(ClassPtr, &ClassInfo);
 		break;
 	}
 	case RouteId_BatchIssued:
 	{
-		Trace::FAnalysisSessionEditScope _(Session);
+		FAnalysisSessionEditScope _(Session);
 		uint64 Cycle = EventData.GetValue<uint64>("Cycle");
 		uint64 BatchId = EventData.GetValue<uint64>("BatchId");
 		double Time = Context.EventTime.AsSeconds(Cycle);
@@ -397,7 +400,7 @@ bool FAsyncLoadingTraceAnalyzer::OnEvent(uint16 RouteId, EStyle Style, const FOn
 	}
 	case RouteId_BatchResolved:
 	{
-		Trace::FAnalysisSessionEditScope _(Session);
+		FAnalysisSessionEditScope _(Session);
 		uint64 Cycle = EventData.GetValue<uint64>("Cycle");
 		uint64 BatchId = EventData.GetValue<uint64>("BatchId");
 		uint64 TotalSize = EventData.GetValue<uint64>("TotalSize");
@@ -411,13 +414,13 @@ bool FAsyncLoadingTraceAnalyzer::OnEvent(uint16 RouteId, EStyle Style, const FOn
 	}
 	case RouteId_BeginObjectScope:
 	{
-		Trace::FAnalysisSessionEditScope _(Session);
+		FAnalysisSessionEditScope _(Session);
 		uint64 ObjectPtr = EventData.GetValue<uint64>("Object");
-		Trace::FPackageExportInfo* Export = ExportsMap.FindRef(ObjectPtr);
+		FPackageExportInfo* Export = ExportsMap.FindRef(ObjectPtr);
 		uint32 ThreadId = FTraceAnalyzerUtils::GetThreadIdField(Context);
 		FThreadState& ThreadState = GetThreadState(ThreadId);
-		Trace::ELoadTimeProfilerObjectEventType EventType = static_cast<Trace::ELoadTimeProfilerObjectEventType>(EventData.GetValue<uint8>("EventType"));
-		//EventType = static_cast<Trace::ELoadTimeProfilerObjectEventType>(99); // debug
+		ELoadTimeProfilerObjectEventType EventType = static_cast<ELoadTimeProfilerObjectEventType>(EventData.GetValue<uint8>("EventType"));
+		//EventType = static_cast<ELoadTimeProfilerObjectEventType>(99); // debug
 		uint64 Cycle = EventData.GetValue<uint64>("Cycle");
 		double Time = Context.EventTime.AsSeconds(Cycle);
 		ThreadState.EnterExportScope(Time, Export, EventType);
@@ -425,10 +428,10 @@ bool FAsyncLoadingTraceAnalyzer::OnEvent(uint16 RouteId, EStyle Style, const FOn
 	}
 	case RouteId_EndObjectScope:
 	{
-		Trace::FAnalysisSessionEditScope _(Session);
+		FAnalysisSessionEditScope _(Session);
 		uint32 ThreadId = FTraceAnalyzerUtils::GetThreadIdField(Context);
 		FThreadState& ThreadState = GetThreadState(ThreadId);
-		//if (ensure(ThreadState.GetCurrentExportScopeEventType() == static_cast<Trace::ELoadTimeProfilerObjectEventType>(99))) // debug
+		//if (ensure(ThreadState.GetCurrentExportScopeEventType() == static_cast<ELoadTimeProfilerObjectEventType>(99))) // debug
 		{
 			uint64 Cycle = EventData.GetValue<uint64>("Cycle");
 			double Time = Context.EventTime.AsSeconds(Cycle);
@@ -452,7 +455,7 @@ bool FAsyncLoadingTraceAnalyzer::OnEvent(uint16 RouteId, EStyle Style, const FOn
 	return true;
 }
 
-void FAsyncLoadingTraceAnalyzer::FThreadState::EnterExportScope(double Time, const Trace::FPackageExportInfo* ExportInfo, Trace::ELoadTimeProfilerObjectEventType EventType)
+void FAsyncLoadingTraceAnalyzer::FThreadState::EnterExportScope(double Time, const FPackageExportInfo* ExportInfo, ELoadTimeProfilerObjectEventType EventType)
 {
 	FScopeStackEntry& StackEntry = CpuScopeStack[CpuScopeStackDepth++];
 	StackEntry.Event.Export = ExportInfo;
@@ -473,11 +476,11 @@ void FAsyncLoadingTraceAnalyzer::FThreadState::LeaveExportScope(double Time)
 	}
 	else
 	{
-		CurrentEvent = Trace::FLoadTimeProfilerCpuEvent();
+		CurrentEvent = FLoadTimeProfilerCpuEvent();
 	}
 }
 
-Trace::ELoadTimeProfilerObjectEventType FAsyncLoadingTraceAnalyzer::FThreadState::GetCurrentExportScopeEventType()
+ELoadTimeProfilerObjectEventType FAsyncLoadingTraceAnalyzer::FThreadState::GetCurrentExportScopeEventType()
 {
 	if (CpuScopeStackDepth > 0)
 	{
@@ -486,16 +489,16 @@ Trace::ELoadTimeProfilerObjectEventType FAsyncLoadingTraceAnalyzer::FThreadState
 	}
 	else
 	{
-		return Trace::ELoadTimeProfilerObjectEventType::LoadTimeProfilerObjectEventType_None;
+		return ELoadTimeProfilerObjectEventType::LoadTimeProfilerObjectEventType_None;
 	}
 }
 
-Trace::FPackageExportInfo* FAsyncLoadingTraceAnalyzer::FThreadState::GetCurrentExportScope()
+FPackageExportInfo* FAsyncLoadingTraceAnalyzer::FThreadState::GetCurrentExportScope()
 {
 	if (CpuScopeStackDepth > 0)
 	{
 		FScopeStackEntry& StackEntry = CpuScopeStack[CpuScopeStackDepth - 1];
-		return const_cast<Trace::FPackageExportInfo*>(StackEntry.Event.Export);
+		return const_cast<FPackageExportInfo*>(StackEntry.Event.Export);
 	}
 	else
 	{
@@ -509,8 +512,8 @@ void FAsyncLoadingTraceAnalyzer::PackageRequestAssociation(const FOnEventContext
 	{
 		RequestState->AsyncPackages.Add(AsyncPackageState);
 		AsyncPackageState->Request = RequestState;
-		Trace::FLoadRequest* LoadRequest = RequestState->Group->LoadRequest;
-		Trace::FAnalysisSessionEditScope _(Session);
+		FLoadRequest* LoadRequest = RequestState->Group->LoadRequest;
+		FAnalysisSessionEditScope _(Session);
 		if (!LoadRequest)
 		{
 			LoadRequest = &LoadTimeProfilerProvider.CreateRequest();
@@ -523,3 +526,5 @@ void FAsyncLoadingTraceAnalyzer::PackageRequestAssociation(const FOnEventContext
 		LoadRequest->Packages.Add(AsyncPackageState->PackageInfo);
 	}
 }
+
+} // namespace TraceServices
diff --git a/Engine/Source/Developer/TraceServices/Private/Analyzers/LoadTimeTraceAnalysis.h b/Engine/Source/Developer/TraceServices/Private/Analyzers/LoadTimeTraceAnalysis.h
--- a/Engine/Source/Developer/TraceServices/Private/Analyzers/LoadTimeTraceAnalysis.h
+++ b/Engine/Source/Developer/TraceServices/Private/Analyzers/LoadTimeTraceAnalysis.h
@@ -7,12 +7,12 @@
 #include "AnalysisServicePrivate.h"
 #include "Model/LoadTimeProfilerPrivate.h"
 
-namespace Trace
+namespace TraceServices
 {
-	struct FClassInfo;
-}
 
-inline bool operator!=(const Trace::FLoadTimeProfilerCpuEvent& Lhs, const Trace::FLoadTimeProfilerCpuEvent& Rhs)
+struct FClassInfo;
+
+inline bool operator!=(const FLoadTimeProfilerCpuEvent& Lhs, const FLoadTimeProfilerCpuEvent& Rhs)
 {
 	return Lhs.Package != Rhs.Package ||
 		Lhs.Export != Rhs.Export ||
@@ -22,7 +22,7 @@ inline bool operator!=(const Trace::FLoadTimeProfilerCpuEvent& Lhs, const Trace:
 class FAsyncLoadingTraceAnalyzer : public Trace::IAnalyzer
 {
 public:
-	FAsyncLoadingTraceAnalyzer(Trace::IAnalysisSession& Session, Trace::FLoadTimeProfilerProvider& LoadTimeProfilerProvider);
+	FAsyncLoadingTraceAnalyzer(IAnalysisSession& Session, FLoadTimeProfilerProvider& LoadTimeProfilerProvider);
 	virtual ~FAsyncLoadingTraceAnalyzer();
 
 	virtual void OnAnalysisBegin(const FOnAnalysisContext& Context) override;
@@ -36,7 +36,7 @@ private:
 	{
 		FString Name;
 		TArray<FRequestState*> Requests;
-		Trace::FLoadRequest* LoadRequest = nullptr;
+		FLoadRequest* LoadRequest = nullptr;
 		uint64 LatestEndCycle = 0;
 		uint64 ActiveRequestsCount = 0;
 		bool bIsClosed = false;
@@ -53,7 +53,7 @@ private:
 
 	struct FAsyncPackageState
 	{
-		Trace::FPackageInfo* PackageInfo = nullptr;
+		FPackageInfo* PackageInfo = nullptr;
 		FRequestState* Request = nullptr;
 		uint64 LoadHandle = uint64(-1);
 		uint64 LoadStartCycle = 0;
@@ -62,7 +62,7 @@ private:
 
 	struct FScopeStackEntry
 	{
-		Trace::FLoadTimeProfilerCpuEvent Event;
+		FLoadTimeProfilerCpuEvent Event;
 		bool EnteredEvent;
 	};
 
@@ -70,20 +70,20 @@ private:
 	{
 		FScopeStackEntry CpuScopeStack[256];
 		uint64 CpuScopeStackDepth = 0;
-		Trace::FLoadTimeProfilerCpuEvent CurrentEvent;
+		FLoadTimeProfilerCpuEvent CurrentEvent;
 		TArray<TSharedPtr<FRequestGroupState>> RequestGroupStack;
 		
-		Trace::FLoadTimeProfilerProvider::CpuTimelineInternal* CpuTimeline;
+		FLoadTimeProfilerProvider::CpuTimelineInternal* CpuTimeline;
 
-		void EnterExportScope(double Time, const Trace::FPackageExportInfo* ExportInfo, Trace::ELoadTimeProfilerObjectEventType EventType);
+		void EnterExportScope(double Time, const FPackageExportInfo* ExportInfo, ELoadTimeProfilerObjectEventType EventType);
 		void LeaveExportScope(double Time);
-		Trace::ELoadTimeProfilerObjectEventType GetCurrentExportScopeEventType();
-		Trace::FPackageExportInfo* GetCurrentExportScope();
+		ELoadTimeProfilerObjectEventType GetCurrentExportScopeEventType();
+		FPackageExportInfo* GetCurrentExportScope();
 	};
 
 	void PackageRequestAssociation(const FOnEventContext& Context, FAsyncPackageState* AsyncPackageState, FRequestState* RequestState);
 	FThreadState& GetThreadState(uint32 ThreadId);
-	const Trace::FClassInfo* GetClassInfo(uint64 ClassPtr) const;
+	const FClassInfo* GetClassInfo(uint64 ClassPtr) const;
 
 	enum : uint16
 	{
@@ -124,8 +124,8 @@ private:
 	TCHAR FormatBuffer[FormatBufferSize];
 	TCHAR TempBuffer[FormatBufferSize];
 
-	Trace::IAnalysisSession& Session;
-	Trace::FLoadTimeProfilerProvider& LoadTimeProfilerProvider;
+	IAnalysisSession& Session;
+	FLoadTimeProfilerProvider& LoadTimeProfilerProvider;
 
 	template<typename ValueType>
 	struct FPointerMapKeyFuncs
@@ -156,12 +156,14 @@ private:
 	using TPointerMap = TMap<uint64, ValueType, FDefaultSetAllocator, FPointerMapKeyFuncs<ValueType>>;
 
 	TPointerMap<FAsyncPackageState*> ActiveAsyncPackagesMap;
-	TPointerMap<Trace::FPackageExportInfo*> ExportsMap;
+	TPointerMap<FPackageExportInfo*> ExportsMap;
 	TMap<uint64, FRequestState*> ActiveRequestsMap;
 	TPointerMap<uint64> ActiveBatchesMap;
 	TMap<uint32, FThreadState*> ThreadStatesMap;
-	TPointerMap<const Trace::FClassInfo*> ClassInfosMap;
+	TPointerMap<const FClassInfo*> ClassInfosMap;
 
 	// Backwards compatibility
 	TPointerMap<FAsyncPackageState*> LinkerToAsyncPackageMap;
 };
+
+} // namespace TraceServices
diff --git a/Engine/Source/Developer/TraceServices/Private/Analyzers/LogTraceAnalysis.cpp b/Engine/Source/Developer/TraceServices/Private/Analyzers/LogTraceAnalysis.cpp
--- a/Engine/Source/Developer/TraceServices/Private/Analyzers/LogTraceAnalysis.cpp
+++ b/Engine/Source/Developer/TraceServices/Private/Analyzers/LogTraceAnalysis.cpp
@@ -5,7 +5,10 @@
 #include "Logging/LogTrace.h"
 #include "Model/LogPrivate.h"
 
-FLogTraceAnalyzer::FLogTraceAnalyzer(Trace::IAnalysisSession& InSession, Trace::FLogProvider& InLogProvider)
+namespace TraceServices
+{
+
+FLogTraceAnalyzer::FLogTraceAnalyzer(IAnalysisSession& InSession, FLogProvider& InLogProvider)
 	: Session(InSession)
 	, LogProvider(InLogProvider)
 {
@@ -23,7 +26,7 @@ void FLogTraceAnalyzer::OnAnalysisBegin(const FOnAnalysisContext& Context)
 
 bool FLogTraceAnalyzer::OnEvent(uint16 RouteId, EStyle Style, const FOnEventContext& Context)
 {
-	Trace::FAnalysisSessionEditScope _(Session);
+	FAnalysisSessionEditScope _(Session);
 
 	const auto& EventData = Context.EventData;
 	switch (RouteId)
@@ -31,7 +34,7 @@ bool FLogTraceAnalyzer::OnEvent(uint16 RouteId, EStyle Style, const FOnEventCont
 	case RouteId_LogCategory:
 	{
 		uint64 CategoryPointer = EventData.GetValue<uint64>("CategoryPointer");
-		Trace::FLogCategory& Category = LogProvider.GetCategory(CategoryPointer);
+		FLogCategory& Category = LogProvider.GetCategory(CategoryPointer);
 
 		FString Name = FTraceAnalyzerUtils::LegacyAttachmentString<TCHAR>("Name", Context);
 		Category.Name = Session.StoreString(*Name);
@@ -42,9 +45,9 @@ bool FLogTraceAnalyzer::OnEvent(uint16 RouteId, EStyle Style, const FOnEventCont
 	case RouteId_LogMessageSpec:
 	{
 		uint64 LogPoint = EventData.GetValue<uint64>("LogPoint");
-		Trace::FLogMessageSpec& Spec = LogProvider.GetMessageSpec(LogPoint);
+		FLogMessageSpec& Spec = LogProvider.GetMessageSpec(LogPoint);
 		uint64 CategoryPointer = EventData.GetValue<uint64>("CategoryPointer");
-		Trace::FLogCategory& Category = LogProvider.GetCategory(CategoryPointer);
+		FLogCategory& Category = LogProvider.GetCategory(CategoryPointer);
 		Spec.Category = &Category;
 		Spec.Line = EventData.GetValue<int32>("Line");
 		Spec.Verbosity = static_cast<ELogVerbosity::Type>(EventData.GetValue<uint8>("Verbosity"));
@@ -78,3 +81,5 @@ bool FLogTraceAnalyzer::OnEvent(uint16 RouteId, EStyle Style, const FOnEventCont
 
 	return true;
 }
+
+} // namespace TraceServices
diff --git a/Engine/Source/Developer/TraceServices/Private/Analyzers/LogTraceAnalysis.h b/Engine/Source/Developer/TraceServices/Private/Analyzers/LogTraceAnalysis.h
--- a/Engine/Source/Developer/TraceServices/Private/Analyzers/LogTraceAnalysis.h
+++ b/Engine/Source/Developer/TraceServices/Private/Analyzers/LogTraceAnalysis.h
@@ -4,17 +4,17 @@
 
 #include "Trace/Analyzer.h"
 
-namespace Trace
+namespace TraceServices
 {
-	class IAnalysisSession;
-	class FLogProvider;
-}
+
+class IAnalysisSession;
+class FLogProvider;
 
 class FLogTraceAnalyzer
 	: public Trace::IAnalyzer
 {
 public:
-	FLogTraceAnalyzer(Trace::IAnalysisSession& Session, Trace::FLogProvider& LogProvider);
+	FLogTraceAnalyzer(IAnalysisSession& Session, FLogProvider& LogProvider);
 	virtual void OnAnalysisBegin(const FOnAnalysisContext& Context) override;
 	virtual bool OnEvent(uint16 RouteId, EStyle Style, const FOnEventContext& Context) override;
 
@@ -26,6 +26,8 @@ private:
 		RouteId_LogMessage,
 	};
 
-	Trace::IAnalysisSession& Session;
-	Trace::FLogProvider& LogProvider;
+	IAnalysisSession& Session;
+	FLogProvider& LogProvider;
 };
+
+} // namespace TraceServices
diff --git a/Engine/Source/Developer/TraceServices/Private/Analyzers/MemoryAnalysis.cpp b/Engine/Source/Developer/TraceServices/Private/Analyzers/MemoryAnalysis.cpp
--- a/Engine/Source/Developer/TraceServices/Private/Analyzers/MemoryAnalysis.cpp
+++ b/Engine/Source/Developer/TraceServices/Private/Analyzers/MemoryAnalysis.cpp
@@ -3,11 +3,14 @@
 #include "AnalysisServicePrivate.h"
 #include "Common/Utils.h"
 
-FMemoryAnalyzer::FMemoryAnalyzer(Trace::IAnalysisSession& InSession)
+namespace TraceServices
+{
+
+FMemoryAnalyzer::FMemoryAnalyzer(IAnalysisSession& InSession)
 	: Provider(nullptr)
 	, Session(InSession)
 {
-	Provider = Session.EditProvider<Trace::FMemoryProvider>(Trace::FMemoryProvider::ProviderName);
+	Provider = Session.EditProvider<FMemoryProvider>(FMemoryProvider::ProviderName);
 }
 
 FMemoryAnalyzer::~FMemoryAnalyzer()
@@ -30,7 +33,7 @@ bool FMemoryAnalyzer::OnEvent(uint16 RouteId, EStyle Style, const FOnEventContex
 		return false;
 	}
 
-	Trace::FAnalysisSessionEditScope _(Session);
+	FAnalysisSessionEditScope _(Session);
 
 	const auto& EventData = Context.EventData;
 	switch (RouteId)
@@ -59,10 +62,10 @@ bool FMemoryAnalyzer::OnEvent(uint16 RouteId, EStyle Style, const FOnEventContex
 		const TArrayReader<int64>& Samples = EventData.GetArray<int64>("Values");
 
 		const uint32 TagsCount = Tags.Num();
-		TArray<Trace::FMemoryTagSample> Values;
+		TArray<FMemoryTagSample> Values;
 		for (uint32 i = 0; i < TagsCount; ++i)
 		{
-			Values.Push(Trace::FMemoryTagSample{ Samples[i] });
+			Values.Push(FMemoryTagSample{ Samples[i] });
 		}
 		Provider->AddTagSnapshot(TrackerId, Time, Tags.GetData(), Values.GetData(), TagsCount);
 		Sample++;
@@ -73,3 +76,5 @@ bool FMemoryAnalyzer::OnEvent(uint16 RouteId, EStyle Style, const FOnEventContex
 	}
 	return true;
 }
+
+} // namespace TraceServices
diff --git a/Engine/Source/Developer/TraceServices/Private/Analyzers/MemoryAnalysis.h b/Engine/Source/Developer/TraceServices/Private/Analyzers/MemoryAnalysis.h
--- a/Engine/Source/Developer/TraceServices/Private/Analyzers/MemoryAnalysis.h
+++ b/Engine/Source/Developer/TraceServices/Private/Analyzers/MemoryAnalysis.h
@@ -8,16 +8,16 @@
 #include "Containers/UnrealString.h"
 #include "Common/PagedArray.h"
 
-namespace Trace
+namespace TraceServices
 {
-	class IAnalysisSession;
-}
+
+class IAnalysisSession;
 
 class FMemoryAnalyzer
 	: public Trace::IAnalyzer
 {
 public:
-	FMemoryAnalyzer(Trace::IAnalysisSession& Session);
+	FMemoryAnalyzer(IAnalysisSession& Session);
 	~FMemoryAnalyzer();
 	virtual void OnAnalysisBegin(const FOnAnalysisContext& Context) override;
 	virtual bool OnEvent(uint16 RouteId, EStyle Style, const FOnEventContext& Context) override;
@@ -30,7 +30,9 @@ private:
 		RouteId_TagValue,
 	};
 
-	Trace::FMemoryProvider* Provider;
-	Trace::IAnalysisSession& Session;
+	FMemoryProvider* Provider;
+	IAnalysisSession& Session;
 	uint64 Sample = 0;
 };
+
+} // namespace TraceServices
diff --git a/Engine/Source/Developer/TraceServices/Private/Analyzers/MiscTraceAnalysis.cpp b/Engine/Source/Developer/TraceServices/Private/Analyzers/MiscTraceAnalysis.cpp
--- a/Engine/Source/Developer/TraceServices/Private/Analyzers/MiscTraceAnalysis.cpp
+++ b/Engine/Source/Developer/TraceServices/Private/Analyzers/MiscTraceAnalysis.cpp
@@ -8,12 +8,15 @@
 #include "Model/Channel.h"
 #include "Common/Utils.h"
 
-FMiscTraceAnalyzer::FMiscTraceAnalyzer(Trace::IAnalysisSession& InSession,
-									   Trace::FThreadProvider& InThreadProvider,
-									   Trace::FBookmarkProvider& InBookmarkProvider,
-									   Trace::FLogProvider& InLogProvider,
-									   Trace::FFrameProvider& InFrameProvider,
-									   Trace::FChannelProvider& InChannelProvider)
+namespace TraceServices
+{
+
+FMiscTraceAnalyzer::FMiscTraceAnalyzer(IAnalysisSession& InSession,
+									   FThreadProvider& InThreadProvider,
+									   FBookmarkProvider& InBookmarkProvider,
+									   FLogProvider& InLogProvider,
+									   FFrameProvider& InFrameProvider,
+									   FChannelProvider& InChannelProvider)
 	: Session(InSession)
 	, ThreadProvider(InThreadProvider)
 	, BookmarkProvider(InBookmarkProvider)
@@ -21,7 +24,7 @@ FMiscTraceAnalyzer::FMiscTraceAnalyzer(Trace::IAnalysisSession& InSession,
 	, FrameProvider(InFrameProvider)
 	, ChannelProvider(InChannelProvider)
 {
-	Trace::FLogCategory& BookmarkLogCategory = LogProvider.GetCategory(Trace::FLogProvider::ReservedLogCategory_Bookmark);
+	FLogCategory& BookmarkLogCategory = LogProvider.GetCategory(FLogProvider::ReservedLogCategory_Bookmark);
 	BookmarkLogCategory.Name = TEXT("LogBookmark");
 	BookmarkLogCategory.DefaultVerbosity = ELogVerbosity::All;
 }
@@ -52,7 +55,7 @@ void FMiscTraceAnalyzer::OnThreadInfo(const FThreadInfo& ThreadInfo)
 	uint32 ThreadId = ThreadInfo.GetId();
 	FString Name = ThreadInfo.GetName();
 
-	Trace::FAnalysisSessionEditScope _(Session);
+	FAnalysisSessionEditScope _(Session);
 
 	ThreadProvider.AddThread(ThreadId, *Name, EThreadPriority(ThreadInfo.GetSortHint()));
 
@@ -66,7 +69,7 @@ void FMiscTraceAnalyzer::OnThreadInfo(const FThreadInfo& ThreadInfo)
 
 bool FMiscTraceAnalyzer::OnEvent(uint16 RouteId, EStyle Style, const FOnEventContext& Context)
 {
-	Trace::FAnalysisSessionEditScope _(Session);
+	FAnalysisSessionEditScope _(Session);
 
 	const auto& EventData = Context.EventData;
 	switch (RouteId)
@@ -74,7 +77,7 @@ bool FMiscTraceAnalyzer::OnEvent(uint16 RouteId, EStyle Style, const FOnEventCon
 	case RouteId_BookmarkSpec:
 	{
 		uint64 BookmarkPoint = EventData.GetValue<uint64>("BookmarkPoint");
-		Trace::FBookmarkSpec& Spec = BookmarkProvider.GetSpec(BookmarkPoint);
+		FBookmarkSpec& Spec = BookmarkProvider.GetSpec(BookmarkPoint);
 		Spec.Line = EventData.GetValue<int32>("Line");
 
 		FString FileName;
@@ -94,8 +97,8 @@ bool FMiscTraceAnalyzer::OnEvent(uint16 RouteId, EStyle Style, const FOnEventCon
 		}
 
 
-		Trace::FLogMessageSpec& LogMessageSpec = LogProvider.GetMessageSpec(BookmarkPoint);
-		LogMessageSpec.Category = &LogProvider.GetCategory(Trace::FLogProvider::ReservedLogCategory_Bookmark);
+		FLogMessageSpec& LogMessageSpec = LogProvider.GetMessageSpec(BookmarkPoint);
+		LogMessageSpec.Category = &LogProvider.GetCategory(FLogProvider::ReservedLogCategory_Bookmark);
 		LogMessageSpec.Line = Spec.Line;
 		LogMessageSpec.File = Spec.File;
 		LogMessageSpec.FormatString = Spec.FormatString;
@@ -242,3 +245,5 @@ FMiscTraceAnalyzer::FThreadState* FMiscTraceAnalyzer::GetThreadState(uint32 Thre
 	}
 	return &ThreadStateMap[ThreadId].Get();
 }
+
+} // namespace TraceServices
diff --git a/Engine/Source/Developer/TraceServices/Private/Analyzers/MiscTraceAnalysis.h b/Engine/Source/Developer/TraceServices/Private/Analyzers/MiscTraceAnalysis.h
--- a/Engine/Source/Developer/TraceServices/Private/Analyzers/MiscTraceAnalysis.h
+++ b/Engine/Source/Developer/TraceServices/Private/Analyzers/MiscTraceAnalysis.h
@@ -7,26 +7,26 @@
 #include "ProfilingDebugging/MiscTrace.h"
 #include "Common/PagedArray.h"
 
-namespace Trace
+namespace TraceServices
 {
-	class IAnalysisSession;
-	class FThreadProvider;
-	class FBookmarkProvider;
-	class FLogProvider;
-	class FFrameProvider;
-	class FChannelProvider;
-}
+
+class IAnalysisSession;
+class FThreadProvider;
+class FBookmarkProvider;
+class FLogProvider;
+class FFrameProvider;
+class FChannelProvider;
 
 class FMiscTraceAnalyzer
 	: public Trace::IAnalyzer
 {
 public:
-	FMiscTraceAnalyzer(Trace::IAnalysisSession& Session,
-					   Trace::FThreadProvider& ThreadProvider,
-					   Trace::FBookmarkProvider& BookmarkProvider,
-					   Trace::FLogProvider& LogProvider,
-					   Trace::FFrameProvider& FrameProvider, 
-					   Trace::FChannelProvider& ChannelProvider);
+	FMiscTraceAnalyzer(IAnalysisSession& Session,
+					   FThreadProvider& ThreadProvider,
+					   FBookmarkProvider& BookmarkProvider,
+					   FLogProvider& LogProvider,
+					   FFrameProvider& FrameProvider, 
+					   FChannelProvider& ChannelProvider);
 	virtual void OnAnalysisBegin(const FOnAnalysisContext& Context) override;
 	virtual void OnThreadInfo(const FThreadInfo& ThreadInfo) override;
 	virtual bool OnEvent(uint16 RouteId, EStyle Style, const FOnEventContext& Context) override;
@@ -60,13 +60,15 @@ private:
 	void OnChannelAnnounce(const FOnEventContext& Context);
 	void OnChannelToggle(const FOnEventContext& Context);
 
-	Trace::IAnalysisSession& Session;
-	Trace::FThreadProvider& ThreadProvider;
-	Trace::FBookmarkProvider& BookmarkProvider;
-	Trace::FLogProvider& LogProvider;
-	Trace::FFrameProvider& FrameProvider;
-	Trace::FChannelProvider& ChannelProvider;
+	IAnalysisSession& Session;
+	FThreadProvider& ThreadProvider;
+	FBookmarkProvider& BookmarkProvider;
+	FLogProvider& LogProvider;
+	FFrameProvider& FrameProvider;
+	FChannelProvider& ChannelProvider;
 	TMap<uint32, TSharedRef<FThreadState>> ThreadStateMap;
 	uint64 LastFrameCycle[TraceFrameType_Count] = { 0, 0 };
 };
 
+
+} // namespace TraceServices
diff --git a/Engine/Source/Developer/TraceServices/Private/Analyzers/NetTraceAnalyzer.cpp b/Engine/Source/Developer/TraceServices/Private/Analyzers/NetTraceAnalyzer.cpp
--- a/Engine/Source/Developer/TraceServices/Private/Analyzers/NetTraceAnalyzer.cpp
+++ b/Engine/Source/Developer/TraceServices/Private/Analyzers/NetTraceAnalyzer.cpp
@@ -8,6 +8,9 @@
 DECLARE_LOG_CATEGORY_EXTERN(LogNetTrace, Log, All);
 DEFINE_LOG_CATEGORY(LogNetTrace);
 
+namespace TraceServices
+{
+
 enum ENetTraceAnalyzerVersion
 {
 	ENetTraceAnalyzerVersion_Initial = 1,
@@ -16,7 +19,7 @@ enum ENetTraceAnalyzerVersion
 };
 
 
-FNetTraceAnalyzer::FNetTraceAnalyzer(Trace::IAnalysisSession& InSession, Trace::FNetProfilerProvider& InNetProfilerProvider)
+FNetTraceAnalyzer::FNetTraceAnalyzer(IAnalysisSession& InSession, FNetProfilerProvider& InNetProfilerProvider)
 	: Session(InSession)
 	, NetProfilerProvider(InNetProfilerProvider)
 	, NetTraceVersion(0)
@@ -43,7 +46,7 @@ void FNetTraceAnalyzer::OnAnalysisBegin(const FOnAnalysisContext& Context)
 
 	// Default names
 	{
-		Trace::FAnalysisSessionEditScope _(Session);
+		FAnalysisSessionEditScope _(Session);
 		BunchHeaderNameIndex = NetProfilerProvider.AddNetProfilerName(TEXT("BunchHeader"));
 	}
 }
@@ -72,7 +75,7 @@ uint32 FNetTraceAnalyzer::GetTracedEventTypeIndex(uint16 NameIndex, uint8 Level)
 
 bool FNetTraceAnalyzer::OnEvent(uint16 RouteId, EStyle Style, const FOnEventContext& Context)
 {
-	Trace::FAnalysisSessionEditScope _(Session);
+	FAnalysisSessionEditScope _(Session);
 
 	// check that we always get the InitEvent before processing any other events
 	if (!ensure(RouteId == RouteId_InitEvent || NetTraceVersion > 0))
@@ -179,11 +182,11 @@ void FNetTraceAnalyzer::HandlePacketContentEvent(const FOnEventContext& Context,
 		return;
 	}
 
-	const Trace::ENetProfilerConnectionMode ConnectionMode = Trace::ENetProfilerConnectionMode(PacketType);
-	Trace::FNetProfilerConnectionData& ConnectionData = NetProfilerProvider.EditConnectionData(ConnectionState->ConnectionIndex, ConnectionMode);
+	const ENetProfilerConnectionMode ConnectionMode = ENetProfilerConnectionMode(PacketType);
+	FNetProfilerConnectionData& ConnectionData = NetProfilerProvider.EditConnectionData(ConnectionState->ConnectionIndex, ConnectionMode);
 	++ConnectionData.ContentEventChangeCount;
 
-	TArray<Trace::FNetProfilerContentEvent>& Events = (ConnectionState->BunchEvents)[ConnectionMode];
+	TArray<FNetProfilerContentEvent>& Events = (ConnectionState->BunchEvents)[ConnectionMode];
 	TArray<FBunchInfo>& BunchInfos = (ConnectionState->BunchInfos)[ConnectionMode];
 
 	// Decode batched events
@@ -203,7 +206,7 @@ void FNetTraceAnalyzer::HandlePacketContentEvent(const FOnEventContext& Context,
 			case EContentEventType::Object:
 			case EContentEventType::NameId:
 			{
-				Trace::FNetProfilerContentEvent& Event = Events.Emplace_GetRef();
+				FNetProfilerContentEvent& Event = Events.Emplace_GetRef();
 
 				const uint8 DecodedNestingLevel = *BufferPtr++;
 
@@ -327,9 +330,9 @@ void FNetTraceAnalyzer::HandlePacketContentEvent(const FOnEventContext& Context,
 	check(BufferPtr == BufferEnd);
 }
 
-void FNetTraceAnalyzer::AddEvent(TPagedArray<Trace::FNetProfilerContentEvent>& Events, const Trace::FNetProfilerContentEvent& InEvent, uint32 Offset, uint32 LevelOffset)
+void FNetTraceAnalyzer::AddEvent(TPagedArray<FNetProfilerContentEvent>& Events, const FNetProfilerContentEvent& InEvent, uint32 Offset, uint32 LevelOffset)
 {
-	Trace::FNetProfilerContentEvent& Event = Events.PushBack();
+	FNetProfilerContentEvent& Event = Events.PushBack();
 	
 	Event.EventTypeIndex = GetTracedEventTypeIndex(InEvent.NameIndex, InEvent.Level + LevelOffset);
 	Event.NameIndex =  InEvent.NameIndex;
@@ -340,9 +343,9 @@ void FNetTraceAnalyzer::AddEvent(TPagedArray<Trace::FNetProfilerContentEvent>& E
 	Event.BunchInfo = InEvent.BunchInfo;
 }
 
-void FNetTraceAnalyzer::AddEvent(TPagedArray<Trace::FNetProfilerContentEvent>& Events, uint32 StartPos, uint32 EndPos, uint32 Level, uint32 NameIndex, Trace::FNetProfilerBunchInfo BunchInfo)
+void FNetTraceAnalyzer::AddEvent(TPagedArray<FNetProfilerContentEvent>& Events, uint32 StartPos, uint32 EndPos, uint32 Level, uint32 NameIndex, FNetProfilerBunchInfo BunchInfo)
 {
-	Trace::FNetProfilerContentEvent& Event = Events.PushBack();
+	FNetProfilerContentEvent& Event = Events.PushBack();
 
 	Event.EventTypeIndex = GetTracedEventTypeIndex(NameIndex, Level);
 	Event.NameIndex = NameIndex; 
@@ -353,11 +356,11 @@ void FNetTraceAnalyzer::AddEvent(TPagedArray<Trace::FNetProfilerContentEvent>& E
 	Event.BunchInfo = BunchInfo;
 }
 
-void FNetTraceAnalyzer::FlushPacketEvents(FNetTraceConnectionState& ConnectionState, Trace::FNetProfilerConnectionData& ConnectionData, const Trace::ENetProfilerConnectionMode ConnectionMode)
+void FNetTraceAnalyzer::FlushPacketEvents(FNetTraceConnectionState& ConnectionState, FNetProfilerConnectionData& ConnectionData, const ENetProfilerConnectionMode ConnectionMode)
 {
-	TPagedArray<Trace::FNetProfilerContentEvent>& Events = ConnectionData.ContentEvents;
+	TPagedArray<FNetProfilerContentEvent>& Events = ConnectionData.ContentEvents;
 
-	TArray<Trace::FNetProfilerContentEvent>& BunchEvents = ConnectionState.BunchEvents[ConnectionMode];
+	TArray<FNetProfilerContentEvent>& BunchEvents = ConnectionState.BunchEvents[ConnectionMode];
 	const int32 NumPacketEvents = BunchEvents.Num();
 
 	int32 CurrentBunchEventIndex = 0;
@@ -371,7 +374,7 @@ void FNetTraceAnalyzer::FlushPacketEvents(FNetTraceConnectionState& ConnectionSt
 	// Inject any events reported before the first bunch
 	while (CurrentBunchEventIndex < NonBunchEventCount)
 	{
-		const Trace::FNetProfilerContentEvent& BunchEvent = BunchEvents[CurrentBunchEventIndex];
+		const FNetProfilerContentEvent& BunchEvent = BunchEvents[CurrentBunchEventIndex];
 
 		AddEvent(Events, BunchEvent, 0U, 0U);
 	
@@ -392,12 +395,12 @@ void FNetTraceAnalyzer::FlushPacketEvents(FNetTraceConnectionState& ConnectionSt
 			AddEvent(Events, NextBunchOffset, NextBunchOffset + Bunch.HeaderBits + Bunch.BunchBits, 0, Bunch.NameIndex, Bunch.BunchInfo);
 
 			// Bunch header event
-			AddEvent(Events, NextBunchOffset, NextBunchOffset + Bunch.HeaderBits, 1, BunchHeaderNameIndex, Trace::FNetProfilerBunchInfo::MakeBunchInfo(0));
+			AddEvent(Events, NextBunchOffset, NextBunchOffset + Bunch.HeaderBits, 1, BunchHeaderNameIndex, FNetProfilerBunchInfo::MakeBunchInfo(0));
 	
 			// Add events belonging to bunch, including the ones from merged bunches
 			for (uint32 EventIt = 0; EventIt < EventsToAdd; ++EventIt)
 			{
-				const Trace::FNetProfilerContentEvent& BunchEvent = BunchEvents[CurrentBunchEventIndex];
+				const FNetProfilerContentEvent& BunchEvent = BunchEvents[CurrentBunchEventIndex];
 
 				AddEvent(Events, BunchEvent, BunchOffset, 1U);
 				++CurrentBunchEventIndex;
@@ -427,7 +430,7 @@ void FNetTraceAnalyzer::HandlePacketEvent(const FOnEventContext& Context, const
 	const uint16 ConnectionId = EventData.GetValue<uint16>("ConnectionId");
 	const uint8 PacketType = EventData.GetValue<uint8>("PacketType");
 
-	const Trace::ENetProfilerConnectionMode ConnectionMode = Trace::ENetProfilerConnectionMode(PacketType);
+	const ENetProfilerConnectionMode ConnectionMode = ENetProfilerConnectionMode(PacketType);
 
 	// Update LastTimestamp, later on we will be able to get timestamps piggybacked from other analyzers
 	LastTimeStamp = Context.EventTime.AsSeconds(TimestampCycles);
@@ -439,8 +442,8 @@ void FNetTraceAnalyzer::HandlePacketEvent(const FOnEventContext& Context, const
 	}
 
 	// Add the packet
-	Trace::FNetProfilerConnectionData& ConnectionData = NetProfilerProvider.EditConnectionData(ConnectionState->ConnectionIndex, ConnectionMode);
-	Trace::FNetProfilerPacket& Packet = ConnectionData.Packets.PushBack();
+	FNetProfilerConnectionData& ConnectionData = NetProfilerProvider.EditConnectionData(ConnectionState->ConnectionIndex, ConnectionMode);
+	FNetProfilerPacket& Packet = ConnectionData.Packets.PushBack();
 	++ConnectionData.PacketChangeCount;
 
 	// Flush packet events
@@ -451,11 +454,11 @@ void FNetTraceAnalyzer::HandlePacketEvent(const FOnEventContext& Context, const
 	Packet.EventCount = ConnectionData.ContentEvents.Num() - Packet.StartEventIndex;
 	Packet.TimeStamp = GetLastTimestamp();
 	Packet.SequenceNumber = SequenceNumber;
-	Packet.DeliveryStatus = Trace::ENetProfilerDeliveryStatus::Unknown;
+	Packet.DeliveryStatus = ENetProfilerDeliveryStatus::Unknown;
 
 	Packet.ContentSizeInBits = PacketBits;
 	Packet.TotalPacketSizeInBytes = (Packet.ContentSizeInBits + 7u) >> 3u;
-	Packet.DeliveryStatus = Trace::ENetProfilerDeliveryStatus::Delivered;
+	Packet.DeliveryStatus = ENetProfilerDeliveryStatus::Delivered;
 
 	// Mark the beginning of a new packet
 	ConnectionState->CurrentPacketStartIndex[ConnectionMode] = ConnectionData.ContentEvents.Num();
@@ -481,10 +484,10 @@ void FNetTraceAnalyzer::HandlePacketDroppedEvent(const FOnEventContext& Context,
 		return;
 	}
 
-	Trace::FNetProfilerConnectionData& ConnectionData = NetProfilerProvider.EditConnectionData(ConnectionState->ConnectionIndex, Trace::ENetProfilerConnectionMode(PacketType));
+	FNetProfilerConnectionData& ConnectionData = NetProfilerProvider.EditConnectionData(ConnectionState->ConnectionIndex, ENetProfilerConnectionMode(PacketType));
 
 	// Update packet delivery status
-	NetProfilerProvider.EditPacketDeliveryStatus(ConnectionState->ConnectionIndex, Trace::ENetProfilerConnectionMode(PacketType), SequenceNumber, Trace::ENetProfilerDeliveryStatus::Dropped);
+	NetProfilerProvider.EditPacketDeliveryStatus(ConnectionState->ConnectionIndex, ENetProfilerConnectionMode(PacketType), SequenceNumber, ENetProfilerDeliveryStatus::Dropped);
 }
 
 void FNetTraceAnalyzer::HandleConnectionCretedEvent(const FOnEventContext& Context, const FEventData& EventData)
@@ -496,7 +499,7 @@ void FNetTraceAnalyzer::HandleConnectionCretedEvent(const FOnEventContext& Conte
 	check(!GameInstanceState->ActiveConnections.Contains(ConnectionId));
 
 	// Add to both active connections and to persistent connections
- 	Trace::FNetProfilerConnectionInternal& Connection = NetProfilerProvider.CreateConnection(GameInstanceState->GameInstanceIndex);
+ 	FNetProfilerConnectionInternal& Connection = NetProfilerProvider.CreateConnection(GameInstanceState->GameInstanceIndex);
 	TSharedRef<FNetTraceConnectionState> ConnectionState = MakeShared<FNetTraceConnectionState>();
 	GameInstanceState->ActiveConnections.Add(ConnectionId, ConnectionState);
 
@@ -504,11 +507,11 @@ void FNetTraceAnalyzer::HandleConnectionCretedEvent(const FOnEventContext& Conte
 	Connection.Connection.ConnectionId = ConnectionId;
 	Connection.Connection.LifeTime.Begin =  GetLastTimestamp();
 	ConnectionState->ConnectionIndex = Connection.Connection.ConnectionIndex;
-	ConnectionState->CurrentPacketStartIndex[Trace::ENetProfilerConnectionMode::Outgoing] = 0U;
-	ConnectionState->CurrentPacketStartIndex[Trace::ENetProfilerConnectionMode::Incoming] = 0U;
+	ConnectionState->CurrentPacketStartIndex[ENetProfilerConnectionMode::Outgoing] = 0U;
+	ConnectionState->CurrentPacketStartIndex[ENetProfilerConnectionMode::Incoming] = 0U;
 
-	ConnectionState->CurrentPacketBitOffset[Trace::ENetProfilerConnectionMode::Outgoing] = 0U;
-	ConnectionState->CurrentPacketBitOffset[Trace::ENetProfilerConnectionMode::Incoming] = 0U;
+	ConnectionState->CurrentPacketBitOffset[ENetProfilerConnectionMode::Outgoing] = 0U;
+	ConnectionState->CurrentPacketBitOffset[ENetProfilerConnectionMode::Incoming] = 0U;
 }
 
 void FNetTraceAnalyzer::HandleConnectionClosedEvent(const FOnEventContext& Context, const FEventData& EventData)
@@ -520,7 +523,7 @@ void FNetTraceAnalyzer::HandleConnectionClosedEvent(const FOnEventContext& Conte
 
 	if (TSharedRef<FNetTraceConnectionState>* ConnectionState = GameInstanceState->ActiveConnections.Find(ConnectionId))
 	{
-		if (Trace::FNetProfilerConnectionInternal* Connection = NetProfilerProvider.EditConnection((*ConnectionState)->ConnectionIndex))
+		if (FNetProfilerConnectionInternal* Connection = NetProfilerProvider.EditConnection((*ConnectionState)->ConnectionIndex))
 		{
 			// Update connection state
 			Connection->Connection.LifeTime.End =  GetLastTimestamp();
@@ -548,7 +551,7 @@ void FNetTraceAnalyzer::HandleObjectCreatedEvent(const FOnEventContext& Context,
 	
 	if (GameInstanceState->ActiveObjects.Contains(ObjectId))
 	{
-		if (Trace::FNetProfilerObjectInstance* ExistingInstance = NetProfilerProvider.EditObject(GameInstanceState->GameInstanceIndex, GameInstanceState->ActiveObjects[ObjectId].ObjectIndex))
+		if (FNetProfilerObjectInstance* ExistingInstance = NetProfilerProvider.EditObject(GameInstanceState->GameInstanceIndex, GameInstanceState->ActiveObjects[ObjectId].ObjectIndex))
 		{
 			if (ExistingInstance->NameIndex == NameIndex)
 			{
@@ -567,7 +570,7 @@ void FNetTraceAnalyzer::HandleObjectCreatedEvent(const FOnEventContext& Context,
 	}
 
 	// Add persistent object representation
-	Trace::FNetProfilerObjectInstance& ObjectInstance = NetProfilerProvider.CreateObject(GameInstanceState->GameInstanceIndex);
+	FNetProfilerObjectInstance& ObjectInstance = NetProfilerProvider.CreateObject(GameInstanceState->GameInstanceIndex);
 
 	// Fill in object data
 	ObjectInstance.LifeTime.Begin =  GetLastTimestamp();
@@ -590,7 +593,7 @@ void FNetTraceAnalyzer::HandleObjectDestroyedEvent(const FOnEventContext& Contex
 	FNetTraceActiveObjectState DestroyedObjectState;
 	if (GameInstanceState->ActiveObjects.RemoveAndCopyValue(ObjectId, DestroyedObjectState))
 	{
-		if (Trace::FNetProfilerObjectInstance* ObjectInstance = NetProfilerProvider.EditObject(GameInstanceState->GameInstanceIndex, DestroyedObjectState.ObjectIndex))
+		if (FNetProfilerObjectInstance* ObjectInstance = NetProfilerProvider.EditObject(GameInstanceState->GameInstanceIndex, DestroyedObjectState.ObjectIndex))
 		{
 			// Update object data
 			ObjectInstance->LifeTime.End = GetLastTimestamp();
@@ -607,7 +610,7 @@ TSharedRef<FNetTraceAnalyzer::FNetTraceGameInstanceState> FNetTraceAnalyzer::Get
 	else
 	{
 		// Persistent GameInstance
-		Trace::FNetProfilerGameInstanceInternal& GameInstance = NetProfilerProvider.CreateGameInstance();
+		FNetProfilerGameInstanceInternal& GameInstance = NetProfilerProvider.CreateGameInstance();
 		GameInstance.Instance.GameInstanceId = GameInstanceId;
 		GameInstance.Instance.LifeTime.Begin = GetLastTimestamp();
 
@@ -625,7 +628,7 @@ void FNetTraceAnalyzer::DestroyActiveGameInstanceState(uint32 GameInstanceId)
 	if (TSharedRef<FNetTraceAnalyzer::FNetTraceGameInstanceState>* FoundState = ActiveGameInstances.Find(GameInstanceId))
 	{
 		// Mark as closed
-		if (Trace::FNetProfilerGameInstanceInternal* GameInstance = NetProfilerProvider.EditGameInstance((*FoundState)->GameInstanceIndex))
+		if (FNetProfilerGameInstanceInternal* GameInstance = NetProfilerProvider.EditGameInstance((*FoundState)->GameInstanceIndex))
 		{
 			GameInstance->Instance.LifeTime.End = GetLastTimestamp();
 		}
@@ -645,3 +648,5 @@ FNetTraceAnalyzer::FNetTraceConnectionState* FNetTraceAnalyzer::GetActiveConnect
 		
 	return nullptr;
 }
+
+} // namespace TraceServices
diff --git a/Engine/Source/Developer/TraceServices/Private/Analyzers/NetTraceAnalyzer.h b/Engine/Source/Developer/TraceServices/Private/Analyzers/NetTraceAnalyzer.h
--- a/Engine/Source/Developer/TraceServices/Private/Analyzers/NetTraceAnalyzer.h
+++ b/Engine/Source/Developer/TraceServices/Private/Analyzers/NetTraceAnalyzer.h
@@ -6,16 +6,16 @@
 #include "Containers/UnrealString.h"
 #include "Model/NetProfilerProvider.h"
 
-namespace Trace
+namespace TraceServices
 {
-	class IAnalysisSession;
-}
+
+class IAnalysisSession;
 
 class FNetTraceAnalyzer
 	: public Trace::IAnalyzer
 {
 public:
-	FNetTraceAnalyzer(Trace::IAnalysisSession& Session, Trace::FNetProfilerProvider& NetProfilerProvider);
+	FNetTraceAnalyzer(IAnalysisSession& Session, FNetProfilerProvider& NetProfilerProvider);
 	virtual void OnAnalysisBegin(const FOnAnalysisContext& Context) override;
 	virtual bool OnEvent(uint16 RouteId, EStyle Style, const FOnEventContext& Context) override;
 	virtual void OnAnalysisEnd() override;
@@ -45,8 +45,8 @@ private:
 		BunchHeaderEvent = 3,
 	};
 
-	Trace::IAnalysisSession& Session;
-	Trace::FNetProfilerProvider& NetProfilerProvider;
+	IAnalysisSession& Session;
+	FNetProfilerProvider& NetProfilerProvider;
 	uint32 NetTraceVersion;
 	uint32 NetTraceReporterVersion;
 	uint32 BunchHeaderNameIndex;
@@ -58,7 +58,7 @@ private:
 
 	struct FBunchInfo
 	{
-		Trace::FNetProfilerBunchInfo BunchInfo;
+		FNetProfilerBunchInfo BunchInfo;
 		uint32 BunchBits;
 		uint32 HeaderBits;
 		int32 FirstBunchEventIndex;
@@ -72,13 +72,13 @@ private:
 		uint32 ConnectionIndex;
 
 		// Current packet data
-		uint32 CurrentPacketStartIndex[Trace::ENetProfilerConnectionMode::Count];
-		uint32 CurrentPacketBitOffset[Trace::ENetProfilerConnectionMode::Count];
+		uint32 CurrentPacketStartIndex[ENetProfilerConnectionMode::Count];
+		uint32 CurrentPacketBitOffset[ENetProfilerConnectionMode::Count];
 
 		// Current bunch data
-		TArray<Trace::FNetProfilerContentEvent> BunchEvents[Trace::ENetProfilerConnectionMode::Count];
+		TArray<FNetProfilerContentEvent> BunchEvents[ENetProfilerConnectionMode::Count];
 
-		TArray<FBunchInfo> BunchInfos[Trace::ENetProfilerConnectionMode::Count];
+		TArray<FBunchInfo> BunchInfos[ENetProfilerConnectionMode::Count];
 	};
 
 	struct FNetTraceActiveObjectState
@@ -105,9 +105,9 @@ private:
 	void DestroyActiveGameInstanceState(uint32 GameInstanceId);
 
 	FNetTraceConnectionState* GetActiveConnectionState(uint32 GameInstanceId, uint32 ConnectionId);
-	Trace::FNetProfilerTimeStamp GetLastTimestamp() const { return LastTimeStamp; }
+	FNetProfilerTimeStamp GetLastTimestamp() const { return LastTimeStamp; }
 
-	void FlushPacketEvents(FNetTraceConnectionState& ConnectionState, Trace::FNetProfilerConnectionData& ConnectionData, const Trace::ENetProfilerConnectionMode ConnectionMode);
+	void FlushPacketEvents(FNetTraceConnectionState& ConnectionState, FNetProfilerConnectionData& ConnectionData, const ENetProfilerConnectionMode ConnectionMode);
 
 	void HandlePacketEvent(const FOnEventContext& Context, const FEventData& EventData);	
 	void HandlePacketContentEvent(const FOnEventContext& Context, const FEventData& EventData);
@@ -117,11 +117,13 @@ private:
 	void HandleObjectCreatedEvent(const FOnEventContext& Context, const FEventData& EventData);
 	void HandleObjectDestroyedEvent(const FOnEventContext& Context, const FEventData& EventData);
 
-	void AddEvent(TPagedArray<Trace::FNetProfilerContentEvent>& Events, const Trace::FNetProfilerContentEvent& Event, uint32 Offset, uint32 LevelOffset);
-	void AddEvent(TPagedArray<Trace::FNetProfilerContentEvent>& Events, uint32 StartPos, uint32 EndPos, uint32 Level, uint32 NameIndex, Trace::FNetProfilerBunchInfo Info);
+	void AddEvent(TPagedArray<FNetProfilerContentEvent>& Events, const FNetProfilerContentEvent& Event, uint32 Offset, uint32 LevelOffset);
+	void AddEvent(TPagedArray<FNetProfilerContentEvent>& Events, uint32 StartPos, uint32 EndPos, uint32 Level, uint32 NameIndex, FNetProfilerBunchInfo Info);
 
 private:
 
 	TMap<uint32, TSharedRef<FNetTraceGameInstanceState>> ActiveGameInstances;
-	Trace::FNetProfilerTimeStamp LastTimeStamp;
+	FNetProfilerTimeStamp LastTimeStamp;
 };
+
+} // namespace TraceServices
diff --git a/Engine/Source/Developer/TraceServices/Private/Analyzers/PlatformFileTraceAnalysis.cpp b/Engine/Source/Developer/TraceServices/Private/Analyzers/PlatformFileTraceAnalysis.cpp
--- a/Engine/Source/Developer/TraceServices/Private/Analyzers/PlatformFileTraceAnalysis.cpp
+++ b/Engine/Source/Developer/TraceServices/Private/Analyzers/PlatformFileTraceAnalysis.cpp
@@ -12,7 +12,10 @@
 #define PLATFORMFILETRACE_WARNING(x)
 #endif
 
-FPlatformFileTraceAnalyzer::FPlatformFileTraceAnalyzer(Trace::IAnalysisSession& InSession, Trace::FFileActivityProvider& InFileActivityProvider)
+namespace TraceServices
+{
+
+FPlatformFileTraceAnalyzer::FPlatformFileTraceAnalyzer(IAnalysisSession& InSession, FFileActivityProvider& InFileActivityProvider)
 	: Session(InSession)
 	, FileActivityProvider(InFileActivityProvider)
 {
@@ -35,7 +38,7 @@ void FPlatformFileTraceAnalyzer::OnAnalysisBegin(const FOnAnalysisContext& Conte
 
 bool FPlatformFileTraceAnalyzer::OnEvent(uint16 RouteId, EStyle Style, const FOnEventContext& Context)
 {
-	Trace::FAnalysisSessionEditScope _(Session);
+	FAnalysisSessionEditScope _(Session);
 
 	const auto& EventData = Context.EventData;
 	switch (RouteId)
@@ -54,7 +57,7 @@ bool FPlatformFileTraceAnalyzer::OnEvent(uint16 RouteId, EStyle Style, const FOn
 		const TCHAR* FileName = reinterpret_cast<const TCHAR*>(EventData.GetAttachment());
 		uint32 FileIndex = FileActivityProvider.GetFileIndex(FileName);
 		FPendingActivity& Open = PendingOpenMap.Add(ThreadId);
-		Open.ActivityIndex = FileActivityProvider.BeginActivity(FileIndex, Trace::FileActivityType_Open, ThreadId, 0, 0, Time);
+		Open.ActivityIndex = FileActivityProvider.BeginActivity(FileIndex, FileActivityType_Open, ThreadId, 0, 0, Time);
 		Open.FileIndex = FileIndex;
 		break;
 	}
@@ -101,7 +104,7 @@ bool FPlatformFileTraceAnalyzer::OnEvent(uint16 RouteId, EStyle Style, const FOn
 		{
 			OpenFilesMap.Remove(FileHandle);
 			FPendingActivity& Close = PendingCloseMap.Add(ThreadId);
-			Close.ActivityIndex = FileActivityProvider.BeginActivity(*FindFileIndex, Trace::FileActivityType_Close, ThreadId, 0, 0, Time);
+			Close.ActivityIndex = FileActivityProvider.BeginActivity(*FindFileIndex, FileActivityType_Close, ThreadId, 0, 0, Time);
 			Close.FileIndex = *FindFileIndex;
 		}
 		else
@@ -147,7 +150,7 @@ bool FPlatformFileTraceAnalyzer::OnEvent(uint16 RouteId, EStyle Style, const FOn
 			FileIndex = FileActivityProvider.GetUnknownFileIndex();
 			OpenFilesMap.Add(FileHandle, FileIndex);
 		}
-		uint64 ReadIndex = FileActivityProvider.BeginActivity(FileIndex, Trace::FileActivityType_Read, ThreadId, Offset, Size, Time);
+		uint64 ReadIndex = FileActivityProvider.BeginActivity(FileIndex, FileActivityType_Read, ThreadId, Offset, Size, Time);
 		FPendingActivity& Read = ActiveReadsMap.Add(ReadHandle);
 		Read.FileIndex = FileIndex;
 		Read.ActivityIndex = ReadIndex;
@@ -192,7 +195,7 @@ bool FPlatformFileTraceAnalyzer::OnEvent(uint16 RouteId, EStyle Style, const FOn
 			FileIndex = FileActivityProvider.GetUnknownFileIndex();
 			OpenFilesMap.Add(FileHandle, FileIndex);
 		}
-		uint64 WriteIndex = FileActivityProvider.BeginActivity(FileIndex, Trace::FileActivityType_Write, ThreadId, Offset, Size, Time);
+		uint64 WriteIndex = FileActivityProvider.BeginActivity(FileIndex, FileActivityType_Write, ThreadId, Offset, Size, Time);
 		FPendingActivity& Write = ActiveWritesMap.Add(WriteHandle);
 		Write.FileIndex = FileIndex;
 		Write.ActivityIndex = WriteIndex;
@@ -224,3 +227,5 @@ bool FPlatformFileTraceAnalyzer::OnEvent(uint16 RouteId, EStyle Style, const FOn
 
 #undef PLATFORMFILETRACE_WARNING
 #undef DEBUG_PLATFORMFILETRACE
+
+} // namespace TraceServices
diff --git a/Engine/Source/Developer/TraceServices/Private/Analyzers/PlatformFileTraceAnalysis.h b/Engine/Source/Developer/TraceServices/Private/Analyzers/PlatformFileTraceAnalysis.h
--- a/Engine/Source/Developer/TraceServices/Private/Analyzers/PlatformFileTraceAnalysis.h
+++ b/Engine/Source/Developer/TraceServices/Private/Analyzers/PlatformFileTraceAnalysis.h
@@ -6,17 +6,17 @@
 #include "Containers/Map.h"
 #include "TraceServices/Model/AnalysisSession.h"
 
-namespace Trace
+namespace TraceServices
 {
-	class FTimeline;
-	class FFileActivityProvider;
-}
+
+class FTimeline;
+class FFileActivityProvider;
 
 class FPlatformFileTraceAnalyzer
 	: public Trace::IAnalyzer
 {
 public:
-	FPlatformFileTraceAnalyzer(Trace::IAnalysisSession& Session, Trace::FFileActivityProvider& FileActivityProvider);
+	FPlatformFileTraceAnalyzer(IAnalysisSession& Session, FFileActivityProvider& FileActivityProvider);
 	virtual void OnAnalysisBegin(const FOnAnalysisContext& Context) override;
 	virtual bool OnEvent(uint16 RouteId, EStyle Style, const FOnEventContext& Context) override;
 
@@ -42,11 +42,13 @@ private:
 		uint32 ThreadId;
 	};
 
-	Trace::IAnalysisSession& Session;
-	Trace::FFileActivityProvider& FileActivityProvider;
+	IAnalysisSession& Session;
+	FFileActivityProvider& FileActivityProvider;
 	TMap<uint64, uint32> OpenFilesMap;
 	TMap<uint32, FPendingActivity> PendingOpenMap;
 	TMap<uint32, FPendingActivity> PendingCloseMap;
 	TMap<uint64, FPendingActivity> ActiveReadsMap;
 	TMap<uint64, FPendingActivity> ActiveWritesMap;
 };
+
+} // namespace TraceServices
diff --git a/Engine/Source/Developer/TraceServices/Private/Analyzers/StatsTraceAnalysis.cpp b/Engine/Source/Developer/TraceServices/Private/Analyzers/StatsTraceAnalysis.cpp
--- a/Engine/Source/Developer/TraceServices/Private/Analyzers/StatsTraceAnalysis.cpp
+++ b/Engine/Source/Developer/TraceServices/Private/Analyzers/StatsTraceAnalysis.cpp
@@ -4,7 +4,10 @@
 #include "Common/Utils.h"
 #include "TraceServices/Model/Counters.h"
 
-FStatsAnalyzer::FStatsAnalyzer(Trace::IAnalysisSession& InSession, Trace::ICounterProvider& InCounterProvider)
+namespace TraceServices
+{
+
+FStatsAnalyzer::FStatsAnalyzer(IAnalysisSession& InSession, ICounterProvider& InCounterProvider)
 	: Session(InSession)
 	, CounterProvider(InCounterProvider)
 {
@@ -21,7 +24,7 @@ void FStatsAnalyzer::OnAnalysisBegin(const FOnAnalysisContext& Context)
 
 bool FStatsAnalyzer::OnEvent(uint16 RouteId, EStyle Style, const FOnEventContext& Context)
 {
-	Trace::FAnalysisSessionEditScope _(Session);
+	FAnalysisSessionEditScope _(Session);
 
 	const auto& EventData = Context.EventData;
 	switch (RouteId)
@@ -29,7 +32,7 @@ bool FStatsAnalyzer::OnEvent(uint16 RouteId, EStyle Style, const FOnEventContext
 	case RouteId_Spec:
 	{
 		uint32 StatId = EventData.GetValue<uint32>("Id");
-		Trace::IEditableCounter* Counter = CountersMap.FindRef(StatId);
+		IEditableCounter* Counter = CountersMap.FindRef(StatId);
 		if (!Counter)
 		{
 			Counter = CounterProvider.CreateCounter();
@@ -49,10 +52,10 @@ bool FStatsAnalyzer::OnEvent(uint16 RouteId, EStyle Style, const FOnEventContext
 		}
 
 
-		Trace::ECounterDisplayHint DisplayHint = Trace::CounterDisplayHint_None;
+		ECounterDisplayHint DisplayHint = CounterDisplayHint_None;
 		if (EventData.GetValue<bool>("IsMemory"))
 		{
-			DisplayHint = Trace::CounterDisplayHint_Memory;
+			DisplayHint = CounterDisplayHint_Memory;
 		}
 		Counter->SetName(Session.StoreString(*Name));
 		Counter->SetDescription(Session.StoreString(*Description));
@@ -81,7 +84,7 @@ bool FStatsAnalyzer::OnEvent(uint16 RouteId, EStyle Style, const FOnEventContext
 
 			uint64 DecodedIdAndOp = FTraceAnalyzerUtils::Decode7bit(BufferPtr);
 			uint32 StatId = DecodedIdAndOp >> 3;
-			Trace::IEditableCounter* Counter = CountersMap.FindRef(StatId);
+			IEditableCounter* Counter = CountersMap.FindRef(StatId);
 			if (!Counter)
 			{
 				Counter = CounterProvider.CreateCounter();
@@ -156,3 +159,5 @@ TSharedRef<FStatsAnalyzer::FThreadState> FStatsAnalyzer::GetThreadState(uint32 T
 		return ThreadStatesMap[ThreadId];
 	}
 }
+
+} // namespace TraceServices
diff --git a/Engine/Source/Developer/TraceServices/Private/Analyzers/StatsTraceAnalysis.h b/Engine/Source/Developer/TraceServices/Private/Analyzers/StatsTraceAnalysis.h
--- a/Engine/Source/Developer/TraceServices/Private/Analyzers/StatsTraceAnalysis.h
+++ b/Engine/Source/Developer/TraceServices/Private/Analyzers/StatsTraceAnalysis.h
@@ -5,18 +5,18 @@
 #include "Trace/Analyzer.h"
 #include "Templates/SharedPointer.h"
 
-namespace Trace
+namespace TraceServices
 {
-	class IAnalysisSession;
-	class IEditableCounter;
-	class ICounterProvider;
-}
+
+class IAnalysisSession;
+class IEditableCounter;
+class ICounterProvider;
 
 class FStatsAnalyzer
 	: public Trace::IAnalyzer
 {
 public:
-	FStatsAnalyzer(Trace::IAnalysisSession& Session, Trace::ICounterProvider& CounterProvider);
+	FStatsAnalyzer(IAnalysisSession& Session, ICounterProvider& CounterProvider);
 	virtual void OnAnalysisBegin(const FOnAnalysisContext& Context) override;
 	virtual bool OnEvent(uint16 RouteId, EStyle Style, const FOnEventContext& Context) override;
 
@@ -34,8 +34,10 @@ private:
 
 	TSharedRef<FThreadState> GetThreadState(uint32 ThreadId);
 
-	Trace::IAnalysisSession& Session;
-	Trace::ICounterProvider& CounterProvider;
-	TMap<uint32, Trace::IEditableCounter*> CountersMap;
+	IAnalysisSession& Session;
+	ICounterProvider& CounterProvider;
+	TMap<uint32, IEditableCounter*> CountersMap;
 	TMap<uint32, TSharedRef<FThreadState>> ThreadStatesMap;
 };
+
+} // namespace TraceServices
diff --git a/Engine/Source/Developer/TraceServices/Private/Common/FormatArgs.cpp b/Engine/Source/Developer/TraceServices/Private/Common/FormatArgs.cpp
--- a/Engine/Source/Developer/TraceServices/Private/Common/FormatArgs.cpp
+++ b/Engine/Source/Developer/TraceServices/Private/Common/FormatArgs.cpp
@@ -4,7 +4,7 @@
 #include "ProfilingDebugging/FormatArgsTrace.h"
 #include "CoreMinimal.h"
 
-namespace Trace
+namespace TraceServices
 {
 
 const TCHAR* FFormatArgsHelper::ExtractNextFormatArg(const TCHAR* FormatString, FFormatArgSpec& Spec)
@@ -376,4 +376,4 @@ void FFormatArgsHelper::Format(TCHAR* Out, uint64 MaxOut, TCHAR* Temp, uint64 Ma
 	}
 }
 
-}
\ No newline at end of file
+} // namespace TraceServices
diff --git a/Engine/Source/Developer/TraceServices/Private/Common/FormatArgs.h b/Engine/Source/Developer/TraceServices/Private/Common/FormatArgs.h
--- a/Engine/Source/Developer/TraceServices/Private/Common/FormatArgs.h
+++ b/Engine/Source/Developer/TraceServices/Private/Common/FormatArgs.h
@@ -4,7 +4,7 @@
 
 #include "CoreTypes.h"
 
-namespace Trace
+namespace TraceServices
 {
 
 struct FFormatArgsHelper
@@ -40,4 +40,4 @@ private:
 	static int32 FormatArgument(TCHAR* Out, uint64 MaxOut, TCHAR* Temp, uint64 MaxTemp, const FFormatArgSpec& ArgSpec, FFormatArgsStreamContext& ArgStream);
 };
 
-}
\ No newline at end of file
+} // namespace TraceServices
diff --git a/Engine/Source/Developer/TraceServices/Private/Common/SlabAllocator.h b/Engine/Source/Developer/TraceServices/Private/Common/SlabAllocator.h
--- a/Engine/Source/Developer/TraceServices/Private/Common/SlabAllocator.h
+++ b/Engine/Source/Developer/TraceServices/Private/Common/SlabAllocator.h
@@ -6,8 +6,11 @@
 #include "Containers/Array.h"
 #include "TraceServices/Containers/Allocators.h"
 
+namespace TraceServices
+{
+
 class FSlabAllocator
-	: public Trace::ILinearAllocator
+	: public ILinearAllocator
 {
 public:
 	FSlabAllocator(uint64 InSlabSize)
@@ -47,3 +50,5 @@ private:
 	uint64 CurrentSlabAllocatedSize = 0;
 	uint64 TotalAllocatedSize = 0;
 };
+
+} // namespace TraceServices
diff --git a/Engine/Source/Developer/TraceServices/Private/Common/Utils.cpp b/Engine/Source/Developer/TraceServices/Private/Common/Utils.cpp
--- a/Engine/Source/Developer/TraceServices/Private/Common/Utils.cpp
+++ b/Engine/Source/Developer/TraceServices/Private/Common/Utils.cpp
@@ -5,15 +5,15 @@
 #include "Templates/SharedPointer.h"
 #include "HAL/FileManager.h"
 
-namespace Trace
+namespace TraceServices
 {
 
-void Table2Csv(const Trace::IUntypedTable& Table, const TCHAR* Filename)
+void Table2Csv(const IUntypedTable& Table, const TCHAR* Filename)
 {
 	TSharedPtr<FArchive> OutputFile = MakeShareable(IFileManager::Get().CreateFileWriter(Filename));
 	check(OutputFile);
 	FString Header;
-	const Trace::ITableLayout& Layout = Table.GetLayout();
+	const ITableLayout& Layout = Table.GetLayout();
 	int32 ColumnCount = Layout.GetColumnCount();
 	for (int32 ColumnIndex = 0; ColumnIndex < ColumnCount; ++ColumnIndex)
 	{
@@ -29,7 +29,7 @@ void Table2Csv(const Trace::IUntypedTable& Table, const TCHAR* Filename)
 	}
 	auto AnsiHeader = StringCast<ANSICHAR>(*Header);
 	OutputFile->Serialize((void*)AnsiHeader.Get(), AnsiHeader.Length());
-	TUniquePtr<Trace::IUntypedTableReader> TableReader(Table.CreateReader());
+	TUniquePtr<IUntypedTableReader> TableReader(Table.CreateReader());
 	for (; TableReader->IsValid(); TableReader->NextRow())
 	{
 		FString Line;
@@ -37,19 +37,19 @@ void Table2Csv(const Trace::IUntypedTable& Table, const TCHAR* Filename)
 		{
 			switch (Layout.GetColumnType(ColumnIndex))
 			{
-			case Trace::TableColumnType_Bool:
+			case TableColumnType_Bool:
 				Line += TableReader->GetValueBool(ColumnIndex) ? "true" : "false";
 				break;
-			case Trace::TableColumnType_Int:
+			case TableColumnType_Int:
 				Line += FString::Printf(TEXT("%lld"), TableReader->GetValueInt(ColumnIndex));
 				break;
-			case Trace::TableColumnType_Float:
+			case TableColumnType_Float:
 				Line += FString::Printf(TEXT("%f"), TableReader->GetValueFloat(ColumnIndex));
 				break;
-			case Trace::TableColumnType_Double:
+			case TableColumnType_Double:
 				Line += FString::Printf(TEXT("%f"), TableReader->GetValueDouble(ColumnIndex));
 				break;
-			case Trace::TableColumnType_CString:
+			case TableColumnType_CString:
 				FString ValueString = TableReader->GetValueCString(ColumnIndex);
 				ValueString.ReplaceInline(TEXT(","), TEXT(" "));
 				Line += ValueString;
@@ -70,4 +70,4 @@ void Table2Csv(const Trace::IUntypedTable& Table, const TCHAR* Filename)
 	OutputFile->Close();
 }
 
-}
+} // namespace TraceServices
diff --git a/Engine/Source/Developer/TraceServices/Private/Common/Utils.h b/Engine/Source/Developer/TraceServices/Private/Common/Utils.h
--- a/Engine/Source/Developer/TraceServices/Private/Common/Utils.h
+++ b/Engine/Source/Developer/TraceServices/Private/Common/Utils.h
@@ -6,8 +6,8 @@
 #include "Trace/Analyzer.h"
 #include "TraceServices/Utils.h"
 
-namespace FTraceAnalyzerUtils
-{
+namespace TraceServices {
+namespace FTraceAnalyzerUtils {
 
 inline uint64 Decode7bit(const uint8*& BufferPtr)
 {
@@ -44,3 +44,4 @@ inline uint32 GetThreadIdField(
 }
 
 } // namespace FTraceAnalyzerUtils
+} // namespace TraceServices
diff --git a/Engine/Source/Developer/TraceServices/Private/Model/CallstacksProvider.cpp b/Engine/Source/Developer/TraceServices/Private/Model/CallstacksProvider.cpp
--- a/Engine/Source/Developer/TraceServices/Private/Model/CallstacksProvider.cpp
+++ b/Engine/Source/Developer/TraceServices/Private/Model/CallstacksProvider.cpp
@@ -5,7 +5,8 @@
 #include "TraceServices/Model/AnalysisSession.h"
 #include "Algo/Unique.h"
 
-namespace Trace {
+namespace TraceServices
+{
 
 /////////////////////////////////////////////////////////////////////
 #ifdef TRACE_CALLSTACK_STATS
@@ -108,4 +109,4 @@ FName FCallstacksProvider::GetName() const
 	return Name;
 }
 
-} // namespace Trace
+} // namespace TraceServices
diff --git a/Engine/Source/Developer/TraceServices/Private/Model/CallstacksProvider.h b/Engine/Source/Developer/TraceServices/Private/Model/CallstacksProvider.h
--- a/Engine/Source/Developer/TraceServices/Private/Model/CallstacksProvider.h
+++ b/Engine/Source/Developer/TraceServices/Private/Model/CallstacksProvider.h
@@ -6,7 +6,8 @@
 #include "HAL/CriticalSection.h"
 #include "Common/PagedArray.h"
 
-namespace Trace {
+namespace TraceServices
+{
 
 class IAnalysisSession;
 
@@ -45,4 +46,4 @@ private:
 	TPagedArray<FStackFrame>		Frames;
 };
 
-} // namespace Trace
+} // namespace TraceServices
diff --git a/Engine/Source/Developer/TraceServices/Private/Model/Channel.cpp b/Engine/Source/Developer/TraceServices/Private/Model/Channel.cpp
--- a/Engine/Source/Developer/TraceServices/Private/Model/Channel.cpp
+++ b/Engine/Source/Developer/TraceServices/Private/Model/Channel.cpp
@@ -3,7 +3,8 @@
 #include "UObject/NameTypes.h"
 #include "Misc/DateTime.h"
 
-namespace Trace {
+namespace TraceServices
+{
 
 const FName FChannelProvider::ProviderName("ChannelProvider");
 
@@ -65,4 +66,4 @@ FDateTime FChannelProvider::GetTimeStamp() const
 	return TimeStamp;
 }
 
-}
+} // namespace TraceServices
diff --git a/Engine/Source/Developer/TraceServices/Private/Model/Diagnostics.cpp b/Engine/Source/Developer/TraceServices/Private/Model/Diagnostics.cpp
--- a/Engine/Source/Developer/TraceServices/Private/Model/Diagnostics.cpp
+++ b/Engine/Source/Developer/TraceServices/Private/Model/Diagnostics.cpp
@@ -2,7 +2,7 @@
 #include "TraceServices/Model/Diagnostics.h"
 #include "Model/DiagnosticsPrivate.h"
 
-namespace Trace
+namespace TraceServices
 {
 
 FName FDiagnosticsProvider::ProviderName(TEXT("DiagnosticsProvider"));
@@ -36,4 +36,4 @@ const IDiagnosticsProvider& ReadDiagnosticsProvider(const IAnalysisSession& Sess
 	return *Session.ReadProvider<IDiagnosticsProvider>(FDiagnosticsProvider::ProviderName);
 }
 
-} // namespace Trace
+} // namespace TraceServices
diff --git a/Engine/Source/Developer/TraceServices/Private/Model/DiagnosticsPrivate.h b/Engine/Source/Developer/TraceServices/Private/Model/DiagnosticsPrivate.h
--- a/Engine/Source/Developer/TraceServices/Private/Model/DiagnosticsPrivate.h
+++ b/Engine/Source/Developer/TraceServices/Private/Model/DiagnosticsPrivate.h
@@ -4,7 +4,7 @@
 #include "TraceServices/AnalysisService.h"
 #include "TraceServices/Model/Diagnostics.h"
 
-namespace Trace
+namespace TraceServices
 {
 
 class FDiagnosticsProvider : public IDiagnosticsProvider
@@ -26,4 +26,4 @@ private:
 	bool bIsSessionInfoAvailable = false;
 };
 
-} // namespace Trace
+} // namespace TraceServices
diff --git a/Engine/Source/Developer/TraceServices/Private/Model/LoadTimeProfiler.cpp b/Engine/Source/Developer/TraceServices/Private/Model/LoadTimeProfiler.cpp
--- a/Engine/Source/Developer/TraceServices/Private/Model/LoadTimeProfiler.cpp
+++ b/Engine/Source/Developer/TraceServices/Private/Model/LoadTimeProfiler.cpp
@@ -5,7 +5,7 @@
 #include "AnalysisServicePrivate.h"
 #include "Common/TimelineStatistics.h"
 
-namespace Trace
+namespace TraceServices
 {
 
 FLoadTimeProfilerProvider::FLoaderFrameCounter::FLoaderFrameCounter(ELoaderFrameCounterType InType, const TPagedArray<FLoaderFrame>& InFrames)
@@ -228,7 +228,7 @@ ITable<FLoadTimeProfilerAggregatedStats>* FLoadTimeProfilerProvider::CreateObjec
 	{
 		return Event.Export ? Event.Export->Class : nullptr;
 	};
-	TMap<const Trace::FClassInfo*, FAggregatedTimingStats> Aggregation;
+	TMap<const FClassInfo*, FAggregatedTimingStats> Aggregation;
 	FTimelineStatistics::CreateAggregation(Timelines, BucketMapper, IntervalStart, IntervalEnd, Aggregation);
 	TTable<FLoadTimeProfilerAggregatedStats>* Table = new TTable<FLoadTimeProfilerAggregatedStats>(AggregatedStatsTableLayout);
 	for (const auto& KV : Aggregation)
@@ -305,7 +305,7 @@ ITable<FPackagesTableRow>* FLoadTimeProfilerProvider::CreatePackageDetailsTable(
 					Row->AsyncLoadingThreadTime += EndTime - StartTime; // TODO: Should be exclusive time
 				}
 			}
-			return Trace::EEventEnumerate::Continue;
+			return EEventEnumerate::Continue;
 		});
 	}
 
@@ -366,13 +366,13 @@ ITable<FExportsTableRow>* FLoadTimeProfilerProvider::CreateExportDetailsTable(do
 					Row->AsyncLoadingThreadTime += EndTime - StartTime; // TODO: Should be exclusive time
 				}
 			}
-			return Trace::EEventEnumerate::Continue;
+			return EEventEnumerate::Continue;
 		});
 	}
 	return Table;
 }
 
-const Trace::FClassInfo& FLoadTimeProfilerProvider::AddClassInfo(const TCHAR* ClassName)
+const FClassInfo& FLoadTimeProfilerProvider::AddClassInfo(const TCHAR* ClassName)
 {
 	Session.WriteAccessCheck();
 
@@ -381,7 +381,7 @@ const Trace::FClassInfo& FLoadTimeProfilerProvider::AddClassInfo(const TCHAR* Cl
 	return ClassInfo;
 }
 
-Trace::FLoadRequest& FLoadTimeProfilerProvider::CreateRequest()
+FLoadRequest& FLoadTimeProfilerProvider::CreateRequest()
 {
 	Session.WriteAccessCheck();
 
@@ -389,7 +389,7 @@ Trace::FLoadRequest& FLoadTimeProfilerProvider::CreateRequest()
 	return RequestInfo;
 }
 
-Trace::FPackageInfo& FLoadTimeProfilerProvider::EditPackageInfo(const TCHAR* PackageName)
+FPackageInfo& FLoadTimeProfilerProvider::EditPackageInfo(const TCHAR* PackageName)
 {
 	Session.WriteAccessCheck();
 
@@ -514,7 +514,7 @@ void FLoadTimeProfilerProvider::EndIoDispatcherBatch(uint64 BatchHandle, double
 	ActiveIoDispatcherBatchesCounter->AddValue(Time, int64(-1));
 }
 
-Trace::FPackageExportInfo& FLoadTimeProfilerProvider::CreateExport()
+FPackageExportInfo& FLoadTimeProfilerProvider::CreateExport()
 {
 	Session.WriteAccessCheck();
 
@@ -600,4 +600,4 @@ const TCHAR* GetLoadTimeProfilerObjectEventTypeString(ELoadTimeProfilerObjectEve
 	return TEXT("[invalid]");
 }
 
-}
\ No newline at end of file
+} // namespace TraceServices
diff --git a/Engine/Source/Developer/TraceServices/Private/Model/Memory.cpp b/Engine/Source/Developer/TraceServices/Private/Model/Memory.cpp
--- a/Engine/Source/Developer/TraceServices/Private/Model/Memory.cpp
+++ b/Engine/Source/Developer/TraceServices/Private/Model/Memory.cpp
@@ -2,7 +2,7 @@
 #include "TraceServices/Model/Memory.h"
 #include "Model/MemoryPrivate.h"
 
-namespace Trace
+namespace TraceServices
 {
 
 FName FMemoryProvider::ProviderName(TEXT("MemoryProvider"));
@@ -16,7 +16,7 @@ FMemoryProvider::FMemoryProvider(IAnalysisSession& InSession)
 
 void FMemoryProvider::AddEventSpec(FMemoryTagId TagId, const TCHAR* Name, FMemoryTagId ParentTagId)
 {
-	if (TagId == FMemoryTag::InvalidTagId)
+	if (TagId == FMemoryTagInfo::InvalidTagId)
 	{
 		// invalid tag id
 		return;
@@ -28,13 +28,13 @@ void FMemoryProvider::AddEventSpec(FMemoryTagId TagId, const TCHAR* Name, FMemor
 		return;
 	}
 
-	TagDescs.Add(TagId, FMemoryTag{ TagId, FString(Name), ParentTagId, 0 });
+	TagDescs.Add(TagId, FMemoryTagInfo{ TagId, FString(Name), ParentTagId, 0 });
 	++TagsSerial;
 }
 
 void FMemoryProvider::AddTrackerSpec(FMemoryTrackerId TrackerId, const TCHAR* Name)
 {
-	if (TrackerId < 0 || TrackerId > FMemoryTracker::MaxValidTrackerId)
+	if (TrackerId < 0 || TrackerId > FMemoryTrackerInfo::MaxValidTrackerId)
 	{
 		// invalid tracker id
 		return;
@@ -46,7 +46,7 @@ void FMemoryProvider::AddTrackerSpec(FMemoryTrackerId TrackerId, const TCHAR* Na
 		return;
 	}
 
-	TrackerDescs.Add(TrackerId, FMemoryTracker{ TrackerId, FString(Name) });
+	TrackerDescs.Add(TrackerId, FMemoryTrackerInfo{ TrackerId, FString(Name) });
 }
 
 void FMemoryProvider::AddTagSnapshot(FMemoryTrackerId TrackerId, double Time, const int64* Tags, const FMemoryTagSample* Values, uint32 TagCount)
@@ -63,7 +63,7 @@ void FMemoryProvider::AddTagSnapshot(FMemoryTrackerId TrackerId, double Time, co
 	}
 
 	FTrackerData& TrackerData = Trackers[TrackerId];
-	FMemoryTracker& Tracker = TrackerDescs[TrackerId];
+	FMemoryTrackerInfo& Tracker = TrackerDescs[TrackerId];
 
 	TrackerData.SampleTimes.Push(Time);
 	const int32 SampleCount = TrackerData.SampleTimes.Num();
@@ -97,7 +97,7 @@ void FMemoryProvider::AddTagSnapshot(FMemoryTrackerId TrackerId, double Time, co
 
 		if (!TagSamples->TagPtr)
 		{
-			// Cache pointer to FMemoryTag to avoid further lookups for this tag.
+			// Cache pointer to FMemoryTagInfo to avoid further lookups for this tag.
 			TagSamples->TagPtr = TagDescs.Find(TagId);
 		}
 
@@ -123,7 +123,7 @@ uint32 FMemoryProvider::GetTagCount() const
 	return TagDescs.Num();
 }
 
-void FMemoryProvider::EnumerateTags(TFunctionRef<void(const FMemoryTag&)> Callback) const
+void FMemoryProvider::EnumerateTags(TFunctionRef<void(const FMemoryTagInfo&)> Callback) const
 {
 	for (auto& Tag : TagDescs)
 	{
@@ -131,7 +131,7 @@ void FMemoryProvider::EnumerateTags(TFunctionRef<void(const FMemoryTag&)> Callba
 	}
 }
 
-const FMemoryTag* FMemoryProvider::GetTag(FMemoryTagId TagId) const
+const FMemoryTagInfo* FMemoryProvider::GetTag(FMemoryTagId TagId) const
 {
 	return TagDescs.Find(TagId);
 }
@@ -141,7 +141,7 @@ uint32 FMemoryProvider::GetTrackerCount() const
 	return TrackerDescs.Num();
 }
 
-void FMemoryProvider::EnumerateTrackers(TFunctionRef<void(const FMemoryTracker&)> Callback) const
+void FMemoryProvider::EnumerateTrackers(TFunctionRef<void(const FMemoryTrackerInfo&)> Callback) const
 {
 	for (auto& Desc : TrackerDescs)
 	{
@@ -220,4 +220,4 @@ const IMemoryProvider& ReadMemoryProvider(const IAnalysisSession& Session)
 	return *Session.ReadProvider<IMemoryProvider>(FMemoryProvider::ProviderName);
 }
 
-} // namespace Trace
+} // namespace TraceServices
diff --git a/Engine/Source/Developer/TraceServices/Private/Model/MemoryPrivate.h b/Engine/Source/Developer/TraceServices/Private/Model/MemoryPrivate.h
--- a/Engine/Source/Developer/TraceServices/Private/Model/MemoryPrivate.h
+++ b/Engine/Source/Developer/TraceServices/Private/Model/MemoryPrivate.h
@@ -5,7 +5,7 @@
 #include "TraceServices/Model/Memory.h"
 #include "Common/PagedArray.h"
 
-namespace Trace
+namespace TraceServices
 {
 
 class FMemoryProvider : public IMemoryProvider
@@ -20,10 +20,10 @@ public:
 
 	virtual uint32 GetTagSerial() const override;
 	virtual uint32 GetTagCount() const override;
-	virtual void EnumerateTags(TFunctionRef<void(const FMemoryTag&)> Callback) const override;
-	virtual const FMemoryTag* GetTag(FMemoryTagId Id) const override;
+	virtual void EnumerateTags(TFunctionRef<void(const FMemoryTagInfo&)> Callback) const override;
+	virtual const FMemoryTagInfo* GetTag(FMemoryTagId Id) const override;
 	virtual uint32 GetTrackerCount() const override;
-	virtual void EnumerateTrackers(TFunctionRef<void(const FMemoryTracker&)> Callback) const override;
+	virtual void EnumerateTrackers(TFunctionRef<void(const FMemoryTrackerInfo&)> Callback) const override;
 	virtual uint64 GetTagSampleCount(FMemoryTrackerId Tracker, FMemoryTagId Tag) const override;
 	virtual void EnumerateTagSamples(
 		FMemoryTrackerId Tracker, 
@@ -47,7 +47,7 @@ private:
 		TPagedArray<FMemoryTagSample> Values;
 
 		// Cached pointer to the actual memory tag.
-		FMemoryTag* TagPtr;
+		FMemoryTagInfo* TagPtr;
 
 		FTagSampleData(ILinearAllocator& Allocator) : Values(Allocator, DefaultPageSize), TagPtr(nullptr) {}
 		FTagSampleData(const FTagSampleData& Other) : Values(Other.Values), TagPtr(Other.TagPtr) {}
@@ -68,10 +68,10 @@ private:
 	};
 
 	IAnalysisSession& Session;
-	TMap<FMemoryTagId, FMemoryTag> TagDescs;
-	TMap<FMemoryTrackerId, FMemoryTracker> TrackerDescs;
+	TMap<FMemoryTagId, FMemoryTagInfo> TagDescs;
+	TMap<FMemoryTrackerId, FMemoryTrackerInfo> TrackerDescs;
 	TArray<FTrackerData> Trackers;
 	uint32 TagsSerial;
 };
 
-} // namespace Trace
+} // namespace TraceServices
