Moved the TraceServices module into its own namespace.

diff --git a/Engine/Source/Developer/TraceServices/Private/Model/NetProfilerProvider.cpp b/Engine/Source/Developer/TraceServices/Private/Model/NetProfilerProvider.cpp
--- a/Engine/Source/Developer/TraceServices/Private/Model/NetProfilerProvider.cpp
+++ b/Engine/Source/Developer/TraceServices/Private/Model/NetProfilerProvider.cpp
@@ -5,7 +5,7 @@
 #include "Containers/ArrayView.h"
 #include "Common/StringStore.h"
 
-namespace Trace
+namespace TraceServices
 {
 
 const FName FNetProfilerProvider::ProviderName("NetProfilerProvider");
@@ -98,7 +98,7 @@ const FNetProfilerEventType* FNetProfilerProvider::GetNetProfilerEventType(uint3
 	return EventTypeIndex < (uint32)EventTypes.Num() ? &EventTypes[EventTypeIndex] : nullptr;
 }
 
-Trace::FNetProfilerGameInstanceInternal& FNetProfilerProvider::CreateGameInstance()
+FNetProfilerGameInstanceInternal& FNetProfilerProvider::CreateGameInstance()
 {
 	Session.WriteAccessCheck();
 
@@ -115,7 +115,7 @@ Trace::FNetProfilerGameInstanceInternal& FNetProfilerProvider::CreateGameInstanc
 	return GameInstance;
 }
 
-Trace::FNetProfilerGameInstanceInternal* FNetProfilerProvider::EditGameInstance(uint32 GameInstanceIndex)
+FNetProfilerGameInstanceInternal* FNetProfilerProvider::EditGameInstance(uint32 GameInstanceIndex)
 {
 	Session.WriteAccessCheck();
 
@@ -129,7 +129,7 @@ Trace::FNetProfilerGameInstanceInternal* FNetProfilerProvider::EditGameInstance(
 	}
 }
 
-Trace::FNetProfilerConnectionInternal& FNetProfilerProvider::CreateConnection(uint32 GameInstanceIndex)
+FNetProfilerConnectionInternal& FNetProfilerProvider::CreateConnection(uint32 GameInstanceIndex)
 {
 	Session.WriteAccessCheck();
 
@@ -159,21 +159,21 @@ Trace::FNetProfilerConnectionInternal& FNetProfilerProvider::CreateConnection(ui
 	return Connection;
 }
 
-Trace::FNetProfilerObjectInstance& FNetProfilerProvider::CreateObject(uint32 GameInstanceIndex)
+FNetProfilerObjectInstance& FNetProfilerProvider::CreateObject(uint32 GameInstanceIndex)
 {
 	Session.WriteAccessCheck();
 
 	FNetProfilerGameInstanceInternal* GameInstance = EditGameInstance(GameInstanceIndex);
 	check(GameInstance);
 
-	Trace::FNetProfilerObjectInstance& Object = GameInstance->Objects->PushBack();
+	FNetProfilerObjectInstance& Object = GameInstance->Objects->PushBack();
 	Object.ObjectIndex = GameInstance->Objects->Num() - 1;
 	++GameInstance->ObjectsChangeCount;
 
 	return Object;
 }
 
-Trace::FNetProfilerObjectInstance* FNetProfilerProvider::EditObject(uint32 GameInstanceIndex, uint32 ObjectIndex)
+FNetProfilerObjectInstance* FNetProfilerProvider::EditObject(uint32 GameInstanceIndex, uint32 ObjectIndex)
 {
 	Session.WriteAccessCheck();
 
@@ -191,7 +191,7 @@ Trace::FNetProfilerObjectInstance* FNetProfilerProvider::EditObject(uint32 GameI
 	}
 }
 
-Trace::FNetProfilerConnectionInternal* FNetProfilerProvider::EditConnection(uint32 ConnectionIndex)
+FNetProfilerConnectionInternal* FNetProfilerProvider::EditConnection(uint32 ConnectionIndex)
 {
 	Session.WriteAccessCheck();
 
@@ -211,7 +211,7 @@ void FNetProfilerProvider::EditPacketDeliveryStatus(uint32 ConnectionIndex, ENet
 	check(ConnectionIndex < (uint32)Connections.Num());
 
 	FNetProfilerConnectionInternal& Connection = Connections[ConnectionIndex];
-	Trace::FNetProfilerConnectionData& Data = *Connections[ConnectionIndex].Data[Mode];
+	FNetProfilerConnectionData& Data = *Connections[ConnectionIndex].Data[Mode];
 
 	// try to locate packet
 	uint32 PacketCount = Data.Packets.Num();
@@ -227,7 +227,7 @@ void FNetProfilerProvider::EditPacketDeliveryStatus(uint32 ConnectionIndex, ENet
 	}
 }
 
-Trace::FNetProfilerConnectionData& FNetProfilerProvider::EditConnectionData(uint32 ConnectionIndex, ENetProfilerConnectionMode Mode)
+FNetProfilerConnectionData& FNetProfilerProvider::EditConnectionData(uint32 ConnectionIndex, ENetProfilerConnectionMode Mode)
 {
 	check(ConnectionIndex < (uint32)Connections.Num());
 
@@ -678,4 +678,4 @@ ITable<FNetProfilerAggregatedStats>* FNetProfilerProvider::CreateAggregation(uin
 	return Table;
 }
 
-}
+} // namespace TraceServices
diff --git a/Engine/Source/Developer/TraceServices/Private/Model/NetProfilerProvider.h b/Engine/Source/Developer/TraceServices/Private/Model/NetProfilerProvider.h
--- a/Engine/Source/Developer/TraceServices/Private/Model/NetProfilerProvider.h
+++ b/Engine/Source/Developer/TraceServices/Private/Model/NetProfilerProvider.h
@@ -6,7 +6,7 @@
 #include "Common/SlabAllocator.h"
 #include "Model/Tables.h"
 
-namespace Trace
+namespace TraceServices
 {
 
 class FAnalysisSessionLock;
@@ -16,7 +16,7 @@ struct FNetProfilerGameInstanceInternal
 {
 	FNetProfilerGameInstance Instance;
 
-	TPagedArray<Trace::FNetProfilerObjectInstance>* Objects;
+	TPagedArray<FNetProfilerObjectInstance>* Objects;
 	TArray<uint32, TInlineAllocator<128>> Connections;
 	uint32 ObjectsChangeCount;
 };
@@ -59,17 +59,17 @@ public:
 
 	uint32 AddNetProfilerEventType(uint32 NameIndex, uint32 Level);
 
-	Trace::FNetProfilerGameInstanceInternal& CreateGameInstance();
-	Trace::FNetProfilerGameInstanceInternal* EditGameInstance(uint32 GameInstanceIndex);
+	FNetProfilerGameInstanceInternal& CreateGameInstance();
+	FNetProfilerGameInstanceInternal* EditGameInstance(uint32 GameInstanceIndex);
 
-	Trace::FNetProfilerConnectionInternal& CreateConnection(uint32 GameInstanceIndex);
-	Trace::FNetProfilerConnectionInternal* EditConnection(uint32 ConnectionIndex);
-	Trace::FNetProfilerConnectionData& EditConnectionData(uint32 ConnectionIndex, ENetProfilerConnectionMode Mode);
+	FNetProfilerConnectionInternal& CreateConnection(uint32 GameInstanceIndex);
+	FNetProfilerConnectionInternal* EditConnection(uint32 ConnectionIndex);
+	FNetProfilerConnectionData& EditConnectionData(uint32 ConnectionIndex, ENetProfilerConnectionMode Mode);
 
 	void EditPacketDeliveryStatus(uint32 ConnectionIndex, ENetProfilerConnectionMode Mode, uint32 SequenceNumber, ENetProfilerDeliveryStatus DeliveryStatus);
 
-	Trace::FNetProfilerObjectInstance& CreateObject(uint32 GameInstanceIndex);
-	Trace::FNetProfilerObjectInstance* EditObject(uint32 GameInstanceIndex, uint32 ObjectIndex);
+	FNetProfilerObjectInstance& CreateObject(uint32 GameInstanceIndex);
+	FNetProfilerObjectInstance* EditObject(uint32 GameInstanceIndex, uint32 ObjectIndex);
 
 	// INetProfilerProvider Interface
 
@@ -140,4 +140,4 @@ private:
 	TTableLayout<FNetProfilerAggregatedStats> AggregatedStatsTableLayout;
 };
 
-}
+} // namespace TraceServices
diff --git a/Engine/Source/Developer/TraceServices/Private/Modules/CsvProfilerModule.cpp b/Engine/Source/Developer/TraceServices/Private/Modules/CsvProfilerModule.cpp
--- a/Engine/Source/Developer/TraceServices/Private/Modules/CsvProfilerModule.cpp
+++ b/Engine/Source/Developer/TraceServices/Private/Modules/CsvProfilerModule.cpp
@@ -5,7 +5,7 @@
 #include "AnalysisServicePrivate.h"
 #include "TraceServices/ModuleService.h"
 
-namespace Trace
+namespace TraceServices
 {
 
 static const FName CsvProfilerModuleName("TraceModule_CsvProfiler");
@@ -30,7 +30,7 @@ void FCsvProfilerModule::OnAnalysisBegin(IAnalysisSession& Session)
 
 void FCsvProfilerModule::GenerateReports(const IAnalysisSession& Session, const TCHAR* CmdLine, const TCHAR* OutputDirectory)
 {
-	const ICsvProfilerProvider* CsvProfilerProvider = Trace::ReadCsvProfilerProvider(Session);
+	const ICsvProfilerProvider* CsvProfilerProvider = ReadCsvProfilerProvider(Session);
 	if (CsvProfilerProvider)
 	{
 		CsvProfilerProvider->EnumerateCaptures([CsvProfilerProvider, &OutputDirectory](const FCaptureInfo& CaptureInfo)
@@ -45,4 +45,4 @@ const ICsvProfilerProvider* ReadCsvProfilerProvider(const IAnalysisSession& Sess
 	return Session.ReadProvider<ICsvProfilerProvider>(CsvProfilerProviderName);
 }
 
-}
+} // namespace TraceServices
diff --git a/Engine/Source/Developer/TraceServices/Private/Modules/LoadTimeProfilerModule.cpp b/Engine/Source/Developer/TraceServices/Private/Modules/LoadTimeProfilerModule.cpp
--- a/Engine/Source/Developer/TraceServices/Private/Modules/LoadTimeProfilerModule.cpp
+++ b/Engine/Source/Developer/TraceServices/Private/Modules/LoadTimeProfilerModule.cpp
@@ -7,7 +7,7 @@
 #include "Model/FileActivity.h"
 #include "HAL/FileManager.h"
 
-namespace Trace
+namespace TraceServices
 {
 
 static const FName LoadTimeProfilerModuleName("TraceModule_LoadTimeProfiler");
@@ -53,7 +53,7 @@ void FLoadTimeProfilerModule::GenerateReports(const IAnalysisSession& Session, c
 {
 	double CaptureStartTime = -DBL_MAX;
 	double CaptureEndTime = DBL_MAX;
-	const IBookmarkProvider& BookmarkProvider = Trace::ReadBookmarkProvider(Session);
+	const IBookmarkProvider& BookmarkProvider = ReadBookmarkProvider(Session);
 	FString BeginCaptureBookmarkName;
 	FParse::Value(CmdLine, TEXT("-BeginCaptureBookmark="), BeginCaptureBookmarkName);
 	FString EndCaptureBookmarkName;
@@ -80,7 +80,7 @@ void FLoadTimeProfilerModule::GenerateReports(const IAnalysisSession& Session, c
 		CaptureEndTime = Session.GetDurationSeconds();
 	}
 
-	const ILoadTimeProfilerProvider* LoadTimeProfiler = Trace::ReadLoadTimeProfilerProvider(Session);
+	const ILoadTimeProfilerProvider* LoadTimeProfiler = ReadLoadTimeProfilerProvider(Session);
 	FString ReportDirectory = FString(OutputDirectory) / TEXT("LoadTimeProfiler");
 	if (LoadTimeProfiler)
 	{
@@ -90,7 +90,7 @@ void FLoadTimeProfilerModule::GenerateReports(const IAnalysisSession& Session, c
 		Table2Csv(*ExportsTable.Get(), *(ReportDirectory / TEXT("Exports.csv")));
 		Table2Csv(LoadTimeProfiler->GetRequestsTable(), *(ReportDirectory / TEXT("Requests.csv")));
 	}
-	const IFileActivityProvider* FileActivityProvider = Trace::ReadFileActivityProvider(Session);
+	const IFileActivityProvider* FileActivityProvider = ReadFileActivityProvider(Session);
 	if (FileActivityProvider)
 	{
 		Table2Csv(FileActivityProvider->GetFileActivityTable(), *(FString(ReportDirectory) / TEXT("FileActivity.csv")));
@@ -114,4 +114,4 @@ const IFileActivityProvider* ReadFileActivityProvider(const IAnalysisSession& Se
 	return Session.ReadProvider<IFileActivityProvider>(FileActivityProviderName);
 }
 
-}
+} // namespace TraceServices
diff --git a/Engine/Source/Developer/TraceServices/Private/TraceServicesModule.cpp b/Engine/Source/Developer/TraceServices/Private/TraceServicesModule.cpp
--- a/Engine/Source/Developer/TraceServices/Private/TraceServicesModule.cpp
+++ b/Engine/Source/Developer/TraceServices/Private/TraceServicesModule.cpp
@@ -18,88 +18,88 @@ class FTraceServicesModule
 	: public ITraceServicesModule
 {
 public:
-	virtual TSharedPtr<Trace::IAnalysisService> GetAnalysisService() override;
-	virtual TSharedPtr<Trace::IModuleService> GetModuleService() override;
-	virtual TSharedPtr<Trace::IAnalysisService> CreateAnalysisService() override;
-	virtual TSharedPtr<Trace::IModuleService> CreateModuleService() override;
+	virtual TSharedPtr<TraceServices::IAnalysisService> GetAnalysisService() override;
+	virtual TSharedPtr<TraceServices::IModuleService> GetModuleService() override;
+	virtual TSharedPtr<TraceServices::IAnalysisService> CreateAnalysisService() override;
+	virtual TSharedPtr<TraceServices::IModuleService> CreateModuleService() override;
 
 	virtual void StartupModule() override;
 	virtual void ShutdownModule() override;
 
 private:
-	TSharedPtr<Trace::FAnalysisService> AnalysisService;
-	TSharedPtr<Trace::FModuleService> ModuleService;
+	TSharedPtr<TraceServices::FAnalysisService> AnalysisService;
+	TSharedPtr<TraceServices::FModuleService> ModuleService;
 
-	Trace::FTimingProfilerModule TimingProfilerModule;
-	Trace::FLoadTimeProfilerModule LoadTimeProfilerModule;
-	Trace::FStatsModule StatsModule;
-	Trace::FCsvProfilerModule CsvProfilerModule;
-	Trace::FCountersModule CountersModule;
-	Trace::FNetProfilerModule NetProfilerModule;
-	Trace::FMemoryModule MemoryModule;
-	Trace::FDiagnosticsModule DiagnosticsModule;
+	TraceServices::FTimingProfilerModule TimingProfilerModule;
+	TraceServices::FLoadTimeProfilerModule LoadTimeProfilerModule;
+	TraceServices::FStatsModule StatsModule;
+	TraceServices::FCsvProfilerModule CsvProfilerModule;
+	TraceServices::FCountersModule CountersModule;
+	TraceServices::FNetProfilerModule NetProfilerModule;
+	TraceServices::FMemoryModule MemoryModule;
+	TraceServices::FDiagnosticsModule DiagnosticsModule;
 };
 
-TSharedPtr<Trace::IAnalysisService> FTraceServicesModule::GetAnalysisService()
+TSharedPtr<TraceServices::IAnalysisService> FTraceServicesModule::GetAnalysisService()
 {
 	if (!AnalysisService.IsValid())
 	{
 		GetModuleService();
-		AnalysisService = MakeShared<Trace::FAnalysisService>(*ModuleService.Get());
+		AnalysisService = MakeShared<TraceServices::FAnalysisService>(*ModuleService.Get());
 	}
 	return AnalysisService;
 }
 
-TSharedPtr<Trace::IModuleService> FTraceServicesModule::GetModuleService()
+TSharedPtr<TraceServices::IModuleService> FTraceServicesModule::GetModuleService()
 {
 	if (!ModuleService.IsValid())
 	{
-		ModuleService = MakeShared<Trace::FModuleService>();
+		ModuleService = MakeShared<TraceServices::FModuleService>();
 	}
 	return ModuleService;
 }
 
-TSharedPtr<Trace::IAnalysisService> FTraceServicesModule::CreateAnalysisService()
+TSharedPtr<TraceServices::IAnalysisService> FTraceServicesModule::CreateAnalysisService()
 {
 	checkf(!AnalysisService.IsValid(), TEXT("A AnalysisService already exists."));
 	GetModuleService();
-	AnalysisService = MakeShared<Trace::FAnalysisService>(*ModuleService.Get());
+	AnalysisService = MakeShared<TraceServices::FAnalysisService>(*ModuleService.Get());
 	return AnalysisService;
 }
 
-TSharedPtr<Trace::IModuleService> FTraceServicesModule::CreateModuleService()
+TSharedPtr<TraceServices::IModuleService> FTraceServicesModule::CreateModuleService()
 {
 	checkf(!ModuleService.IsValid(), TEXT("A ModuleService already exists."));
-	ModuleService = MakeShared<Trace::FModuleService>();
+	ModuleService = MakeShared<TraceServices::FModuleService>();
 	return ModuleService;
 }
 
 void FTraceServicesModule::StartupModule()
 {
-	IModularFeatures::Get().RegisterModularFeature(Trace::ModuleFeatureName, &TimingProfilerModule);
-	IModularFeatures::Get().RegisterModularFeature(Trace::ModuleFeatureName, &LoadTimeProfilerModule);
+	IModularFeatures::Get().RegisterModularFeature(TraceServices::ModuleFeatureName, &TimingProfilerModule);
+	IModularFeatures::Get().RegisterModularFeature(TraceServices::ModuleFeatureName, &LoadTimeProfilerModule);
 #if EXPERIMENTAL_STATSTRACE_ENABLED
-	IModularFeatures::Get().RegisterModularFeature(Trace::ModuleFeatureName, &StatsModule);
+	IModularFeatures::Get().RegisterModularFeature(TraceServices::ModuleFeatureName, &StatsModule);
 #endif
-	IModularFeatures::Get().RegisterModularFeature(Trace::ModuleFeatureName, &CsvProfilerModule);
-	IModularFeatures::Get().RegisterModularFeature(Trace::ModuleFeatureName, &CountersModule);
-	IModularFeatures::Get().RegisterModularFeature(Trace::ModuleFeatureName, &NetProfilerModule);
-	IModularFeatures::Get().RegisterModularFeature(Trace::ModuleFeatureName, &MemoryModule);
-	IModularFeatures::Get().RegisterModularFeature(Trace::ModuleFeatureName, &DiagnosticsModule);
+	IModularFeatures::Get().RegisterModularFeature(TraceServices::ModuleFeatureName, &CsvProfilerModule);
+	IModularFeatures::Get().RegisterModularFeature(TraceServices::ModuleFeatureName, &CountersModule);
+	IModularFeatures::Get().RegisterModularFeature(TraceServices::ModuleFeatureName, &NetProfilerModule);
+	IModularFeatures::Get().RegisterModularFeature(TraceServices::ModuleFeatureName, &MemoryModule);
+	IModularFeatures::Get().RegisterModularFeature(TraceServices::ModuleFeatureName, &DiagnosticsModule);
 }
 
 void FTraceServicesModule::ShutdownModule()
 {
-	IModularFeatures::Get().UnregisterModularFeature(Trace::ModuleFeatureName, &DiagnosticsModule);
-	IModularFeatures::Get().UnregisterModularFeature(Trace::ModuleFeatureName, &MemoryModule);
-	IModularFeatures::Get().UnregisterModularFeature(Trace::ModuleFeatureName, &NetProfilerModule);
-	IModularFeatures::Get().UnregisterModularFeature(Trace::ModuleFeatureName, &CountersModule);
-	IModularFeatures::Get().UnregisterModularFeature(Trace::ModuleFeatureName, &CsvProfilerModule);
+	IModularFeatures::Get().UnregisterModularFeature(TraceServices::ModuleFeatureName, &DiagnosticsModule);
+	IModularFeatures::Get().UnregisterModularFeature(TraceServices::ModuleFeatureName, &MemoryModule);
+	IModularFeatures::Get().UnregisterModularFeature(TraceServices::ModuleFeatureName, &NetProfilerModule);
+	IModularFeatures::Get().UnregisterModularFeature(TraceServices::ModuleFeatureName, &CountersModule);
+	IModularFeatures::Get().UnregisterModularFeature(TraceServices::ModuleFeatureName, &CsvProfilerModule);
 #if EXPERIMENTAL_STATSTRACE_ENABLED
-	IModularFeatures::Get().UnregisterModularFeature(Trace::ModuleFeatureName, &StatsModule);
+	IModularFeatures::Get().UnregisterModularFeature(TraceServices::ModuleFeatureName, &StatsModule);
 #endif
-	IModularFeatures::Get().UnregisterModularFeature(Trace::ModuleFeatureName, &LoadTimeProfilerModule);
-	IModularFeatures::Get().UnregisterModularFeature(Trace::ModuleFeatureName, &TimingProfilerModule);
+	IModularFeatures::Get().UnregisterModularFeature(TraceServices::ModuleFeatureName, &LoadTimeProfilerModule);
+	IModularFeatures::Get().UnregisterModularFeature(TraceServices::ModuleFeatureName, &TimingProfilerModule);
 
 	AnalysisService.Reset();
 	ModuleService.Reset();
diff --git a/Engine/Source/Developer/TraceServices/Public/Common/PagedArray.h b/Engine/Source/Developer/TraceServices/Public/Common/PagedArray.h
--- a/Engine/Source/Developer/TraceServices/Public/Common/PagedArray.h
+++ b/Engine/Source/Developer/TraceServices/Public/Common/PagedArray.h
@@ -5,6 +5,8 @@
 #include "TraceServices/Containers/Allocators.h"
 #include "Containers/Array.h"
 
+namespace TraceServices {
+
 template<typename ItemType>
 struct TPagedArrayPage
 {
@@ -235,7 +237,7 @@ public:
 	typedef InPageType PageType;
 	typedef TPagedArrayIterator<InItemType, InPageType> TIterator;
 
-	TPagedArray(Trace::ILinearAllocator& InAllocator, uint64 InPageSize)
+	TPagedArray(ILinearAllocator& InAllocator, uint64 InPageSize)
 		: Allocator(InAllocator)
 		, PageSize(InPageSize)
 	{
@@ -401,7 +403,7 @@ private:
 	template<typename ItemType, typename PageType>
 	friend class TPagedArrayIterator;
 
-	Trace::ILinearAllocator& Allocator;
+	ILinearAllocator& Allocator;
 	TArray<PageType> PagesArray;
 	PageType* FirstPage = nullptr;
 	PageType* LastPage = nullptr;
@@ -419,4 +421,6 @@ template<typename ItemType, typename PageType>
 inline SIZE_T GetNum(const TPagedArray<ItemType, PageType>& PagedArray)
 {
 	return PagedArray.NumPages();
-}
\ No newline at end of file
+}
+
+} // namespace TraceServices
diff --git a/Engine/Source/Developer/TraceServices/Public/Common/VariablePagedArray.h b/Engine/Source/Developer/TraceServices/Public/Common/VariablePagedArray.h
--- a/Engine/Source/Developer/TraceServices/Public/Common/VariablePagedArray.h
+++ b/Engine/Source/Developer/TraceServices/Public/Common/VariablePagedArray.h
@@ -5,6 +5,9 @@
 #include "TraceServices/Containers/Allocators.h"
 #include "Containers/Array.h"
 
+namespace TraceServices
+{
+
 ////////////////////////////////////////////////////////////////////////////////////////////////////
 //
 // The TVariablePagedArrayPage allocates items in fixed size pages. Each page allocates space for
@@ -40,7 +43,7 @@
 //
 ////////////////////////////////////////////////////////////////////////////////////////////////////
 
-// For debuging the TVariablePagedArray implementation.
+// For debugging the TVariablePagedArray implementation.
 #define DEBUG_VARIABLE_PAGED_ARRAY 0
 
 #if DEBUG_VARIABLE_PAGED_ARRAY
@@ -341,7 +344,7 @@ public:
 	typedef TVariablePagedArrayPageGroup<InItemType> PageGroupType;
 	typedef TVariablePagedArrayIterator<InItemType> TIterator;
 
-	TVariablePagedArray(Trace::ILinearAllocator& InAllocator, uint64 InPageSize)
+	TVariablePagedArray(ILinearAllocator& InAllocator, uint64 InPageSize)
 		: Allocator(InAllocator)
 		, Pages()
 		, FirstPage(nullptr)
@@ -773,7 +776,7 @@ private:
 	template<typename ItemType>
 	friend class TVariablePagedArrayIterator;
 
-	Trace::ILinearAllocator& Allocator;
+	ILinearAllocator& Allocator;
 
 	TArray<PageType> Pages;
 	PageType* FirstPage;
@@ -789,3 +792,5 @@ private:
 
 ////////////////////////////////////////////////////////////////////////////////////////////////////
 ////////////////////////////////////////////////////////////////////////////////////////////////////
+
+} // namespace TraceServices
diff --git a/Engine/Source/Developer/TraceServices/Public/Model/AsyncEnumerateTask.h b/Engine/Source/Developer/TraceServices/Public/Model/AsyncEnumerateTask.h
--- a/Engine/Source/Developer/TraceServices/Public/Model/AsyncEnumerateTask.h
+++ b/Engine/Source/Developer/TraceServices/Public/Model/AsyncEnumerateTask.h
@@ -7,8 +7,9 @@
 #include "Model/MonotonicTimelineData.h"
 #include "TraceServices/Containers/Timelines.h"
 
-namespace Trace
+namespace TraceServices
 {
+
 template<typename InEventType, typename SettingsType>
 struct FAsyncEnumerateTaskData
 {
@@ -328,4 +329,5 @@ public:
 
 	FAsyncEnumerateTaskData<InEventType, SettingsType> Data;
 };
-}
\ No newline at end of file
+
+} // namespace TraceServices
diff --git a/Engine/Source/Developer/TraceServices/Public/Model/IntervalTimeline.h b/Engine/Source/Developer/TraceServices/Public/Model/IntervalTimeline.h
--- a/Engine/Source/Developer/TraceServices/Public/Model/IntervalTimeline.h
+++ b/Engine/Source/Developer/TraceServices/Public/Model/IntervalTimeline.h
@@ -7,7 +7,7 @@
 #include "TraceServices/Containers/Allocators.h"
 #include <limits>
 
-namespace Trace
+namespace TraceServices
 {
 
 struct FIntervalTimelineDefaultSettings
@@ -42,7 +42,7 @@ public:
 		{
 			Callback(true, InStartTime, InEvent);
 			Callback(false, InEndTime, InEvent);
-			return Trace::EEventEnumerate::Continue;
+			return EEventEnumerate::Continue;
 		});
 	}
 	
@@ -163,4 +163,4 @@ private:
 	uint64 ModCount = 0;
 };
 
-}
\ No newline at end of file
+} // namespace TraceServices
diff --git a/Engine/Source/Developer/TraceServices/Public/Model/MonotonicTimelineData.h b/Engine/Source/Developer/TraceServices/Public/Model/MonotonicTimelineData.h
--- a/Engine/Source/Developer/TraceServices/Public/Model/MonotonicTimelineData.h
+++ b/Engine/Source/Developer/TraceServices/Public/Model/MonotonicTimelineData.h
@@ -4,106 +4,108 @@
 
 #include "Common/PagedArray.h"
 
-namespace Trace
+namespace TraceServices
 {
-	struct FEventScopeEntry
+
+struct FEventScopeEntry
+{
+	double Time;
+};
+
+struct FEventStackEntry
+{
+	uint64 EnterScopeIndex;
+	uint64 EventIndex;
+	double ExclTime = 0.0;
+	double EndTime = -1.0; //By convention a negative EndTime means the event does not end in the current page
+};
+
+struct FEventScopeEntryPage
+{
+	FEventScopeEntry* Items = nullptr;
+	uint64 Count = 0;
+	double BeginTime = 0.0;
+	double EndTime = 0.0;
+	uint64 BeginEventIndex = 0;
+	uint64 EndEventIndex = 0;
+	FEventStackEntry* InitialStack = nullptr;
+	uint16 InitialStackCount = 0;
+};
+
+template<typename EventType>
+struct FDetailLevelDepthState
+{
+	int64 PendingScopeEnterIndex = -1;
+	int64 PendingEventIndex = -1;
+
+	EventType DominatingEvent;
+	double DominatingEventStartTime = 0.0;
+	double DominatingEventEndTime = 0.0;
+	double DominatingEventDuration = 0.0;
+
+	double EnterTime = 0.0;
+	double ExitTime = 0.0;
+};
+
+template<typename EventType, typename SettingsType>
+struct FDetailLevelInsertionState
+{
+	double LastTime = -1.0;
+	uint16 CurrentDepth = 0;
+	int32 PendingDepth = -1;
+	FDetailLevelDepthState<EventType> DepthStates[SettingsType::MaxDepth];
+	FEventStackEntry EventStack[SettingsType::MaxDepth];
+	uint64 CurrentScopeEntryPageIndex = (uint64)-1;
+};
+
+template<typename EventType, typename SettingsType>
+struct FDetailLevel
+{
+	FDetailLevel(ILinearAllocator& Allocator, double InResolution)
+		: Resolution(InResolution)
+		, ScopeEntries(Allocator, SettingsType::ScopeEntriesPageSize)
+		, Events(Allocator, SettingsType::EventsPageSize)
 	{
-		double Time;
-	};
 
-	struct FEventStackEntry
+	}
+
+	double GetScopeEntryTime(uint64 Index) const
 	{
-		uint64 EnterScopeIndex;
-		uint64 EventIndex;
-		double ExclTime = 0.0;
-		double EndTime = -1.0; //By convention a negative EndTime means the event does not end in the current page
-	};
+		const FEventScopeEntry& ScopeEntry = ScopeEntries[Index];
+		return ScopeEntry.Time < 0 ? -ScopeEntry.Time : ScopeEntry.Time;
+	}
 
-	struct FEventScopeEntryPage
+	void SetEvent(uint64 Index, const EventType& Event)
 	{
-		FEventScopeEntry* Items = nullptr;
-		uint64 Count = 0;
-		double BeginTime = 0.0;
-		double EndTime = 0.0;
-		uint64 BeginEventIndex = 0;
-		uint64 EndEventIndex = 0;
-		FEventStackEntry* InitialStack = nullptr;
-		uint16 InitialStackCount = 0;
-	};
-
-	template<typename EventType>
-	struct FDetailLevelDepthState
+		Events[Index] = Event;
+	}
+
+	EventType GetEvent(uint64 Index) const
 	{
-		int64 PendingScopeEnterIndex = -1;
-		int64 PendingEventIndex = -1;
+		return Events[Index];
+	}
 
-		EventType DominatingEvent;
-		double DominatingEventStartTime = 0.0;
-		double DominatingEventEndTime = 0.0;
-		double DominatingEventDuration = 0.0;
+	double Resolution;
+	TPagedArray<FEventScopeEntry, FEventScopeEntryPage> ScopeEntries;
+	TPagedArray<EventType> Events;
 
-		double EnterTime = 0.0;
-		double ExitTime = 0.0;
-	};
+	FDetailLevelInsertionState<EventType, SettingsType> InsertionState;
+};
 
-	template<typename EventType, typename SettingsType>
-	struct FDetailLevelInsertionState
-	{
-		double LastTime = -1.0;
-		uint16 CurrentDepth = 0;
-		int32 PendingDepth = -1;
-		FDetailLevelDepthState<EventType> DepthStates[SettingsType::MaxDepth];
-		FEventStackEntry EventStack[SettingsType::MaxDepth];
-		uint64 CurrentScopeEntryPageIndex = (uint64)-1;
-	};
-
-	template<typename EventType, typename SettingsType>
-	struct FDetailLevel
-	{
-		FDetailLevel(ILinearAllocator& Allocator, double InResolution)
-			: Resolution(InResolution)
-			, ScopeEntries(Allocator, SettingsType::ScopeEntriesPageSize)
-			, Events(Allocator, SettingsType::EventsPageSize)
-		{
-
-		}
-
-		double GetScopeEntryTime(uint64 Index) const
-		{
-			const FEventScopeEntry& ScopeEntry = ScopeEntries[Index];
-			return ScopeEntry.Time < 0 ? -ScopeEntry.Time : ScopeEntry.Time;
-		}
-
-		void SetEvent(uint64 Index, const EventType& Event)
-		{
-			Events[Index] = Event;
-		}
-
-		EventType GetEvent(uint64 Index) const
-		{
-			return Events[Index];
-		}
-
-		double Resolution;
-		TPagedArray<FEventScopeEntry, FEventScopeEntryPage> ScopeEntries;
-		TPagedArray<EventType> Events;
-
-		FDetailLevelInsertionState<EventType, SettingsType> InsertionState;
-	};
-
-	template<typename EventType, typename SettingsType>
-	struct FEventInfoStackEntry
+template<typename EventType, typename SettingsType>
+struct FEventInfoStackEntry
+{
+	FEventInfoStackEntry() = default;
+	FEventInfoStackEntry(const FEventStackEntry& EventStackEntry, const FDetailLevel<EventType, SettingsType>& DetailLevel)
 	{
-		FEventInfoStackEntry() = default;
-		FEventInfoStackEntry(const FEventStackEntry& EventStackEntry, const FDetailLevel<EventType, SettingsType>& DetailLevel)
-		{
-			StartTime = DetailLevel.GetScopeEntryTime(EventStackEntry.EnterScopeIndex);
-			Event = DetailLevel.GetEvent(EventStackEntry.EventIndex);
-			ExclTime = EventStackEntry.ExclTime;
-		}
-
-		EventType Event;
-		double StartTime = 0.0;
-		double ExclTime = 0.0;
-	};
-}
\ No newline at end of file
+		StartTime = DetailLevel.GetScopeEntryTime(EventStackEntry.EnterScopeIndex);
+		Event = DetailLevel.GetEvent(EventStackEntry.EventIndex);
+		ExclTime = EventStackEntry.ExclTime;
+	}
+
+	EventType Event;
+	double StartTime = 0.0;
+	double ExclTime = 0.0;
+};
+
+} // namespace TraceServices
diff --git a/Engine/Source/Developer/TraceServices/Public/Model/PointTimeline.h b/Engine/Source/Developer/TraceServices/Public/Model/PointTimeline.h
--- a/Engine/Source/Developer/TraceServices/Public/Model/PointTimeline.h
+++ b/Engine/Source/Developer/TraceServices/Public/Model/PointTimeline.h
@@ -7,7 +7,7 @@
 #include "TraceServices/Containers/Allocators.h"
 #include <limits>
 
-namespace Trace
+namespace TraceServices
 {
 
 struct FPointTimelineDefaultSettings
@@ -42,7 +42,7 @@ public:
 		{
 			Callback(true, InStartTime, InEvent);
 			Callback(false, InEndTime, InEvent);
-			return Trace::EEventEnumerate::Continue;
+			return EEventEnumerate::Continue;
 		});
 	}
 	
@@ -174,4 +174,4 @@ private:
 	bool bEnumerateOutsideRange = false;
 };
 
-}
\ No newline at end of file
+} // namespace TraceServices
diff --git a/Engine/Source/Developer/TraceServices/Public/TraceServices/ITraceServicesModule.h b/Engine/Source/Developer/TraceServices/Public/TraceServices/ITraceServicesModule.h
--- a/Engine/Source/Developer/TraceServices/Public/TraceServices/ITraceServicesModule.h
+++ b/Engine/Source/Developer/TraceServices/Public/TraceServices/ITraceServicesModule.h
@@ -5,22 +5,22 @@
 #include "CoreMinimal.h"
 #include "Modules/ModuleInterface.h"
 
-namespace Trace
+namespace TraceServices
 {
 
 class IAnalysisService;
 class IModuleService;
 
-}
+} // namespace TraceServices
 
 class ITraceServicesModule
 	: public IModuleInterface
 {
 public:
-	virtual TSharedPtr<Trace::IAnalysisService> GetAnalysisService() = 0;
-	virtual TSharedPtr<Trace::IModuleService> GetModuleService() = 0;
-	virtual TSharedPtr<Trace::IAnalysisService> CreateAnalysisService() = 0;
-	virtual TSharedPtr<Trace::IModuleService> CreateModuleService() = 0;
+	virtual TSharedPtr<TraceServices::IAnalysisService> GetAnalysisService() = 0;
+	virtual TSharedPtr<TraceServices::IModuleService> GetModuleService() = 0;
+	virtual TSharedPtr<TraceServices::IAnalysisService> CreateAnalysisService() = 0;
+	virtual TSharedPtr<TraceServices::IModuleService> CreateModuleService() = 0;
 
 	virtual ~ITraceServicesModule() = default;
 };
diff --git a/Engine/Source/Developer/TraceServices/Public/TraceServices/Model/AnalysisSession.h b/Engine/Source/Developer/TraceServices/Public/TraceServices/Model/AnalysisSession.h
--- a/Engine/Source/Developer/TraceServices/Public/TraceServices/Model/AnalysisSession.h
+++ b/Engine/Source/Developer/TraceServices/Public/TraceServices/Model/AnalysisSession.h
@@ -10,9 +10,15 @@ class FName;
 namespace Trace
 {
 
-class ILinearAllocator;
 class IAnalyzer;
 
+}
+
+namespace TraceServices
+{
+
+class ILinearAllocator;
+
 class IProvider
 {
 public:
@@ -43,7 +49,7 @@ public:
 	virtual const TCHAR* StoreString(const TCHAR* String) = 0;
 	virtual const TCHAR* StoreString(const FStringView& String) = 0;
 	
-	virtual void AddAnalyzer(IAnalyzer* Analyzer) = 0;
+	virtual void AddAnalyzer(Trace::IAnalyzer* Analyzer) = 0;
 
 	virtual void AddProvider(const FName& Name, IProvider* Provider) = 0;
 	template<typename ProviderType>
@@ -95,4 +101,4 @@ struct FAnalysisSessionEditScope
 	IAnalysisSession& AnalysisSession;
 };
 	
-}
+} // namespace TraceServices
diff --git a/Engine/Source/Developer/TraceServices/Public/TraceServices/Model/Memory.h b/Engine/Source/Developer/TraceServices/Public/TraceServices/Model/Memory.h
--- a/Engine/Source/Developer/TraceServices/Public/TraceServices/Model/Memory.h
+++ b/Engine/Source/Developer/TraceServices/Public/TraceServices/Model/Memory.h
@@ -7,13 +7,13 @@
 #include "Templates/Function.h"
 #include "TraceServices/Model/AnalysisSession.h"
 
-namespace Trace
+namespace TraceServices
 {
 
 typedef int32 FMemoryTrackerId;
 typedef int64 FMemoryTagId;
 
-struct FMemoryTracker
+struct FMemoryTrackerInfo
 {
 	static const FMemoryTrackerId InvalidTrackerId = -1;
 	static const FMemoryTrackerId MaxValidTrackerId = 63;
@@ -25,7 +25,7 @@ struct FMemoryTracker
 	FString Name;
 };
 
-struct FMemoryTag
+struct FMemoryTagInfo
 {
 	static const FMemoryTagId InvalidTagId = 0;
 
@@ -71,14 +71,14 @@ public:
 	 * Enumerate the registered tags.
 	 * @param Callback Function that will be called for each registered tag.
 	 */
-	virtual void EnumerateTags(TFunctionRef<void(const FMemoryTag&)> Callback) const = 0;
+	virtual void EnumerateTags(TFunctionRef<void(const FMemoryTagInfo&)> Callback) const = 0;
 
 	/**
 	 * Returns the meta data for a tag id.
 	 * @param Id Tag id
 	 * @return Memory tag information.
 	 */
-	virtual const FMemoryTag* GetTag(FMemoryTagId Id) const = 0;
+	virtual const FMemoryTagInfo* GetTag(FMemoryTagId Id) const = 0;
 
 	/**
 	 * Gets the number of samples for a given tag from a given tracker.
@@ -98,7 +98,7 @@ public:
 	 * Enumerate the registered trackers.
 	 * @param Callback Function that is called for each registered tracker.
 	 */
-	virtual void EnumerateTrackers(TFunctionRef<void(const FMemoryTracker&)> Callback) const = 0;
+	virtual void EnumerateTrackers(TFunctionRef<void(const FMemoryTrackerInfo&)> Callback) const = 0;
 
 	/** Enumerates samples (values) for a specified LLM tag.
 	 * @param Tracker Tracker index.
@@ -113,4 +113,4 @@ public:
 
 TRACESERVICES_API const IMemoryProvider& ReadMemoryProvider(const IAnalysisSession& Session);
 
-} // namespace Trace
+} // namespace TraceServices
diff --git a/Engine/Source/Developer/TraceServices/Public/TraceServices/Utils.h b/Engine/Source/Developer/TraceServices/Public/TraceServices/Utils.h
--- a/Engine/Source/Developer/TraceServices/Public/TraceServices/Utils.h
+++ b/Engine/Source/Developer/TraceServices/Public/TraceServices/Utils.h
@@ -2,8 +2,8 @@
 
 #pragma once
 
-namespace FTraceAnalyzerUtils
-{
+namespace TraceServices {
+namespace FTraceAnalyzerUtils {
 
 template <typename AttachedCharType>
 inline FString LegacyAttachmentString(
@@ -21,3 +21,4 @@ inline FString LegacyAttachmentString(
 }
 
 } // namespace FTraceAnalyzerUtils
+} // namespace TraceServices
