Fixed static analysis errors.

#rnx

diff --git a/Engine/Source/Developer/TraceServices/Private/Analyzers/AllocationsAnalysis.cpp b/Engine/Source/Developer/TraceServices/Private/Analyzers/AllocationsAnalysis.cpp
--- a/Engine/Source/Developer/TraceServices/Private/Analyzers/AllocationsAnalysis.cpp
+++ b/Engine/Source/Developer/TraceServices/Private/Analyzers/AllocationsAnalysis.cpp
@@ -4,9 +4,9 @@
 
 #if 0
 
-//namespace TraceServices
-//{
-//
+namespace TraceServices
+{
+
 ////////////////////////////////////////////////////////////////////////////////
 enum
 {
diff --git a/Engine/Source/Developer/TraceServices/Private/Analyzers/AllocationsAnalysis.h b/Engine/Source/Developer/TraceServices/Private/Analyzers/AllocationsAnalysis.h
--- a/Engine/Source/Developer/TraceServices/Private/Analyzers/AllocationsAnalysis.h
+++ b/Engine/Source/Developer/TraceServices/Private/Analyzers/AllocationsAnalysis.h
@@ -22,6 +22,6 @@ public:
 	virtual void	OnAnalysisBegin(const FOnAnalysisContext& Context) override;
 };
 
-} // namespace UE
+} // namespace TraceServices
 
 #endif // 0
