Another include-tool analysis fix.

#rnx

diff --git a/Engine/Source/Developer/TraceServices/Private/Analyzers/AllocationsAnalysis.h b/Engine/Source/Developer/TraceServices/Private/Analyzers/AllocationsAnalysis.h
--- a/Engine/Source/Developer/TraceServices/Private/Analyzers/AllocationsAnalysis.h
+++ b/Engine/Source/Developer/TraceServices/Private/Analyzers/AllocationsAnalysis.h
@@ -22,6 +22,4 @@ public:
 	virtual void	OnAnalysisBegin(const FOnAnalysisContext& Context) override;
 };
 
-} // namespace TraceServices
-
 #endif // 0
