Fixed two of TraceServices' symbol name colliding with Core and TraceInsight modules and causing ununity ambiguation.

#rnx

diff --git a/Engine/Source/Developer/TraceInsights/Private/Insights/ViewModels/LogMessage.cpp b/Engine/Source/Developer/TraceInsights/Private/Insights/ViewModels/LogMessage.cpp
--- a/Engine/Source/Developer/TraceInsights/Private/Insights/ViewModels/LogMessage.cpp
+++ b/Engine/Source/Developer/TraceInsights/Private/Insights/ViewModels/LogMessage.cpp
@@ -26,7 +26,7 @@ FLogMessageRecord::FLogMessageRecord()
 
 ////////////////////////////////////////////////////////////////////////////////////////////////////
 
-FLogMessageRecord::FLogMessageRecord(const TraceServices::FLogMessage& TraceLogMessage)
+FLogMessageRecord::FLogMessageRecord(const TraceServices::FLogMessageInfo& TraceLogMessage)
 	: Index(static_cast<int32>(TraceLogMessage.Index))
 	, Time(TraceLogMessage.Time)
 	, Verbosity(TraceLogMessage.Verbosity)
@@ -158,7 +158,7 @@ FLogMessageRecord& FLogMessageCache::Get(uint64 Index)
 	{
 		TraceServices::FAnalysisSessionReadScope SessionReadScope(*Session.Get());
 		const TraceServices::ILogProvider& LogProvider = TraceServices::ReadLogProvider(*Session.Get());
-		LogProvider.ReadMessage(Index, [this, Index](const TraceServices::FLogMessage& Message)
+		LogProvider.ReadMessage(Index, [this, Index](const TraceServices::FLogMessageInfo& Message)
 		{
 			FScopeLock Lock(&CriticalSection);
 			FLogMessageRecord Entry(Message);
@@ -187,7 +187,7 @@ TSharedPtr<FLogMessageRecord> FLogMessageCache::GetUncached(uint64 Index) const
 	{
 		TraceServices::FAnalysisSessionReadScope SessionReadScope(*Session.Get());
 		const TraceServices::ILogProvider& LogProvider = TraceServices::ReadLogProvider(*Session.Get());
-		LogProvider.ReadMessage(Index, [&EntryPtr](const TraceServices::FLogMessage& Message)
+		LogProvider.ReadMessage(Index, [&EntryPtr](const TraceServices::FLogMessageInfo& Message)
 		{
 			EntryPtr = MakeShared<FLogMessageRecord>(Message);
 		});
diff --git a/Engine/Source/Developer/TraceInsights/Private/Insights/ViewModels/LogMessage.h b/Engine/Source/Developer/TraceInsights/Private/Insights/ViewModels/LogMessage.h
--- a/Engine/Source/Developer/TraceInsights/Private/Insights/ViewModels/LogMessage.h
+++ b/Engine/Source/Developer/TraceInsights/Private/Insights/ViewModels/LogMessage.h
@@ -8,7 +8,7 @@
 
 namespace TraceServices
 {
-	struct FLogMessage;
+	struct FLogMessageInfo;
 	class IAnalysisSession;
 }
 
@@ -18,7 +18,7 @@ class FLogMessageRecord
 {
 public:
 	FLogMessageRecord();
-	FLogMessageRecord(const TraceServices::FLogMessage& Message);
+	FLogMessageRecord(const TraceServices::FLogMessageInfo& Message);
 
 	FText GetIndexAsText() const;
 	FText GetTimeAsText() const;
diff --git a/Engine/Source/Developer/TraceInsights/Private/Insights/ViewModels/MarkersTimingTrack.cpp b/Engine/Source/Developer/TraceInsights/Private/Insights/ViewModels/MarkersTimingTrack.cpp
--- a/Engine/Source/Developer/TraceInsights/Private/Insights/ViewModels/MarkersTimingTrack.cpp
+++ b/Engine/Source/Developer/TraceInsights/Private/Insights/ViewModels/MarkersTimingTrack.cpp
@@ -131,7 +131,7 @@ void FMarkersTimingTrack::UpdateDrawState(const FTimingTrackViewport& InViewport
 		LogProvider.EnumerateMessages(
 			Builder.GetViewport().GetStartTime(),
 			Builder.GetViewport().GetEndTime(),
-			[&Builder](const TraceServices::FLogMessage& Message) { Builder.AddLogMessage(Message); });
+			[&Builder](const TraceServices::FLogMessageInfo& Message) { Builder.AddLogMessage(Message); });
 
 		Builder.EndLog();
 	}
@@ -312,7 +312,7 @@ void FTimeMarkerTrackBuilder::BeginLog(const TraceServices::ILogProvider& LogPro
 
 ////////////////////////////////////////////////////////////////////////////////////////////////////
 
-void FTimeMarkerTrackBuilder::AddLogMessage(const TraceServices::FLogMessage& Message)
+void FTimeMarkerTrackBuilder::AddLogMessage(const TraceServices::FLogMessageInfo& Message)
 {
 	Track.NumLogMessages++;
 
@@ -323,7 +323,7 @@ void FTimeMarkerTrackBuilder::AddLogMessage(const TraceServices::FLogMessage& Me
 		//TODO: Search API like: LogProviderPtr->SearchMessage(StartIndex, ESearchDirection::Backward, LambdaPredicate, bResolveFormatString);
 		LogProviderPtr->ReadMessage(
 			Message.Index - 1,
-			[this](const TraceServices::FLogMessage& Message) { AddLogMessage(Message); });
+			[this](const TraceServices::FLogMessageInfo& Message) { AddLogMessage(Message); });
 	}
 
 	check(Message.Category != nullptr);
diff --git a/Engine/Source/Developer/TraceInsights/Private/Insights/ViewModels/MarkersTimingTrack.h b/Engine/Source/Developer/TraceInsights/Private/Insights/ViewModels/MarkersTimingTrack.h
--- a/Engine/Source/Developer/TraceInsights/Private/Insights/ViewModels/MarkersTimingTrack.h
+++ b/Engine/Source/Developer/TraceInsights/Private/Insights/ViewModels/MarkersTimingTrack.h
@@ -11,7 +11,7 @@
 
 namespace TraceServices
 {
-	struct FLogMessage;
+	struct FLogMessageInfo;
 	class ILogProvider;
 }
 
@@ -125,7 +125,7 @@ public:
 	const FTimingTrackViewport& GetViewport() { return Viewport; }
 
 	void BeginLog(const TraceServices::ILogProvider& LogProvider);
-	void AddLogMessage(const TraceServices::FLogMessage& Message);
+	void AddLogMessage(const TraceServices::FLogMessageInfo& Message);
 	void EndLog();
 
 	static FLinearColor GetColorByCategory(const TCHAR* const Category);
diff --git a/Engine/Source/Developer/TraceInsights/Private/Insights/Widgets/SLogView.cpp b/Engine/Source/Developer/TraceInsights/Private/Insights/Widgets/SLogView.cpp
--- a/Engine/Source/Developer/TraceInsights/Private/Insights/Widgets/SLogView.cpp
+++ b/Engine/Source/Developer/TraceInsights/Private/Insights/Widgets/SLogView.cpp
@@ -730,7 +730,7 @@ void SLogView::Tick(const FGeometry& AllottedGeometry, const double InCurrentTim
 			UE_LOG(TraceInsights, Log, TEXT("[LogView] Total Log Categories: %d"), TotalNumCategories);
 
 			TSet<FName> Categories;
-			LogProvider.EnumerateCategories([&Categories](const TraceServices::FLogCategory& Category)
+			LogProvider.EnumerateCategories([&Categories](const TraceServices::FLogCategoryInfo& Category)
 			{
 				FString CategoryStr(Category.Name);
 				if (CategoryStr.StartsWith(TEXT("Log")))
diff --git a/Engine/Source/Developer/TraceServices/Private/Analyzers/LogTraceAnalysis.cpp b/Engine/Source/Developer/TraceServices/Private/Analyzers/LogTraceAnalysis.cpp
--- a/Engine/Source/Developer/TraceServices/Private/Analyzers/LogTraceAnalysis.cpp
+++ b/Engine/Source/Developer/TraceServices/Private/Analyzers/LogTraceAnalysis.cpp
@@ -34,7 +34,7 @@ bool FLogTraceAnalyzer::OnEvent(uint16 RouteId, EStyle Style, const FOnEventCont
 	case RouteId_LogCategory:
 	{
 		uint64 CategoryPointer = EventData.GetValue<uint64>("CategoryPointer");
-		FLogCategory& Category = LogProvider.GetCategory(CategoryPointer);
+		FLogCategoryInfo& Category = LogProvider.GetCategory(CategoryPointer);
 
 		FString Name = FTraceAnalyzerUtils::LegacyAttachmentString<TCHAR>("Name", Context);
 		Category.Name = Session.StoreString(*Name);
@@ -47,7 +47,7 @@ bool FLogTraceAnalyzer::OnEvent(uint16 RouteId, EStyle Style, const FOnEventCont
 		uint64 LogPoint = EventData.GetValue<uint64>("LogPoint");
 		FLogMessageSpec& Spec = LogProvider.GetMessageSpec(LogPoint);
 		uint64 CategoryPointer = EventData.GetValue<uint64>("CategoryPointer");
-		FLogCategory& Category = LogProvider.GetCategory(CategoryPointer);
+		FLogCategoryInfo& Category = LogProvider.GetCategory(CategoryPointer);
 		Spec.Category = &Category;
 		Spec.Line = EventData.GetValue<int32>("Line");
 		Spec.Verbosity = static_cast<ELogVerbosity::Type>(EventData.GetValue<uint8>("Verbosity"));
diff --git a/Engine/Source/Developer/TraceServices/Private/Analyzers/MiscTraceAnalysis.cpp b/Engine/Source/Developer/TraceServices/Private/Analyzers/MiscTraceAnalysis.cpp
--- a/Engine/Source/Developer/TraceServices/Private/Analyzers/MiscTraceAnalysis.cpp
+++ b/Engine/Source/Developer/TraceServices/Private/Analyzers/MiscTraceAnalysis.cpp
@@ -24,7 +24,7 @@ FMiscTraceAnalyzer::FMiscTraceAnalyzer(IAnalysisSession& InSession,
 	, FrameProvider(InFrameProvider)
 	, ChannelProvider(InChannelProvider)
 {
-	FLogCategory& BookmarkLogCategory = LogProvider.GetCategory(FLogProvider::ReservedLogCategory_Bookmark);
+	FLogCategoryInfo& BookmarkLogCategory = LogProvider.GetCategory(FLogProvider::ReservedLogCategory_Bookmark);
 	BookmarkLogCategory.Name = TEXT("LogBookmark");
 	BookmarkLogCategory.DefaultVerbosity = ELogVerbosity::All;
 }
diff --git a/Engine/Source/Developer/TraceServices/Private/Model/Log.cpp b/Engine/Source/Developer/TraceServices/Private/Model/Log.cpp
--- a/Engine/Source/Developer/TraceServices/Private/Model/Log.cpp
+++ b/Engine/Source/Developer/TraceServices/Private/Model/Log.cpp
@@ -42,7 +42,7 @@ FLogProvider::FLogProvider(IAnalysisSession& InSession)
 		AddColumn(&FLogMessageInternal::Message, TEXT("Message"));
 }
 
-FLogCategory& FLogProvider::GetCategory(uint64 CategoryPointer)
+FLogCategoryInfo& FLogProvider::GetCategory(uint64 CategoryPointer)
 {
 	Session.WriteAccessCheck();
 	if (CategoryMap.Contains(CategoryPointer))
@@ -51,7 +51,7 @@ FLogCategory& FLogProvider::GetCategory(uint64 CategoryPointer)
 	}
 	else
 	{
-		FLogCategory& Category = Categories.PushBack();
+		FLogCategoryInfo& Category = Categories.PushBack();
 		CategoryMap.Add(CategoryPointer, &Category);
 		return Category;
 	}
@@ -93,7 +93,7 @@ uint64 FLogProvider::GetMessageCount() const
 	return Messages.Num();
 }
 
-bool FLogProvider::ReadMessage(uint64 Index, TFunctionRef<void(const FLogMessage &)> Callback) const
+bool FLogProvider::ReadMessage(uint64 Index, TFunctionRef<void(const FLogMessageInfo&)> Callback) const
 {
 	Session.ReadAccessCheck();
 	if (Index >= Messages.Num())
@@ -104,7 +104,7 @@ bool FLogProvider::ReadMessage(uint64 Index, TFunctionRef<void(const FLogMessage
 	return true;
 }
 
-void FLogProvider::EnumerateMessages(double IntervalStart, double IntervalEnd, TFunctionRef<void(const FLogMessage&)> Callback) const
+void FLogProvider::EnumerateMessages(double IntervalStart, double IntervalEnd, TFunctionRef<void(const FLogMessageInfo&)> Callback) const
 {
 	Session.ReadAccessCheck();
 	if (IntervalStart > IntervalEnd)
@@ -122,7 +122,7 @@ void FLogProvider::EnumerateMessages(double IntervalStart, double IntervalEnd, T
 	}
 }
 
-void FLogProvider::EnumerateMessagesByIndex(uint64 Start, uint64 End, TFunctionRef<void(const FLogMessage &)> Callback) const
+void FLogProvider::EnumerateMessagesByIndex(uint64 Start, uint64 End, TFunctionRef<void(const FLogMessageInfo&)> Callback) const
 {
 	Session.ReadAccessCheck();
 	uint64 Count = Messages.Num();
@@ -144,10 +144,10 @@ void FLogProvider::EnumerateMessagesByIndex(uint64 Start, uint64 End, TFunctionR
 	}
 }
 
-void FLogProvider::ConstructMessage(uint64 Index, TFunctionRef<void(const FLogMessage &)> Callback) const
+void FLogProvider::ConstructMessage(uint64 Index, TFunctionRef<void(const FLogMessageInfo&)> Callback) const
 {
 	const FLogMessageInternal& InternalMessage = Messages[Index];
-	FLogMessage Message;
+	FLogMessageInfo Message;
 	Message.Index = Index;
 	Message.Time = InternalMessage.Time;
 	Message.Category = InternalMessage.Spec->Category;
@@ -158,7 +158,7 @@ void FLogProvider::ConstructMessage(uint64 Index, TFunctionRef<void(const FLogMe
 	Callback(Message);
 }
 
-void FLogProvider::EnumerateCategories(TFunctionRef<void(const FLogCategory &)> Callback) const
+void FLogProvider::EnumerateCategories(TFunctionRef<void(const FLogCategoryInfo&)> Callback) const
 {
 	Session.ReadAccessCheck();
 	for (auto Iterator = Categories.GetIteratorFromItem(0); Iterator; ++Iterator)
diff --git a/Engine/Source/Developer/TraceServices/Private/Model/LogPrivate.h b/Engine/Source/Developer/TraceServices/Private/Model/LogPrivate.h
--- a/Engine/Source/Developer/TraceServices/Private/Model/LogPrivate.h
+++ b/Engine/Source/Developer/TraceServices/Private/Model/LogPrivate.h
@@ -16,7 +16,7 @@ class FStringStore;
 
 struct FLogMessageSpec
 {
-	FLogCategory* Category = nullptr;
+	FLogCategoryInfo* Category = nullptr;
 	const TCHAR* File = nullptr;
 	const TCHAR* FormatString = nullptr;
 	int32 Line;
@@ -43,20 +43,20 @@ public:
 
 	FLogProvider(IAnalysisSession& Session);
 
-	FLogCategory& GetCategory(uint64 CategoryPointer);
+	FLogCategoryInfo& GetCategory(uint64 CategoryPointer);
 	FLogMessageSpec& GetMessageSpec(uint64 LogPoint);
 	void AppendMessage(uint64 LogPoint, double Time, const uint8* FormatArgs);
 
 	virtual uint64 GetMessageCount() const override;
-	virtual bool ReadMessage(uint64 Index, TFunctionRef<void(const FLogMessage &)> Callback) const override;
-	virtual void EnumerateMessages(double IntervalStart, double IntervalEnd, TFunctionRef<void(const FLogMessage&)> Callback) const override;
-	virtual void EnumerateMessagesByIndex(uint64 Start, uint64 End, TFunctionRef<void(const FLogMessage&)> Callback) const override;
+	virtual bool ReadMessage(uint64 Index, TFunctionRef<void(const FLogMessageInfo&)> Callback) const override;
+	virtual void EnumerateMessages(double IntervalStart, double IntervalEnd, TFunctionRef<void(const FLogMessageInfo&)> Callback) const override;
+	virtual void EnumerateMessagesByIndex(uint64 Start, uint64 End, TFunctionRef<void(const FLogMessageInfo&)> Callback) const override;
 	virtual uint64 GetCategoryCount() const override { return Categories.Num(); }
-	virtual void EnumerateCategories(TFunctionRef<void(const FLogCategory&)> Callback) const override;
+	virtual void EnumerateCategories(TFunctionRef<void(const FLogCategoryInfo&)> Callback) const override;
 	virtual const IUntypedTable& GetMessagesTable() const override { return MessagesTable; }
 
 private:
-	void ConstructMessage(uint64 Id, TFunctionRef<void(const FLogMessage &)> Callback) const;
+	void ConstructMessage(uint64 Id, TFunctionRef<void(const FLogMessageInfo&)> Callback) const;
 
 	enum
 	{
@@ -64,9 +64,9 @@ private:
 	};
 
 	IAnalysisSession& Session;
-	TMap<uint64, FLogCategory*> CategoryMap;
+	TMap<uint64, FLogCategoryInfo*> CategoryMap;
 	TMap<uint64, FLogMessageSpec*> SpecMap;
-	TPagedArray<FLogCategory> Categories;
+	TPagedArray<FLogCategoryInfo> Categories;
 	TPagedArray<FLogMessageSpec> MessageSpecs;
 	TPagedArray<FLogMessageInternal> Messages;
 	TCHAR FormatBuffer[FormatBufferSize];
diff --git a/Engine/Source/Developer/TraceServices/Public/TraceServices/Model/Log.h b/Engine/Source/Developer/TraceServices/Public/TraceServices/Model/Log.h
--- a/Engine/Source/Developer/TraceServices/Public/TraceServices/Model/Log.h
+++ b/Engine/Source/Developer/TraceServices/Public/TraceServices/Model/Log.h
@@ -10,17 +10,17 @@
 namespace TraceServices
 {
 
-struct FLogCategory
+struct FLogCategoryInfo
 {
 	const TCHAR* Name = nullptr;
 	ELogVerbosity::Type DefaultVerbosity;
 };
 
-struct FLogMessage
+struct FLogMessageInfo
 {
 	uint64 Index;
 	double Time;
-	const FLogCategory* Category = nullptr;
+	const FLogCategoryInfo* Category = nullptr;
 	const TCHAR* File = nullptr;
 	const TCHAR* Message = nullptr;
 	int32 Line;
@@ -33,11 +33,11 @@ class ILogProvider
 public:
 	virtual ~ILogProvider() = default;
 	virtual uint64 GetMessageCount() const = 0;
-	virtual void EnumerateMessages(double IntervalStart, double IntervalEnd, TFunctionRef<void(const FLogMessage&)> Callback) const = 0;
-	virtual void EnumerateMessagesByIndex(uint64 Start, uint64 End, TFunctionRef<void(const FLogMessage&)> Callback) const = 0;
-	virtual bool ReadMessage(uint64 Index, TFunctionRef<void(const FLogMessage&)> Callback) const = 0;
+	virtual void EnumerateMessages(double IntervalStart, double IntervalEnd, TFunctionRef<void(const FLogMessageInfo&)> Callback) const = 0;
+	virtual void EnumerateMessagesByIndex(uint64 Start, uint64 End, TFunctionRef<void(const FLogMessageInfo&)> Callback) const = 0;
+	virtual bool ReadMessage(uint64 Index, TFunctionRef<void(const FLogMessageInfo&)> Callback) const = 0;
 	virtual uint64 GetCategoryCount() const = 0;
-	virtual void EnumerateCategories(TFunctionRef<void(const FLogCategory&)> Callback) const = 0;
+	virtual void EnumerateCategories(TFunctionRef<void(const FLogCategoryInfo&)> Callback) const = 0;
 	virtual const IUntypedTable& GetMessagesTable() const = 0;
 };
 
