[Insights]
  - Added detection for AllocationsProvider availability to know when to enable the Memory Insights tab (in addition to detecting availability of the LLM tags memory provider).
  - Simplified the API for allocation queries (using A, B, C, D time markers).
  - Added initial AllocationsAnalysis that processes the allocation trace events and delegates them to AllocationsProvider.
  - Added initial AllocationsProvider. Uses a simple map for live allocations.
  - Added FSbTree ("SBIF with offsetted cells" data structure for storing the retired allocations in AllocationsProvider).
  - Added async allocation queries (based on A, B, C, D time markers). Currently it uses a single TaskGraph task.
  - Added "LogTraceServices" log category for general purpose logging in TraceServices module.

diff --git a/Engine/Source/Developer/TraceInsights/Private/Insights/MemoryProfiler/MemoryProfilerManager.cpp b/Engine/Source/Developer/TraceInsights/Private/Insights/MemoryProfiler/MemoryProfilerManager.cpp
--- a/Engine/Source/Developer/TraceInsights/Private/Insights/MemoryProfiler/MemoryProfilerManager.cpp
+++ b/Engine/Source/Developer/TraceInsights/Private/Insights/MemoryProfiler/MemoryProfilerManager.cpp
@@ -5,6 +5,7 @@
 #include "Modules/ModuleManager.h"
 #include "MessageLog/Public/MessageLogModule.h"
 #include "TraceServices/AnalysisService.h"
+#include "TraceServices/Model/AllocationsProvider.h"
 #include "TraceServices/Model/Memory.h"
 #include "WorkspaceMenuStructure.h"
 #include "WorkspaceMenuStructureModule.h"
@@ -232,7 +233,7 @@ bool FMemoryProfilerManager::Tick(float DeltaTime)
 	// Check if session has Memory events (to spawn the tab), but not too often.
 	if (!bIsAvailable && AvailabilityCheck.Tick())
 	{
-		uint32 TagCount = 0;
+		bool bShouldBeAvailable = false;
 
 		TSharedPtr<const TraceServices::IAnalysisSession> Session = FInsightsManager::Get()->GetSession();
 		if (Session.IsValid())
@@ -246,7 +247,20 @@ bool FMemoryProfilerManager::Tick(float DeltaTime)
 			}
 
 			const TraceServices::IMemoryProvider& MemoryProvider = TraceServices::ReadMemoryProvider(*Session.Get());
-			TagCount = MemoryProvider.GetTagCount();
+			uint32 TagCount = MemoryProvider.GetTagCount();
+			if (TagCount > 0)
+			{
+				bShouldBeAvailable = true;
+			}
+
+			const TraceServices::IAllocationsProvider* AllocationsProvider = TraceServices::ReadAllocationsProvider(*Session.Get());
+			if (AllocationsProvider)
+			{
+				if (AllocationsProvider->IsInitialized())
+				{
+					bShouldBeAvailable = true;
+				}
+			}
 		}
 		else
 		{
@@ -254,7 +268,7 @@ bool FMemoryProfilerManager::Tick(float DeltaTime)
 			AvailabilityCheck.Disable();
 		}
 
-		if (TagCount > 0)
+		if (bShouldBeAvailable)
 		{
 			bIsAvailable = true;
 
diff --git a/Engine/Source/Developer/TraceInsights/Private/Insights/MemoryProfiler/ViewModels/MemorySharedState.cpp b/Engine/Source/Developer/TraceInsights/Private/Insights/MemoryProfiler/ViewModels/MemorySharedState.cpp
--- a/Engine/Source/Developer/TraceInsights/Private/Insights/MemoryProfiler/ViewModels/MemorySharedState.cpp
+++ b/Engine/Source/Developer/TraceInsights/Private/Insights/MemoryProfiler/ViewModels/MemorySharedState.cpp
@@ -878,91 +878,96 @@ TSharedPtr<FMemoryGraphTrack> FMemorySharedState::CreateGraphTrack(const Insight
 
 void FMemorySharedState::InitMemoryRules()
 {
+	using ERule = TraceServices::IAllocationsProvider::EQueryRule;
+
 	MemoryRules.Reset();
 
 	MemoryRules.Add(MakeShared<Insights::FMemoryRuleSpec>(
-		Insights::EMemoryRule::aAf, 1,
+		ERule::aAf, 1,
 		LOCTEXT("MemRule_aAf_Short", "*A*"),
 		LOCTEXT("MemRule_aAf_Verbose", "Active Allocs"),
 		LOCTEXT("MemRule_aAf_Desc", "Identifies active allocations at time A.\n(a ≤ A ≤ f)")));
 
 	MemoryRules.Add(MakeShared<Insights::FMemoryRuleSpec>(
-		Insights::EMemoryRule::afA, 1,
+		ERule::afA, 1,
 		LOCTEXT("MemRule_afA_Short", "**A"),
 		LOCTEXT("MemRule_afA_Verbose", "Before"),
 		LOCTEXT("MemRule_afA_Desc", "Identifies allocations allocated and freed before time A.\n(a ≤ f ≤ A)")));
 
 	MemoryRules.Add(MakeShared<Insights::FMemoryRuleSpec>(
-		Insights::EMemoryRule::Aaf, 1,
+		ERule::Aaf, 1,
 		LOCTEXT("MemRule_Aaf_Short", "A**"),
 		LOCTEXT("MemRule_Aaf_Verbose", "After"),
 		LOCTEXT("MemRule_Aaf_Desc", "Identifies allocations allocated and freed after time A.\n(A ≤ a ≤ f)")));
 
 	MemoryRules.Add(MakeShared<Insights::FMemoryRuleSpec>(
-		Insights::EMemoryRule::aAfB, 2,
+		ERule::aAfB, 2,
 		LOCTEXT("MemRule_aAfB_Short", "*A*B"),
 		LOCTEXT("MemRule_aAfB_Verbose", "Decline"),
 		LOCTEXT("MemRule_aAfB_Desc", "Identifies allocations allocated before time A and freed between time A and time B.\n(a ≤ A ≤ f ≤ B)")));
 
 	MemoryRules.Add(MakeShared<Insights::FMemoryRuleSpec>(
-		Insights::EMemoryRule::AaBf, 2,
+		ERule::AaBf, 2,
 		LOCTEXT("MemRule_AaBf_Short", "A*B*"),
 		LOCTEXT("MemRule_AaBf_Verbose", "Growth"),
 		LOCTEXT("MemRule_AaBf_Desc", "Identifies allocations allocated between time A and time B and freed after time B.\n(A ≤ a ≤ B ≤ f)")));
 
 	MemoryRules.Add(MakeShared<Insights::FMemoryRuleSpec>(
-		Insights::EMemoryRule::AafB, 2,
+		ERule::AafB, 2,
 		LOCTEXT("MemRule_AafB_Short", "A**B"),
 		LOCTEXT("MemRule_AafB_Verbose", "Short Living Allocs"),
 		LOCTEXT("MemRule_AafB_Desc", "Identifies allocations allocated and freed between time A and time B.\n(A ≤ a ≤ f ≤ B)")));
 
 	MemoryRules.Add(MakeShared<Insights::FMemoryRuleSpec>(
-		Insights::EMemoryRule::aABf, 2,
+		ERule::aABf, 2,
 		LOCTEXT("MemRule_aABf_Short", "*A B*"),
 		LOCTEXT("MemRule_aABf_Verbose", "Long Living Allocs"),
 		LOCTEXT("MemRule_aABf_Desc", "Identifies allocations allocated before time A and freed after time B.\n(a ≤ A ≤ B ≤ f)")));
 
 	MemoryRules.Add(MakeShared<Insights::FMemoryRuleSpec>(
-		Insights::EMemoryRule::AaBCf, 3,
+		ERule::AaBCf, 3,
 		LOCTEXT("MemRule_AaBCf_Short", "A*B C*"),
 		LOCTEXT("MemRule_AaBCf_Verbose", "Memory Leaks"),
 		LOCTEXT("MemRule_AaBCf_Desc", "Identifies allocations allocated between time A and time B and freed after time C.\n(A ≤ a ≤ B ≤ C ≤ f)")));
 
 	MemoryRules.Add(MakeShared<Insights::FMemoryRuleSpec>(
-		Insights::EMemoryRule::AaBfC, 3,
+		ERule::AaBfC, 3,
 		LOCTEXT("MemRule_AaBfC_Short", "A*B*C"),
 		LOCTEXT("MemRule_AaBfC_Verbose", "Limited Lifetime"),
 		LOCTEXT("MemRule_AaBfC_Desc", "Identifies allocations allocated between time A and time B and freed between time B and time C.\n(A ≤ a ≤ B ≤ f ≤ C)")));
 
 	MemoryRules.Add(MakeShared<Insights::FMemoryRuleSpec>(
-		Insights::EMemoryRule::aABfC, 3,
+		ERule::aABfC, 3,
 		LOCTEXT("MemRule_aABfC_Short", "*A B*C"),
 		LOCTEXT("MemRule_aABfC_Verbose", "Decline of Long Living Allocs"),
 		LOCTEXT("MemRule_aABfC_Desc", "Identifies allocations allocated before time A and freed between time B and time C.\n(a ≤ A ≤ B ≤ f ≤ C)")));
 
 	MemoryRules.Add(MakeShared<Insights::FMemoryRuleSpec>(
-		Insights::EMemoryRule::AaBCfD, 4,
+		ERule::AaBCfD, 4,
 		LOCTEXT("MemRule_AaBCfD_Short", "A*B C*D"),
 		LOCTEXT("MemRule_AaBCfD_Verbose", "Specific Lifetime"),
 		LOCTEXT("MemRule_AaBCfD_Desc", "Identifies allocations allocated between time A and time B and freed between time C and time D.\n(A ≤ a ≤ B ≤ C ≤ f ≤ D)")));
 
-	MemoryRules.Add(MakeShared<Insights::FMemoryRuleSpec>(
-		Insights::EMemoryRule::A_vs_B, 2,
-		LOCTEXT("MemRule_A_vs_B_Short", "*A* + *B*"),
-		LOCTEXT("MemRule_A_vs_B_Verbose", "Compare A vs. B"),
-		LOCTEXT("MemRule_A_vs_B_Desc", "Compares live allocations at time A with live allocations at time B.\n(*A* vs. *B*)")));
-
-	MemoryRules.Add(MakeShared<Insights::FMemoryRuleSpec>(
-		Insights::EMemoryRule::A_or_B, 2,
-		LOCTEXT("MemRule_A_or_B_Short", "*A* | *B*"),
-		LOCTEXT("MemRule_A_or_B_Verbose", "A or B"),
-		LOCTEXT("MemRule_A_or_B_Desc", "Identifies allocations live at time A or at time B.\n(a ≤ A ≤ f OR a ≤ B ≤ f)\n{*A*} U {*B*}")));
-
-	MemoryRules.Add(MakeShared<Insights::FMemoryRuleSpec>(
-		Insights::EMemoryRule::A_xor_B, 2,
-		LOCTEXT("MemRule_A_xor_B_Short", "*A* ^ *B*"),
-		LOCTEXT("MemRule_A_xor_B_Verbose", "A xor B"),
-		LOCTEXT("MemRule_A_xor_B_Desc", "Identifies allocations live either at time A or at time B (but not both).\n(a ≤ A ≤ f XOR a ≤ B ≤ f)\n({*A*} U {*B*}) \\ {*AB*}")));
+	//TODO
+	//MemoryRules.Add(MakeShared<Insights::FMemoryRuleSpec>(
+	//	ERule::A_vs_B, 2,
+	//	LOCTEXT("MemRule_A_vs_B_Short", "*A* + *B*"),
+	//	LOCTEXT("MemRule_A_vs_B_Verbose", "Compare A vs. B"),
+	//	LOCTEXT("MemRule_A_vs_B_Desc", "Compares live allocations at time A with live allocations at time B.\n(*A* vs. *B*)")));
+
+	//TODO
+	//MemoryRules.Add(MakeShared<Insights::FMemoryRuleSpec>(
+	//	ERule::A_or_B, 2,
+	//	LOCTEXT("MemRule_A_or_B_Short", "*A* | *B*"),
+	//	LOCTEXT("MemRule_A_or_B_Verbose", "A or B"),
+	//	LOCTEXT("MemRule_A_or_B_Desc", "Identifies allocations live at time A or at time B.\n(a ≤ A ≤ f OR a ≤ B ≤ f)\n{*A*} U {*B*}")));
+
+	//TODO
+	//MemoryRules.Add(MakeShared<Insights::FMemoryRuleSpec>(
+	//	ERule::A_xor_B, 2,
+	//	LOCTEXT("MemRule_A_xor_B_Short", "*A* ^ *B*"),
+	//	LOCTEXT("MemRule_A_xor_B_Verbose", "A xor B"),
+	//	LOCTEXT("MemRule_A_xor_B_Desc", "Identifies allocations live either at time A or at time B (but not both).\n(a ≤ A ≤ f XOR a ≤ B ≤ f)\n({*A*} U {*B*}) \\ {*AB*}")));
 
 	CurrentMemoryRule = MemoryRules[0];
 }
diff --git a/Engine/Source/Developer/TraceInsights/Private/Insights/MemoryProfiler/ViewModels/MemorySharedState.h b/Engine/Source/Developer/TraceInsights/Private/Insights/MemoryProfiler/ViewModels/MemorySharedState.h
--- a/Engine/Source/Developer/TraceInsights/Private/Insights/MemoryProfiler/ViewModels/MemorySharedState.h
+++ b/Engine/Source/Developer/TraceInsights/Private/Insights/MemoryProfiler/ViewModels/MemorySharedState.h
@@ -3,6 +3,7 @@
 #pragma once
 
 #include "CoreMinimal.h"
+#include "TraceServices/Model/AllocationsProvider.h"
 
 // Insights
 #include "Insights/ITimingViewExtender.h"
@@ -16,39 +17,19 @@ class STimingView;
 
 namespace Insights
 {
-	struct FReportConfig;
-	struct FReportTypeConfig;
-	struct FReportTypeGraphConfig;
-	class FMemoryTracker;
-
-// Known memory rules.
-// The enum uses the following naming convention:
-//     A, B, C, D = time markers
-//     a = time when "alloc" event occurs
-//     f = time when "free" event occurs (can be infinite)
-// Ex.: "AaBf" means "all memory allocations allocated between time A and time B and freed after time B".
-enum class EMemoryRule
-{
-	aAf,     // active allocs at A
-	afA,     // before
-	Aaf,     // after
-	aAfB,    // decline
-	AaBf,    // growth
-	AafB,    // short living allocs
-	aABf,    // long living allocs
-	AaBCf,   // memory leaks
-	AaBfC,   // limited lifetime
-	aABfC,   // decline of long living allocs
-	AaBCfD,  // specific lifetime
-	A_vs_B,  // compare A vs. B; {aAf} vs. {aBf}
-	A_or_B,  // live at A or at B; {aAf} U {aBf}
-	A_xor_B, // live either at A or at B; ({aAf} U {aBf}) \ {aABf}
-};
+
+struct FReportConfig;
+struct FReportTypeConfig;
+struct FReportTypeGraphConfig;
+class FMemoryTracker;
 
 class FMemoryRuleSpec
 {
 public:
-	FMemoryRuleSpec(EMemoryRule InValue, uint32 InNumTimeMarkers, const FText& InShortName, const FText& InVerboseName, const FText& InDescription)
+	typedef TraceServices::IAllocationsProvider::EQueryRule ERule;
+
+public:
+	FMemoryRuleSpec(ERule InValue, uint32 InNumTimeMarkers, const FText& InShortName, const FText& InVerboseName, const FText& InDescription)
 		: Value(InValue)
 		, NumTimeMarkers(InNumTimeMarkers)
 		, ShortName(InShortName)
@@ -56,18 +37,18 @@ public:
 		, Description(InDescription)
 	{}
 
-	EMemoryRule GetValue() const { return Value; }
+	ERule GetValue() const { return Value; }
 	uint32 GetNumTimeMarkers() const { return NumTimeMarkers; }
 	FText GetShortName() const { return ShortName; }
 	FText GetVerboseName() const { return VerboseName; }
 	FText GetDescription() const { return Description; }
 
 private:
-	EMemoryRule Value;     // ex.: EMemoryRule::AafB
+	ERule Value;           // ex.: ERule::AafB
 	uint32 NumTimeMarkers; // ex.: 2
-	FText ShortName;     // ex.: "A**B"
-	FText VerboseName;   // ex.: "Short Living Allocations"
-	FText Description;   // ex.: "Allocations allocated and freed between time A and time B (A <= a <= f <= B)."
+	FText ShortName;       // ex.: "A**B"
+	FText VerboseName;     // ex.: "Short Living Allocations"
+	FText Description;     // ex.: "Allocations allocated and freed between time A and time B (A <= a <= f <= B)."
 };
 
 } // namespace Insights
diff --git a/Engine/Source/Developer/TraceInsights/Private/Insights/MemoryProfiler/Widgets/SMemAllocTableTreeView.cpp b/Engine/Source/Developer/TraceInsights/Private/Insights/MemoryProfiler/Widgets/SMemAllocTableTreeView.cpp
--- a/Engine/Source/Developer/TraceInsights/Private/Insights/MemoryProfiler/Widgets/SMemAllocTableTreeView.cpp
+++ b/Engine/Source/Developer/TraceInsights/Private/Insights/MemoryProfiler/Widgets/SMemAllocTableTreeView.cpp
@@ -196,7 +196,6 @@ void SMemAllocTableTreeView::OnQueryInvalidated()
 
 void SMemAllocTableTreeView::StartQuery()
 {
-#if defined(UE_USE_ALLOCATIONS_PROVIDER)
 	check(Query == 0);
 
 	if (!Rule)
@@ -205,114 +204,41 @@ void SMemAllocTableTreeView::StartQuery()
 		return;
 	}
 
-	const IAllocationsProvider* AllocationsProvider = ReadAllocationsProvider(*Session.Get());
+	const TraceServices::IAllocationsProvider* AllocationsProvider = TraceServices::ReadAllocationsProvider(*Session.Get());
 	if (!AllocationsProvider)
 	{
 		UE_LOG(MemoryProfiler, Warning, TEXT("[MemAlloc] Invalid allocations provider!"));
 		return;
 	}
-	const IAllocationsProvider& Provider = *AllocationsProvider;
+	const TraceServices::IAllocationsProvider& Provider = *AllocationsProvider;
 
-	//TODO: Simpler API...
-	//QueryHandle StartQuery(IAllocationsProvider::EQueryRule Rule, double TimeA, double TimeB, double TimeC, double TimeD);
-
-	switch (Rule->GetValue())
-	{
-		case Insights::EMemoryRule::aAf:
-		{
-			Query = Provider.StartQuery(TimeMarkers[0], TimeMarkers[0], IAllocationsProvider::ECrosses::Or);
-			break;
-		}
-		case Insights::EMemoryRule::afA:
-		{
-			Query = Provider.StartQuery(0.0, TimeMarkers[0], IAllocationsProvider::ECrosses::None);
-			break;
-		}
-		case Insights::EMemoryRule::Aaf:
-		{
-			Query = Provider.StartQuery(TimeMarkers[0], std::numeric_limits<double>::infinity(), IAllocationsProvider::ECrosses::None);
-			break;
-		}
-		//case Insights::EMemoryRule::aAfB:
-		//{
-		//	Query = Provider.StartQuery(TimeMarkers[0], TimeMarkers[1], IAllocationsProvider::ECrosses::A);
-		//	break;
-		//}
-		//case Insights::EMemoryRule::AaBf:
-		//{
-		//	Query = Provider.StartQuery(TimeMarkers[0], TimeMarkers[1], IAllocationsProvider::ECrosses::B);
-		//	break;
-		//}
-		case Insights::EMemoryRule::AafB:
-		{
-			Query = Provider.StartQuery(TimeMarkers[0], TimeMarkers[1], IAllocationsProvider::ECrosses::None);
-			break;
-		}
-		case Insights::EMemoryRule::aABf:
-		{
-			Query = Provider.StartQuery(TimeMarkers[0], TimeMarkers[1], IAllocationsProvider::ECrosses::And);
-			break;
-		}
-		case Insights::EMemoryRule::AaBCf:
-		{
-			Query = Provider.StartQuery(TimeMarkers[1], TimeMarkers[2], IAllocationsProvider::ECrosses::And);
-			// needs additional filter Alloc.Start > A
-			break;
-		}
-		//case Insights::EMemoryRule::AaBfC:
-		//	TODO
-		//	break;
-		case Insights::EMemoryRule::aABfC:
-		{
-			Query = Provider.StartQuery(TimeMarkers[0], TimeMarkers[1], IAllocationsProvider::ECrosses::And);
-			// needs additional filter Alloc.End < C
-			break;
-		}
-		//case Insights::EMemoryRule::AaBCfD:
-		//	TODO
-		//	break;
-		case Insights::EMemoryRule::A_vs_B:
-		{
-			Query = Provider.StartQuery(TimeMarkers[0], TimeMarkers[1], IAllocationsProvider::ECrosses::Or);
-			break;
-		}
-		case Insights::EMemoryRule::A_or_B:
-		{
-			Query = Provider.StartQuery(TimeMarkers[0], TimeMarkers[1], IAllocationsProvider::ECrosses::Or);
-			break;
-		}
-		case Insights::EMemoryRule::A_xor_B:
-		{
-			Query = Provider.StartQuery(TimeMarkers[0], TimeMarkers[1], IAllocationsProvider::ECrosses::Xor);
-			break;
-		}
-	}
+	TraceServices::IAllocationsProvider::FQueryParams Params = { Rule->GetValue(), TimeMarkers[0], TimeMarkers[1], TimeMarkers[2], TimeMarkers[3] };
+	Query = Provider.StartQuery(Params);
 
 	if (Query == 0)
 	{
 		UE_LOG(MemoryProfiler, Error, TEXT("[MemAlloc] Unsupported query rule (%s)!"), *Rule->GetShortName().ToString());
 	}
-#endif // defined(UE_USE_ALLOCATIONS_PROVIDER)
 }
 
 ////////////////////////////////////////////////////////////////////////////////////////////////////
 
 void SMemAllocTableTreeView::UpdateQuery()
 {
-#if defined(UE_USE_ALLOCATIONS_PROVIDER)
 	if (Query != 0)
 	{
-		const IAllocationsProvider* AllocationsProvider = ReadAllocationsProvider(*Session.Get());
+		const TraceServices::IAllocationsProvider* AllocationsProvider = TraceServices::ReadAllocationsProvider(*Session.Get());
 		if (!AllocationsProvider)
 		{
 			UE_LOG(MemoryProfiler, Warning, TEXT("[MemAlloc] Invalid allocations provider!"));
 			return;
 		}
-		const IAllocationsProvider& Provider = *AllocationsProvider;
+		const TraceServices::IAllocationsProvider& Provider = *AllocationsProvider;
 
-		IAllocationsProvider::FQueryStatus Status = Provider.PollQuery(Query);
-		if (Status.Status <= IAllocationsProvider::EQueryStatus::Done)
+		TraceServices::IAllocationsProvider::FQueryStatus Status = Provider.PollQuery(Query);
+		if (Status.Status <= TraceServices::IAllocationsProvider::EQueryStatus::Done)
 		{
+			UE_LOG(MemoryProfiler, Log, TEXT("[MemAlloc] Query completed."));
 			Query = 0;
 			return;
 		}
@@ -330,7 +256,7 @@ void SMemAllocTableTreeView::UpdateQuery()
 
 			// Multiple 'pages' of results will be returned. No guarantees are made
 			// about the order of pages or the allocations they report.
-			IAllocationsProvider::QueryResult Result = Status.NextResult();
+			TraceServices::IAllocationsProvider::FQueryResult Result = Status.NextResult();
 			while (Result.IsValid())
 			{
 				++PageCount;
@@ -360,7 +286,7 @@ void SMemAllocTableTreeView::UpdateQuery()
 				Allocs.AddUninitialized(AllocCount);
 				for (uint32 AllocIndex = 0; AllocIndex < AllocCount; ++AllocIndex, ++AllocsDestIndex)
 				{
-					const IAllocationsProvider::FAllocation* Allocation = Result->Get(AllocIndex);
+					const TraceServices::IAllocationsProvider::FAllocation* Allocation = Result->Get(AllocIndex);
 					FMemoryAlloc& Alloc = Allocs[AllocsDestIndex];
 					Alloc.StartTime = Allocation->GetStartTime();
 					Alloc.EndTime = Allocation->GetEndTime();
@@ -391,17 +317,15 @@ void SMemAllocTableTreeView::UpdateQuery()
 			}
 		}
 	}
-#endif // defined(UE_USE_ALLOCATIONS_PROVIDER)
 }
 
 ////////////////////////////////////////////////////////////////////////////////////////////////////
 
 void SMemAllocTableTreeView::CancelQuery()
 {
-#if defined(UE_USE_ALLOCATIONS_PROVIDER)
 	if (Query != 0)
 	{
-		const IAllocationsProvider* AllocationsProvider = ReadAllocationsProvider(*Session.Get());
+		const TraceServices::IAllocationsProvider* AllocationsProvider = TraceServices::ReadAllocationsProvider(*Session.Get());
 		if (AllocationsProvider)
 		{
 			AllocationsProvider->CancelQuery(Query);
@@ -409,7 +333,6 @@ void SMemAllocTableTreeView::CancelQuery()
 		}
 		Query = 0;
 	}
-#endif // defined(UE_USE_ALLOCATIONS_PROVIDER)
 }
 
 ////////////////////////////////////////////////////////////////////////////////////////////////////
diff --git a/Engine/Source/Developer/TraceInsights/Private/Insights/MemoryProfiler/Widgets/SMemAllocTableTreeView.h b/Engine/Source/Developer/TraceInsights/Private/Insights/MemoryProfiler/Widgets/SMemAllocTableTreeView.h
--- a/Engine/Source/Developer/TraceInsights/Private/Insights/MemoryProfiler/Widgets/SMemAllocTableTreeView.h
+++ b/Engine/Source/Developer/TraceInsights/Private/Insights/MemoryProfiler/Widgets/SMemAllocTableTreeView.h
@@ -86,9 +86,7 @@ private:
 	int32 TabIndex = -1;
 	TSharedPtr<FMemoryRuleSpec> Rule = nullptr;
 	double TimeMarkers[4];
-#if defined(UE_USE_ALLOCATIONS_PROVIDER)
-	TraceServices::IAllocationsProvider::QueryHandle Query = 0;
-#endif
+	TraceServices::IAllocationsProvider::FQueryHandle Query = 0;
 
 	TSharedPtr<class ITableCellValueFormatter> OriginalLlmTagValueFormatter;
 };
diff --git a/Engine/Source/Developer/TraceServices/Private/Analyzers/AllocationsAnalysis.cpp b/Engine/Source/Developer/TraceServices/Private/Analyzers/AllocationsAnalysis.cpp
--- a/Engine/Source/Developer/TraceServices/Private/Analyzers/AllocationsAnalysis.cpp
+++ b/Engine/Source/Developer/TraceServices/Private/Analyzers/AllocationsAnalysis.cpp
@@ -1,38 +1,143 @@
 // Copyright Epic Games, Inc. All Rights Reserved.
 
 #include "AllocationsAnalysis.h"
-
-#if 0
+#include "Model/AllocationsProvider.h"
+#include "TraceServices/Model/AnalysisSession.h"
 
 namespace TraceServices
 {
 
-////////////////////////////////////////////////////////////////////////////////
-enum
+////////////////////////////////////////////////////////////////////////////////////////////////////
+
+FAllocationsAnalyzer::FAllocationsAnalyzer(IAnalysisSession& InSession, FAllocationsProvider& InAllocationsProvider)
+	: Session(InSession)
+	, AllocationsProvider(InAllocationsProvider)
 {
-	RouteId_Init,
-	RouteId_Alloc,
-	RouteId_Realloc,
-	RouteId_Free,
-};
-
-////////////////////////////////////////////////////////////////////////////////
-FAllocationsAnalyzer::FAllocationsAnalyzer(IAnalysisSession& Session)
+}
+
+////////////////////////////////////////////////////////////////////////////////////////////////////
+
+FAllocationsAnalyzer::~FAllocationsAnalyzer()
 {
 }
 
-////////////////////////////////////////////////////////////////////////////////
+////////////////////////////////////////////////////////////////////////////////////////////////////
+
 void FAllocationsAnalyzer::OnAnalysisBegin(const FOnAnalysisContext& Context)
 {
-	/*
 	FInterfaceBuilder& Builder = Context.InterfaceBuilder;
-	Builder.RouteEvent(RouteId_Init,    "Memory", "Init");
-	Builder.RouteEvent(RouteId_Alloc,   "Memory", "Alloc");
-	Builder.RouteEvent(RouteId_Realloc, "Memory", "Realloc");
-	Builder.RouteEvent(RouteId_Free,    "Memory", "Free");
-	*/
+
+	Builder.RouteEvent(RouteId_Init,       "Memory", "Init");
+	Builder.RouteEvent(RouteId_CoreAdd,    "Memory", "CoreAdd");
+	Builder.RouteEvent(RouteId_CoreRemove, "Memory", "CoreRemove");
+	Builder.RouteEvent(RouteId_Alloc,      "Memory", "Alloc");
+	Builder.RouteEvent(RouteId_Realloc,    "Memory", "Realloc");
+	Builder.RouteEvent(RouteId_Free,       "Memory", "Free");
 }
 
-} // namespace TraceServices
+////////////////////////////////////////////////////////////////////////////////////////////////////
+
+void FAllocationsAnalyzer::OnAnalysisEnd()
+{
+	//AllocationsProvider.DebugPrint();
+}
+
+////////////////////////////////////////////////////////////////////////////////////////////////////
+
+bool FAllocationsAnalyzer::OnEvent(uint16 RouteId, EStyle Style, const FOnEventContext& Context)
+{
+	const auto& EventData = Context.EventData;
+	switch (RouteId)
+	{
+		case RouteId_Init:
+		{
+			const double Time = GetCurrentTime();
+			const uint8 MinAlignment = EventData.GetValue<uint8>("MinAlignment");
+			const uint8 SizeShift = EventData.GetValue<uint8>("SizeShift");
+			const uint8 SummarySizeShift = EventData.GetValue<uint8>("SummarySizeShift");
+			const uint8 Mode = EventData.GetValue<uint8>("Mode");
 
-#endif // 0
+			FAllocationsProvider::FEditScopeLock _(AllocationsProvider);
+			AllocationsProvider.EditInit(Time, MinAlignment, SizeShift, SummarySizeShift, Mode);
+			break;
+		}
+		case RouteId_CoreAdd:
+		{
+			const double Time = GetCurrentTime();
+			const uint64 Owner = EventData.GetValue<uint64>("Owner");
+			const uint64 Base = EventData.GetValue<uint64>("Base");
+			const uint32 Size = EventData.GetValue<uint32>("Size");
+
+			FAllocationsProvider::FEditScopeLock _(AllocationsProvider);
+			AllocationsProvider.EditAddCore(Time, Owner, Base, Size);
+			break;
+		}
+		case RouteId_CoreRemove:
+		{
+			const double Time = GetCurrentTime();
+			const uint64 Owner = EventData.GetValue<uint64>("Owner");
+			const uint64 Base = EventData.GetValue<uint64>("Base");
+			const uint32 Size = EventData.GetValue<uint32>("Size");
+
+			FAllocationsProvider::FEditScopeLock _(AllocationsProvider);
+			AllocationsProvider.EditRemoveCore(Time, Owner, Base, Size);
+			break;
+		}
+		case RouteId_Alloc:
+		{
+			const double Time = GetCurrentTime();
+			const uint64 Owner = EventData.GetValue<uint64>("Owner");
+			const uint64 Address = EventData.GetValue<uint64>("Address");
+			const uint32 Size = EventData.GetValue<uint32>("Size");
+			const uint8 Alignment = EventData.GetValue<uint8>("Alignment");
+			const uint8 Waste = EventData.GetValue<uint8>("Waste");
+			uint32 Tag = 0;//TODO
+
+			FAllocationsProvider::FEditScopeLock _(AllocationsProvider);
+			AllocationsProvider.EditAlloc(Time, Owner, Address, Size, Alignment, Waste, Tag);
+			break;
+		}
+		case RouteId_Realloc:
+		{
+			// TODO: Can we have a struct mapping over the EventData?
+			//       something like: const auto& Ev = (const ReallocEvent&)EventData.Get();
+			const double Time = GetCurrentTime();
+			const uint64 Owner = EventData.GetValue<uint64>("Owner");
+			const uint64 FreeAddress = EventData.GetValue<uint64>("FreeAddress");
+			const uint64 Address = EventData.GetValue<uint64>("Address");
+			const uint32 Size = EventData.GetValue<uint32>("Size");
+			const uint8 Alignment = EventData.GetValue<uint8>("Alignment");
+			const uint8 Waste = EventData.GetValue<uint8>("Waste");
+
+			FAllocationsProvider::FEditScopeLock _(AllocationsProvider);
+			AllocationsProvider.EditRealloc(Time, Owner, FreeAddress, Address, Size, Alignment, Waste);
+			break;
+		}
+		case RouteId_Free:
+		{
+			const double Time = GetCurrentTime();
+			const uint64 Owner = EventData.GetValue<uint64>("Owner");
+			const uint64 FreeAddress = EventData.GetValue<uint64>("FreeAddress");
+
+			FAllocationsProvider::FEditScopeLock _(AllocationsProvider);
+			AllocationsProvider.EditFree(Time, Owner, FreeAddress);
+			break;
+		}
+	}
+
+	return true;
+}
+
+////////////////////////////////////////////////////////////////////////////////////////////////////
+
+double FAllocationsAnalyzer::GetCurrentTime() const
+{
+	// TODO: In order to avoid the session lock, either add a time param in allocation events
+	//       or add a handler for a specific time event.
+	FAnalysisSessionEditScope _(Session);
+	return Session.GetDurationSeconds();
+}
+
+////////////////////////////////////////////////////////////////////////////////////////////////////
+
+} // namespace TraceServices
diff --git a/Engine/Source/Developer/TraceServices/Private/Analyzers/AllocationsAnalysis.h b/Engine/Source/Developer/TraceServices/Private/Analyzers/AllocationsAnalysis.h
--- a/Engine/Source/Developer/TraceServices/Private/Analyzers/AllocationsAnalysis.h
+++ b/Engine/Source/Developer/TraceServices/Private/Analyzers/AllocationsAnalysis.h
@@ -4,22 +4,38 @@
 
 #include "Trace/Analyzer.h"
 
-#if 0
-
 namespace TraceServices
 {
 
 class IAnalysisSession;
 
-};
+class FAllocationsProvider;
 
-////////////////////////////////////////////////////////////////////////////////
-class FAllocationsAnalyzer
-	: public UE::Trace::IAnalyzer
+class FAllocationsAnalyzer : public UE::Trace::IAnalyzer
 {
+private:
+	enum : uint16
+	{
+		RouteId_Init,
+		RouteId_CoreAdd,
+		RouteId_CoreRemove,
+		RouteId_Alloc,
+		RouteId_Realloc,
+		RouteId_Free,
+	};
+
 public:
-					FAllocationsAnalyzer(IAnalysisSession& Session);
-	virtual void	OnAnalysisBegin(const FOnAnalysisContext& Context) override;
+	FAllocationsAnalyzer(IAnalysisSession& Session, FAllocationsProvider& AllocationsProvider);
+	~FAllocationsAnalyzer();
+	virtual void OnAnalysisBegin(const FOnAnalysisContext& Context) override;
+	virtual void OnAnalysisEnd() override;
+	virtual bool OnEvent(uint16 RouteId, EStyle Style, const FOnEventContext& Context) override;
+	double GetCurrentTime() const;
+
+private:
+	IAnalysisSession& Session;
+	FAllocationsProvider& AllocationsProvider;
+	uint64 EventCount = 0; // debug
 };
 
-#endif // 0
+} // namespace TraceServices
diff --git a/Engine/Source/Developer/TraceServices/Private/Common/Utils.cpp b/Engine/Source/Developer/TraceServices/Private/Common/Utils.cpp
--- a/Engine/Source/Developer/TraceServices/Private/Common/Utils.cpp
+++ b/Engine/Source/Developer/TraceServices/Private/Common/Utils.cpp
@@ -5,6 +5,8 @@
 #include "Templates/SharedPointer.h"
 #include "HAL/FileManager.h"
 
+DEFINE_LOG_CATEGORY(LogTraceServices);
+
 namespace TraceServices
 {
 
diff --git a/Engine/Source/Developer/TraceServices/Private/Common/Utils.h b/Engine/Source/Developer/TraceServices/Private/Common/Utils.h
--- a/Engine/Source/Developer/TraceServices/Private/Common/Utils.h
+++ b/Engine/Source/Developer/TraceServices/Private/Common/Utils.h
@@ -3,9 +3,12 @@
 #pragma once
 
 #include "CoreTypes.h"
+#include "Logging/LogMacros.h"
 #include "Trace/Analyzer.h"
 #include "TraceServices/Utils.h"
 
+DECLARE_LOG_CATEGORY_EXTERN(LogTraceServices, Log, All);
+
 namespace TraceServices {
 namespace FTraceAnalyzerUtils {
 
diff --git a/Engine/Source/Developer/TraceServices/Private/Model/AllocationsProvider.cpp b/Engine/Source/Developer/TraceServices/Private/Model/AllocationsProvider.cpp
--- a/Engine/Source/Developer/TraceServices/Private/Model/AllocationsProvider.cpp
+++ b/Engine/Source/Developer/TraceServices/Private/Model/AllocationsProvider.cpp
@@ -2,293 +2,506 @@
 
 #include "AllocationsProvider.h"
 
-#if defined(UE_USE_ALLOCATIONS_PROVIDER)
+#include "AllocationsQuery.h"
+#include "SbTree.h"
+#include "TraceServices/Containers/Allocators.h"
+
+#include <limits>
 
 namespace TraceServices
 {
 
-////////////////////////////////////////////////////////////////////////////////
-struct FAllocationItem
+////////////////////////////////////////////////////////////////////////////////////////////////////
+// FAllocationsProviderLock
+////////////////////////////////////////////////////////////////////////////////////////////////////
+
+thread_local FAllocationsProviderLock* GThreadCurrentAllocationsProviderLock;
+thread_local int32 GThreadCurrentReadAllocationsProviderLockCount;
+thread_local int32 GThreadCurrentWriteAllocationsProviderLockCount;
+
+void FAllocationsProviderLock::ReadAccessCheck() const
+{
+	checkf(GThreadCurrentAllocationsProviderLock == this && (GThreadCurrentReadAllocationsProviderLockCount > 0 || GThreadCurrentWriteAllocationsProviderLockCount > 0),
+		TEXT("Trying to READ from allocations provider outside of a READ scope"));
+}
+
+void FAllocationsProviderLock::WriteAccessCheck() const
+{
+	checkf(GThreadCurrentAllocationsProviderLock == this && GThreadCurrentWriteAllocationsProviderLockCount > 0,
+		TEXT("Trying to WRITE to allocations provider outside of an EDIT/WRITE scope"));
+}
+
+void FAllocationsProviderLock::BeginRead()
+{
+	check(!GThreadCurrentAllocationsProviderLock || GThreadCurrentAllocationsProviderLock == this);
+	checkf(GThreadCurrentWriteAllocationsProviderLockCount == 0, TEXT("Trying to lock allocations provider for READ while holding EDIT/WRITE access"));
+	if (GThreadCurrentReadAllocationsProviderLockCount++ == 0)
+	{
+		GThreadCurrentAllocationsProviderLock = this;
+		RWLock.ReadLock();
+	}
+}
+
+void FAllocationsProviderLock::EndRead()
+{
+	check(GThreadCurrentReadAllocationsProviderLockCount > 0);
+	if (--GThreadCurrentReadAllocationsProviderLockCount == 0)
+	{
+		RWLock.ReadUnlock();
+		GThreadCurrentAllocationsProviderLock = nullptr;
+	}
+}
+
+void FAllocationsProviderLock::BeginWrite()
+{
+	check(!GThreadCurrentAllocationsProviderLock || GThreadCurrentAllocationsProviderLock == this);
+	checkf(GThreadCurrentReadAllocationsProviderLockCount == 0, TEXT("Trying to lock allocations provider for EDIT/WRITE while holding READ access"));
+	if (GThreadCurrentWriteAllocationsProviderLockCount++ == 0)
+	{
+		GThreadCurrentAllocationsProviderLock = this;
+		RWLock.WriteLock();
+	}
+}
+
+void FAllocationsProviderLock::EndWrite()
 {
-	double	StartTime;
-	double	EndTime;
-	uint64	Address;
-	uint64	Size;
-	uint64	BacktraceId;
-	uint16	Tag;
-};
+	check(GThreadCurrentWriteAllocationsProviderLockCount > 0);
+	if (--GThreadCurrentWriteAllocationsProviderLockCount == 0)
+	{
+		RWLock.WriteUnlock();
+		GThreadCurrentAllocationsProviderLock = nullptr;
+	}
+}
+
+////////////////////////////////////////////////////////////////////////////////////////////////////
+// IAllocationsProvider::FAllocation
+////////////////////////////////////////////////////////////////////////////////////////////////////
 
-////////////////////////////////////////////////////////////////////////////////
 double IAllocationsProvider::FAllocation::GetStartTime() const
 {
 	const auto* Inner = (const FAllocationItem*)this;
 	return Inner->StartTime;
 }
 
-////////////////////////////////////////////////////////////////////////////////
+////////////////////////////////////////////////////////////////////////////////////////////////////
+
 double IAllocationsProvider::FAllocation::GetEndTime() const
 {
 	const auto* Inner = (const FAllocationItem*)this;
 	return Inner->EndTime;
 }
 
-////////////////////////////////////////////////////////////////////////////////
+////////////////////////////////////////////////////////////////////////////////////////////////////
+
 uint64 IAllocationsProvider::FAllocation::GetAddress() const
 {
 	const auto* Inner = (const FAllocationItem*)this;
 	return Inner->Address;
 }
 
-////////////////////////////////////////////////////////////////////////////////
-uint64 IAllocationsProvider::FAllocation::GetSize() const
+////////////////////////////////////////////////////////////////////////////////////////////////////
+
+uint32 IAllocationsProvider::FAllocation::GetSize() const
 {
 	const auto* Inner = (const FAllocationItem*)this;
 	return Inner->Size;
 }
 
-////////////////////////////////////////////////////////////////////////////////
-uint64 IAllocationsProvider::FAllocation::GetBacktraceId() const
+////////////////////////////////////////////////////////////////////////////////////////////////////
+
+uint8 IAllocationsProvider::FAllocation::GetAlignment() const
 {
 	const auto* Inner = (const FAllocationItem*)this;
-	return Inner->BacktraceId;
+	return Inner->Alignment;
 }
 
-////////////////////////////////////////////////////////////////////////////////
-uint32 IAllocationsProvider::FAllocation::GetTag() const
+////////////////////////////////////////////////////////////////////////////////////////////////////
+
+uint8 IAllocationsProvider::FAllocation::GetWaste() const
 {
 	const auto* Inner = (const FAllocationItem*)this;
-	return Inner->Tag;
+	return Inner->Waste;
 }
 
+////////////////////////////////////////////////////////////////////////////////////////////////////
 
-
-////////////////////////////////////////////////////////////////////////////////
-struct FAllocationsImpl
+uint64 IAllocationsProvider::FAllocation::GetBacktraceId() const
 {
-	void*				operator new (SIZE_T, uint32 NumItems);
-	FAllocationsImpl*	Next;
-	uint32				NumItems;
-	FAllocationItem		Items[];
-};
+	const auto* Inner = (const FAllocationItem*)this;
+	return Inner->Owner;
+}
+
+////////////////////////////////////////////////////////////////////////////////////////////////////
 
-////////////////////////////////////////////////////////////////////////////////
-void* FAllocationsImpl::operator new (SIZE_T Size, uint32 NumItems)
+uint32 IAllocationsProvider::FAllocation::GetTag() const
 {
-	Size += (NumItems * sizeof(Items[0]));
-	return FMemory::Malloc(Size, alignof(FAllocationsImpl));
+	const auto* Inner = (const FAllocationItem*)this;
+	return Inner->Tag;
 }
 
+////////////////////////////////////////////////////////////////////////////////////////////////////
+// IAllocationsProvider::FAllocations
+////////////////////////////////////////////////////////////////////////////////////////////////////
 
-
-////////////////////////////////////////////////////////////////////////////////
 void IAllocationsProvider::FAllocations::operator delete (void* Address)
 {
-	FMemory::Free(Address);
+	auto* Inner = (const FAllocationsImpl*)Address;
+	delete Inner;
 }
 
-////////////////////////////////////////////////////////////////////////////////
+////////////////////////////////////////////////////////////////////////////////////////////////////
+
 uint32 IAllocationsProvider::FAllocations::Num() const
 {
-	auto* Inner = (FAllocationsImpl*)this;
-	return Inner->NumItems;
+	auto* Inner = (const FAllocationsImpl*)this;
+	return (uint32)(Inner->Items.Num());
 }
 
-////////////////////////////////////////////////////////////////////////////////
+////////////////////////////////////////////////////////////////////////////////////////////////////
+
 const IAllocationsProvider::FAllocation* IAllocationsProvider::FAllocations::Get(uint32 Index) const
 {
-	if (Index >= Num())
+	checkSlow(Index < Num());
+
+	auto* Inner = (const FAllocationsImpl*)this;
+	return (const FAllocation*)(Inner->Items[Index]);
+}
+
+////////////////////////////////////////////////////////////////////////////////////////////////////
+// IAllocationsProvider::FQueryStatus
+////////////////////////////////////////////////////////////////////////////////////////////////////
+
+IAllocationsProvider::FQueryResult IAllocationsProvider::FQueryStatus::NextResult() const
+{
+	auto* Inner = (FAllocationsImpl*)Handle;
+	if (Inner == nullptr)
 	{
 		return nullptr;
 	}
 
-	auto* Inner = (FAllocationsImpl*)this;
-	return (FAllocation*)(Inner->Items + Index);
+	Handle = UPTRINT(Inner->Next);
+
+	auto* Ret = (FAllocations*)Inner;
+	return FQueryResult(Ret);
 }
 
+////////////////////////////////////////////////////////////////////////////////////////////////////
+// FAllocationsProvider
+////////////////////////////////////////////////////////////////////////////////////////////////////
 
-namespace TEMP
+FAllocationsProvider::FAllocationsProvider(ILinearAllocator& InAllocator)
+	: Allocator(InAllocator)
+	, SbTree(nullptr)
 {
+	const uint32 ColumnShift = 17; // 1<<17 = 128K
+	SbTree = new FSbTree(InAllocator, ColumnShift);
+}
 
-////////////////////////////////////////////////////////////////////////////////
-using ECrosses				= IAllocationsProvider::ECrosses;
-using EQueryStatus			= IAllocationsProvider::EQueryStatus;
-using FQueryStatus			= IAllocationsProvider::FQueryStatus;
-using FAllocation			= IAllocationsProvider::FAllocation;
-static const uint64 Mlcg	= 0x106689d45497fdb5ull;
+////////////////////////////////////////////////////////////////////////////////////////////////////
 
-////////////////////////////////////////////////////////////////////////////////
-class FMockQuery
+FAllocationsProvider::~FAllocationsProvider()
 {
-public:
-						FMockQuery(double TimeA, double TimeB, ECrosses Crosses);
-	void				Cancel() { delete this; }
-	FQueryStatus		Poll();
+	if (SbTree)
+	{
+		delete SbTree;
+		SbTree = nullptr;
+	}
+}
 
-private:
-	bool				NextBlock();
-	uint64				NextSeed(uint64 Pow2Max=0);
-	uint64				Seed;
-	uint32				BlockIndex;
-	uint32				BlockEnd;
-	FAllocationsImpl*	BlockList;
-};
+////////////////////////////////////////////////////////////////////////////////////////////////////
 
-////////////////////////////////////////////////////////////////////////////////
-FMockQuery::FMockQuery(double TimeA, double TimeB, ECrosses Crosses)
+FName FAllocationsProvider::GetName()
 {
-	BlockIndex = uint32(TimeA / 5.0);
-	BlockEnd = uint32(TimeB / 5.0) + 1;
-	BlockList = nullptr;
+	static FName Name(TEXT("AllocationsProvider"));
+	return Name;
 }
 
-////////////////////////////////////////////////////////////////////////////////
-bool FMockQuery::NextBlock()
+////////////////////////////////////////////////////////////////////////////////////////////////////
+
+void FAllocationsProvider::EditInit(double InTime, uint8 InMinAlignment, uint8 InSizeShift, uint8 InSummarySizeShift, uint8 InMode)
 {
-	if (BlockIndex >= BlockEnd)
+	if (bInitialized)
 	{
-		return false;
+		// error: already initialized
+		return;
 	}
 
-	Seed = (BlockIndex + 1) * Mlcg;
+	InitTime = InTime;
+	MinAlignment = InMinAlignment;
+	SizeShift = InSizeShift;
+	SummarySizeShift = InSummarySizeShift;
+	Mode = InMode;
 
-	// Build a set of results
-	FAllocationsImpl* ResultList = nullptr;
-	uint32 NumResults = (NextSeed(32) - 16) + 32;
-	for (uint32 i = 0; i < NumResults; ++i)
-	{
-		// How many allocations shall we put in this result page?
-		uint32 NumAllocs = NextSeed(1 << 20);
-		FAllocationsImpl* Result = new (NumAllocs) FAllocationsImpl();
-		for (uint32 j = 0; j < NumAllocs; ++j)
-		{
-			uint64 Address = 0x200000000000 | (NextSeed(1ull << 40) & ~0x0f);
-			double EndTime = double(BlockIndex * 5.0);
-
-			FAllocationItem& Allocation = Result->Items[j];
-			Allocation.StartTime = (EndTime * double(NextSeed(2048)) / 2048.0);
-			Allocation.EndTime = EndTime;
-			Allocation.Address = Address;
-			Allocation.Size = (NextSeed(512) - 256) + 768;
-			Allocation.BacktraceId = NextSeed(1ull << 48);
-			Allocation.Tag = NextSeed(4);
-		}
+	bInitialized = true;
+}
 
-		Result->NumItems = NumAllocs;
-		Result->Next = ResultList;
-		ResultList = Result;
-	}
+////////////////////////////////////////////////////////////////////////////////////////////////////
 
-	BlockList = ResultList;
+void FAllocationsProvider::EditAddCore(double Time, uint64 Owner, uint64 Base, uint32 Size)
+{
+	Lock.WriteAccessCheck();
 
-	++BlockIndex;
-	return true;
+	if (!bInitialized)
+	{
+		return;
+	}
+
+	//TODO
 }
 
-////////////////////////////////////////////////////////////////////////////////
-uint64 FMockQuery::NextSeed(uint64 Pow2Max)
+////////////////////////////////////////////////////////////////////////////////////////////////////
+
+void FAllocationsProvider::EditRemoveCore(double Time, uint64 Owner, uint64 Base, uint32 Size)
 {
-	Seed *= Mlcg;
-	return (Seed >> 8) & (Pow2Max - 1);
+	Lock.WriteAccessCheck();
+
+	if (!bInitialized)
+	{
+		return;
+	}
+
+	//TODO
 }
 
-////////////////////////////////////////////////////////////////////////////////
-FQueryStatus FMockQuery::Poll()
+////////////////////////////////////////////////////////////////////////////////////////////////////
+
+void FAllocationsProvider::EditAlloc(double Time, uint64 Owner, uint64 Address, uint32 Size, uint8 Alignment, uint8 Waste, uint32 Tag)
 {
-	FQueryStatus Status = {};
-	Status.Handle = 0;
+	Lock.WriteAccessCheck();
 
-	// Is there nothing more to do?
-	if (BlockList == nullptr && !NextBlock())
+	if (!bInitialized)
 	{
-		Status.Status = EQueryStatus::Done;
-		return Status;
+		return;
 	}
 
-	// Simulate there not being any results yet.
-	bool bNotReady = (NextSeed(8) > 6);
-	if (bNotReady)
+	FAllocationItem* AllocationPtr = LiveAllocs.Find(Address);
+	if (AllocationPtr)
 	{
-		Status.Status = EQueryStatus::Working;
-		return Status;
+		++AllocErrors;
 	}
-
-	// Slice a bit off the linked list.
-	FAllocationsImpl* Head = BlockList;
-	FAllocationsImpl* Tail = BlockList;
-	for (uint32 i = NextSeed(32); i-- && Tail->Next != nullptr;)
+	else
 	{
-		Tail = Tail->Next;
+		SbTree->SetTimeForEvent(EventIndex, Time);
+
+		FAllocationItem Allocation =
+		{
+			EventIndex,
+			(uint32)-1,
+			Time,
+			std::numeric_limits<double>::infinity(),
+			Owner,
+			Address,
+			Size,
+			Alignment,
+			Waste,
+			static_cast<uint16>(Tag),
+			0, // Reserverd1
+			0, // Reserverd2
+		};
+		LiveAllocs.Add(Address, Allocation);
+
+		//TODO: Cache the last added alloc, as there is ~10% chance to be freed in the next event.
+
+		if (LiveAllocs.Num() > MaxLiveAllocCount)
+		{
+			MaxLiveAllocCount = LiveAllocs.Num();
+		}
+
+		UpdateHistogramByAllocSize(Size, Alignment, Waste);
 	}
-	BlockList = Tail->Next;
-	Tail->Next = nullptr;
 
-	Status.Status = EQueryStatus::Available;
-	Status.Handle = UPTRINT(Head);
-	return Status;
+	++AllocCount;
+	++EventIndex;
 }
 
-} // namespace TEMP
+////////////////////////////////////////////////////////////////////////////////////////////////////
 
+void FAllocationsProvider::EditRealloc(double Time, uint64 Owner, uint64 FreeAddress, uint64 Address, uint32 Size, uint8 Alignment, uint8 Waste)
+{
+	Lock.WriteAccessCheck();
 
+	if (!bInitialized)
+	{
+		return;
+	}
 
-////////////////////////////////////////////////////////////////////////////////
-IAllocationsProvider::QueryResult IAllocationsProvider::FQueryStatus::NextResult() const
-{
-	auto* Inner = (FAllocationsImpl*)Handle;
-	if (Inner == nullptr)
+	SbTree->SetTimeForEvent(EventIndex, Time);
+
+	FAllocationItem* AllocationPtr = LiveAllocs.Find(FreeAddress);
+	if (!AllocationPtr)
 	{
-		return nullptr;
+		++ReallocErrors;
 	}
+	else
+	{
+		check(EventIndex > AllocationPtr->StartEventIndex);
+		AllocationPtr->EndEventIndex = EventIndex;
+		AllocationPtr->EndTime = Time;
 
-	Handle = UPTRINT(Inner->Next);
+		SbTree->AddAlloc(AllocationPtr);
 
-	auto* Ret = (FAllocations*)Inner;
-	return QueryResult(Ret);
+		uint32 EventDistance = AllocationPtr->EndEventIndex - AllocationPtr->StartEventIndex;
+		UpdateHistogramByEventDistance(EventDistance);
+
+		if (FreeAddress == Address)
+		{
+			AllocationPtr->StartEventIndex = EventIndex;
+			AllocationPtr->EndEventIndex = -1;
+			AllocationPtr->StartTime = Time;
+			AllocationPtr->EndTime = std::numeric_limits<double>::infinity();
+			AllocationPtr->Owner = Owner;
+			AllocationPtr->Address = Address;
+			AllocationPtr->Size = Size;
+			AllocationPtr->Alignment = Alignment;
+			AllocationPtr->Waste = Waste;
+		}
+		else
+		{
+			uint16 Tag = AllocationPtr->Tag;
+			LiveAllocs.Remove(FreeAddress);
+
+			FAllocationItem Allocation =
+			{
+				EventIndex,
+				(uint32)-1,
+				Time,
+				std::numeric_limits<double>::infinity(),
+				Owner,
+				Address,
+				Size,
+				Alignment,
+				Waste,
+				Tag,
+				0, // Reserverd1
+				0, // Reserverd2
+			};
+			LiveAllocs.Add(Address, Allocation);
+		}
+
+		UpdateHistogramByAllocSize(Size, Alignment, Waste);
+	}
+
+	++AllocCount;
+	++EventIndex;
 }
 
+////////////////////////////////////////////////////////////////////////////////////////////////////
+
+void FAllocationsProvider::EditFree(double Time, uint64 Owner, uint64 FreeAddress)
+{
+	Lock.WriteAccessCheck();
+
+	if (!bInitialized)
+	{
+		return;
+	}
+
+	SbTree->SetTimeForEvent(EventIndex, Time);
+
+	FAllocationItem* AllocationPtr = LiveAllocs.Find(FreeAddress);
+	if (!AllocationPtr)
+	{
+		++FreeErrors;
+	}
+	else
+	{
+		check(EventIndex > AllocationPtr->StartEventIndex);
+		AllocationPtr->EndEventIndex = EventIndex;
+		AllocationPtr->EndTime = Time;
 
+		SbTree->AddAlloc(AllocationPtr);
 
-////////////////////////////////////////////////////////////////////////////////
-FAllocationsProvider::FAllocationsProvider()
+		uint32 EventDistance = AllocationPtr->EndEventIndex - AllocationPtr->StartEventIndex;
+		UpdateHistogramByEventDistance(EventDistance);
+
+		LiveAllocs.Remove(FreeAddress);
+	}
+
+	++EventIndex;
+}
+
+////////////////////////////////////////////////////////////////////////////////////////////////////
+
+void FAllocationsProvider::UpdateHistogramByAllocSize(uint32 Size, uint8 Alignment, uint8 Waste)
 {
+	if (Size > MaxAllocSize)
+	{
+		MaxAllocSize = Size;
+	}
+
+	// HistogramIndex : Value Range
+	// 0 : [0]
+	// 1 : [1]
+	// 2 : [2 .. 3]
+	// 3 : [4 .. 7]
+	// ...
+	// i : [2^(i-1) .. 2^i-1], i > 0
+	// ...
+	// 32 : [2^31 .. 2^32 - 1]
+	uint32 HistogramIndexPow2 = 32 - FMath::CountLeadingZeros(Size);
+	++AllocSizeHistogramPow2[HistogramIndexPow2];
 }
 
-////////////////////////////////////////////////////////////////////////////////
-FAllocationsProvider::~FAllocationsProvider()
+////////////////////////////////////////////////////////////////////////////////////////////////////
+
+void FAllocationsProvider::UpdateHistogramByEventDistance(uint32 EventDistance)
 {
+	if (EventDistance > MaxEventDistance)
+	{
+		MaxEventDistance = EventDistance;
+	}
+
+	// HistogramIndex : Value Range
+	// 0 : [0 .. 1]
+	// 1 : [2 .. 3]
+	// 2 : [4 .. 7]
+	// 3 : [8 .. 15]
+	// ...
+	// i : [2^i .. 2^(i+1)-1]
+	// ...
+	// 63 : [2^63 .. 2^64-1]
+	uint32 HistogramIndexPow2 = (EventDistance > 0) ? 63 - FMath::CountLeadingZeros64(EventDistance) : 0;
+	++EventDistanceHistogramPow2[HistogramIndexPow2];
 }
 
-////////////////////////////////////////////////////////////////////////////////
-FName FAllocationsProvider::GetName() const
+////////////////////////////////////////////////////////////////////////////////////////////////////
+
+void FAllocationsProvider::DebugPrint() const
 {
-	static FName Name(TEXT("AllocationsProvider"));
-	return Name;
+	SbTree->DebugPrint();
 }
 
-////////////////////////////////////////////////////////////////////////////////
-IAllocationsProvider::QueryHandle FAllocationsProvider::StartQuery(
-	double TimeA,
-	double TimeB,
-	ECrosses Crosses)
+////////////////////////////////////////////////////////////////////////////////////////////////////
+
+IAllocationsProvider::FQueryHandle FAllocationsProvider::StartQuery(const IAllocationsProvider::FQueryParams& Params) const
 {
-	auto* Inner = new TEMP::FMockQuery(TimeA, TimeB, Crosses);
-	return IAllocationsProvider::QueryHandle(Inner);
+	auto* Inner = new FAllocationsQuery(*this, Params);
+	return IAllocationsProvider::FQueryHandle(Inner);
 }
 
-////////////////////////////////////////////////////////////////////////////////
-void FAllocationsProvider::CancelQuery(QueryHandle Query)
+////////////////////////////////////////////////////////////////////////////////////////////////////
+
+void FAllocationsProvider::CancelQuery(FQueryHandle Query) const
 {
-	auto* Inner = (TEMP::FMockQuery*)Query;
+	auto* Inner = (FAllocationsQuery*)Query;
 	return Inner->Cancel();
 }
 
-////////////////////////////////////////////////////////////////////////////////
-IAllocationsProvider::QueryStatus FAllocationsProvider::PollQuery(QueryHandle Query)
+////////////////////////////////////////////////////////////////////////////////////////////////////
+
+const IAllocationsProvider::FQueryStatus FAllocationsProvider::PollQuery(FQueryHandle Query) const
 {
-	auto* Inner = (TEMP::FMockQuery*)Query;
+	auto* Inner = (FAllocationsQuery*)Query;
 	return Inner->Poll();
 }
 
-} // namespace TraceServices
+////////////////////////////////////////////////////////////////////////////////////////////////////
+
+const IAllocationsProvider* ReadAllocationsProvider(const IAnalysisSession& Session)
+{
+	return Session.ReadProvider<IAllocationsProvider>(FAllocationsProvider::GetName());
+}
 
-#endif // UE_USE_ALLOCATIONS_PROVIDER
+////////////////////////////////////////////////////////////////////////////////////////////////////
+
+} // namespace TraceServices
diff --git a/Engine/Source/Developer/TraceServices/Private/Model/AllocationsProvider.h b/Engine/Source/Developer/TraceServices/Private/Model/AllocationsProvider.h
--- a/Engine/Source/Developer/TraceServices/Private/Model/AllocationsProvider.h
+++ b/Engine/Source/Developer/TraceServices/Private/Model/AllocationsProvider.h
@@ -3,26 +3,169 @@
 #pragma once
 
 #include "TraceServices/Model/AllocationsProvider.h"
-#include "UObject/NameTypes.h"
-
-#if defined(UE_USE_ALLOCATIONS_PROVIDER)
 
 namespace TraceServices
 {
 
-////////////////////////////////////////////////////////////////////////////////
-class FAllocationsProvider
-	: public IAllocationsProvider
+class ILinearAllocator;
+class FSbTree;
+
+////////////////////////////////////////////////////////////////////////////////////////////////////
+
+class FAllocationsProviderLock
 {
 public:
-						FAllocationsProvider();
-	virtual				~FAllocationsProvider();
-	FName				GetName() const;
-	virtual QueryHandle	StartQuery(double TimeA, double TimeB, ECrosses Crosses) override;
-	virtual void		CancelQuery(QueryHandle Query) override;
-	virtual QueryStatus	PollQuery(QueryHandle Query) override;
+	void ReadAccessCheck() const;
+	void WriteAccessCheck() const;
+
+	void BeginRead();
+	void EndRead();
+
+	void BeginWrite();
+	void EndWrite();
+
+private:
+	FRWLock RWLock;
 };
 
-} // namespace TraceServices
+////////////////////////////////////////////////////////////////////////////////////////////////////
+
+struct FAllocationItem
+{
+	uint32 StartEventIndex;
+	uint32 EndEventIndex;
+	double StartTime;
+	double EndTime;
+	uint64 Owner;
+	uint64 Address;
+	//uint64 SizeAlignmentWaste; // 48 bits (Size) + 8 bits (Alignment) + 8 bits (Waste)
+	uint32 Size;
+	uint8 Alignment;
+	uint8 Waste;
+	uint16 Tag;
+	uint64 Reserved1;
+	uint64 Reserved2;
+};
+
+static_assert(sizeof(FAllocationItem) == 64, "struct FAllocationItem needs packing");
+
+////////////////////////////////////////////////////////////////////////////////////////////////////
+#if 0
+struct FAllocationCoreItem
+{
+	uint32 StartEventIndex;
+	uint32 EndEventIndex;
+	double StartTime;
+	double EndTime;
+	uint64 Owner;
+	uint64 Base;
+	uint32 Size;
+};
+#endif
+////////////////////////////////////////////////////////////////////////////////////////////////////
+
+class FAllocationsProvider : public IAllocationsProvider
+{
+public:
+	struct FEditScopeLock
+	{
+		FEditScopeLock(FAllocationsProvider& InAllocationsProvider)
+			: AllocationsProvider(InAllocationsProvider)
+		{
+			AllocationsProvider.BeginEdit();
+		}
+
+		~FEditScopeLock()
+		{
+			AllocationsProvider.EndEdit();
+		}
+
+		FAllocationsProvider& AllocationsProvider;
+	};
+
+	struct FReadScopeLock
+	{
+		FReadScopeLock(FAllocationsProvider& InAllocationsProvider)
+			: AllocationsProvider(InAllocationsProvider)
+		{
+			AllocationsProvider.BeginRead();
+		}
+
+		~FReadScopeLock()
+		{
+			AllocationsProvider.EndRead();
+		}
+
+		FAllocationsProvider& AllocationsProvider;
+	};
+
+public:
+	FAllocationsProvider(ILinearAllocator& InAllocator);
+	virtual ~FAllocationsProvider();
+
+	static FName GetName();
+
+	virtual bool IsInitialized() const override { /*Lock.ReadAccessCheck();*/ return bInitialized; }
+
+	virtual FQueryHandle StartQuery(const FQueryParams& Params) const override;
+	virtual void CancelQuery(FQueryHandle Query) const override;
+	virtual const FQueryStatus PollQuery(FQueryHandle Query) const override;
+
+	void BeginEdit() { Lock.BeginWrite(); }
+	void EndEdit() { Lock.EndWrite(); }
+	void EditAccessCheck() const { return Lock.WriteAccessCheck(); }
+
+	void BeginRead() const { Lock.BeginRead(); }
+	void EndRead() const { Lock.EndRead(); }
+	void ReadAccessCheck() const { return Lock.ReadAccessCheck(); }
+
+	void EditInit(double Time, uint8 MinAlignment, uint8 SizeShift, uint8 SummarySizeShift, uint8 Mode);
+	void EditAddCore(double Time, uint64 Owner, uint64 Base, uint32 Size);
+	void EditRemoveCore(double Time, uint64 Owner, uint64 Base, uint32 Size);
+	void EditAlloc(double Time, uint64 Owner, uint64 Address, uint32 Size, uint8 Alignment, uint8 Waste, uint32 Tag);
+	void EditRealloc(double Time, uint64 Owner, uint64 FreeAddress, uint64 Address, uint32 Size, uint8 Alignment, uint8 Waste);
+	void EditFree(double Time, uint64 Owner, uint64 FreeAddress);
+
+	void DebugPrint() const;
+
+	const FSbTree* GetSbTree() const { return SbTree; }
+
+private:
+	void UpdateHistogramByAllocSize(uint32 Size, uint8 Alignment, uint8 Waste);
+	void UpdateHistogramByEventDistance(uint32 EventDistance);
+
+private:
+	ILinearAllocator& Allocator;
+
+	mutable FAllocationsProviderLock Lock;
 
-#endif // UE_USE_ALLOCATIONS_PROVIDER
+	double InitTime = 0;
+	uint8 MinAlignment = 0;
+	uint8 SizeShift = 0;
+	uint8 SummarySizeShift = 0;
+	uint8 Mode = 0;
+	bool bInitialized = false;
+
+	uint32 EventIndex = 0;
+
+	uint64 AllocCount = 0;
+
+	uint64 MaxLiveAllocCount = 0;
+	TMap<uint64, FAllocationItem> LiveAllocs;
+
+	uint64 AllocErrors = 0;
+	uint64 ReallocErrors = 0;
+	uint64 FreeErrors = 0;
+
+	uint32 MaxAllocSize = 0;
+	uint32 AllocSizeHistogramPow2[33] = { 0 };
+
+	uint32 MaxEventDistance = 0;
+	uint32 EventDistanceHistogramPow2[65] = { 0 };
+
+	FSbTree* SbTree;
+};
+
+////////////////////////////////////////////////////////////////////////////////////////////////////
+
+} // namespace TraceServices
diff --git a/Engine/Source/Developer/TraceServices/Private/Model/AllocationsQuery.cpp b/Engine/Source/Developer/TraceServices/Private/Model/AllocationsQuery.cpp
new file mode 100644
--- /dev/null
+++ b/Engine/Source/Developer/TraceServices/Private/Model/AllocationsQuery.cpp
@@ -0,0 +1,164 @@
+// Copyright Epic Games, Inc. All Rights Reserved.
+
+#include "AllocationsQuery.h"
+
+#include "SbTree.h"
+#include "Common/Utils.h"
+#include "TraceServices/Containers/Allocators.h"
+
+#include <limits>
+
+namespace TraceServices
+{
+
+////////////////////////////////////////////////////////////////////////////////////////////////////
+// FAllocationsImpl
+////////////////////////////////////////////////////////////////////////////////////////////////////
+
+FAllocationsImpl::FAllocationsImpl(uint32 NumItems)
+	: Next(nullptr)
+{
+}
+
+////////////////////////////////////////////////////////////////////////////////////////////////////
+
+FAllocationsImpl::~FAllocationsImpl()
+{
+}
+
+////////////////////////////////////////////////////////////////////////////////////////////////////
+// FAllocationsQueryAsyncTask
+////////////////////////////////////////////////////////////////////////////////////////////////////
+
+class FAllocationsQueryAsyncTask
+{
+public:
+	FAllocationsQueryAsyncTask(FAllocationsQuery* InQuery)
+	{
+		QueryPtr = InQuery;
+	}
+
+	FORCEINLINE TStatId GetStatId() const { RETURN_QUICK_DECLARE_CYCLE_STAT(FAllocationsQueryAsyncTask, STATGROUP_TaskGraphTasks); }
+	ENamedThreads::Type GetDesiredThread() { return ENamedThreads::Type::AnyThread; }
+	static ESubsequentsMode::Type GetSubsequentsMode() { return ESubsequentsMode::FireAndForget; }
+
+	void DoTask(ENamedThreads::Type CurrentThread, const FGraphEventRef& MyCompletionGraphEvent)
+	{
+		if (QueryPtr)
+		{
+			QueryPtr->Run();
+		}
+	}
+
+private:
+	FAllocationsQuery* QueryPtr = nullptr;
+};
+
+////////////////////////////////////////////////////////////////////////////////////////////////////
+// FAllocationsQuery
+////////////////////////////////////////////////////////////////////////////////////////////////////
+
+FAllocationsQuery::FAllocationsQuery(const FAllocationsProvider& InProvider, const IAllocationsProvider::FQueryParams& InParams)
+	: Provider(InProvider)
+	, Params(InParams)
+	, IsWorking(true)
+	, IsCanceling(false)
+{
+	// Start the async task.
+	CompletedEvent = TGraphTask<FAllocationsQueryAsyncTask>::CreateTask().ConstructAndDispatchWhenReady(this);
+}
+
+////////////////////////////////////////////////////////////////////////////////////////////////////
+
+void FAllocationsQuery::Cancel()
+{
+	if (CompletedEvent.IsValid())
+	{
+		// Cancel the async task.
+		IsCanceling = true;
+		CompletedEvent->Wait();
+	}
+	delete this;
+}
+
+////////////////////////////////////////////////////////////////////////////////////////////////////
+
+IAllocationsProvider::FQueryStatus FAllocationsQuery::Poll()
+{
+	IAllocationsProvider::FQueryStatus Status = {};
+	Status.Handle = 0;
+
+	FAllocationsImpl* Allocations = nullptr;
+	if (Results.Dequeue(Allocations))
+	{
+		Status.Status = IAllocationsProvider::EQueryStatus::Available;
+		Status.Handle = UPTRINT(Allocations);
+		return Status;
+	}
+
+	if (IsWorking)
+	{
+		Status.Status = IAllocationsProvider::EQueryStatus::Working;
+		return Status;
+	}
+	else
+	{
+		Status.Status = IAllocationsProvider::EQueryStatus::Done;
+		return Status;
+	}
+}
+
+////////////////////////////////////////////////////////////////////////////////////////////////////
+
+void FAllocationsQuery::Run()
+{
+	// Note: This function is called from the async task (FAllocationsQueryAsyncTask),
+	//       so no assumption can be made on which thread is running.
+
+	uint64 StartTime = FPlatformTime::Cycles64();
+
+	uint32 CellCount = 0;
+	uint32 TotalAllocationCount = 0;
+
+	Provider.BeginRead();
+
+	const FSbTree* SbTree = Provider.GetSbTree();
+	if (SbTree)
+	{
+		TArray<const FSbTreeCell*> Cells;
+
+		if (!IsCanceling)
+		{
+			SbTree->Query(Cells, Params);
+			CellCount += Cells.Num();
+		}
+
+		for (const FSbTreeCell* Cell : Cells)
+		{
+			if (IsCanceling)
+			{
+				break;
+			}
+			FAllocationsImpl* Result = new FAllocationsImpl(0);
+			Cell->Query(Result->Items, Params);
+			TotalAllocationCount += Result->Items.Num();
+			Results.Enqueue(Result);
+		}
+	}
+
+	Provider.EndRead();
+
+	IsWorking = false;
+
+	uint64 EndTime = FPlatformTime::Cycles64();
+	const double TotalTime = static_cast<double>(EndTime - StartTime) * FPlatformTime::GetSecondsPerCycle64();
+	if (TotalTime > 0.1)
+	{
+		UE_LOG(LogTraceServices, Log, TEXT("[MemAlloc] Allocations query completed in %.3fs (%u cells, %u allocations)"),
+			TotalTime, CellCount, TotalAllocationCount);
+	}
+}
+
+////////////////////////////////////////////////////////////////////////////////////////////////////
+
+} // namespace TraceServices
diff --git a/Engine/Source/Developer/TraceServices/Private/Model/AllocationsQuery.h b/Engine/Source/Developer/TraceServices/Private/Model/AllocationsQuery.h
new file mode 100644
--- /dev/null
+++ b/Engine/Source/Developer/TraceServices/Private/Model/AllocationsQuery.h
@@ -0,0 +1,51 @@
+// Copyright Epic Games, Inc. All Rights Reserved.
+
+#pragma once
+
+#include "AllocationsProvider.h"
+#include "Async/TaskGraphInterfaces.h"
+#include "HAL/PlatformAtomics.h"
+
+#include <atomic>
+
+namespace TraceServices
+{
+
+class ILinearAllocator;
+
+struct FAllocationsImpl
+{
+	FAllocationsImpl(uint32 NumItems);
+	~FAllocationsImpl();
+
+	FAllocationsImpl* Next;
+	TArray<const FAllocationItem*> Items;
+};
+
+class FAllocationsQuery
+{
+	friend class FAllocationsQueryAsyncTask;
+
+public:
+	FAllocationsQuery(const FAllocationsProvider& InProvider, const IAllocationsProvider::FQueryParams& InParams);
+
+	void Cancel();
+	IAllocationsProvider::FQueryStatus Poll();
+
+private:
+	void Run();
+
+private:
+	const FAllocationsProvider& Provider;
+
+	IAllocationsProvider::FQueryParams Params;
+
+	TQueue<FAllocationsImpl*> Results;
+
+	std::atomic<bool> IsWorking;
+	std::atomic<bool> IsCanceling;
+
+	FGraphEventRef CompletedEvent;
+};
+
+} // namespace TraceServices
diff --git a/Engine/Source/Developer/TraceServices/Private/Model/SbTree.cpp b/Engine/Source/Developer/TraceServices/Private/Model/SbTree.cpp
new file mode 100644
--- /dev/null
+++ b/Engine/Source/Developer/TraceServices/Private/Model/SbTree.cpp
@@ -0,0 +1,692 @@
+// Copyright Epic Games, Inc. All Rights Reserved.
+
+#include "SbTree.h"
+#include "TraceServices/Containers/Allocators.h"
+
+#include <limits>
+
+namespace TraceServices
+{
+
+#define USE_OFFSETTED_CELLS 1
+
+////////////////////////////////////////////////////////////////////////////////////////////////////
+// FSbTreeUtils
+////////////////////////////////////////////////////////////////////////////////////////////////////
+
+uint32 FSbTreeUtils::GetMaxDepth(uint32 TotalColumns)
+{
+	return 32 - FMath::CountLeadingZeros(TotalColumns);
+}
+
+////////////////////////////////////////////////////////////////////////////////////////////////////
+
+uint32 FSbTreeUtils::GetCellAtDepth(uint32 Column, uint32 Depth)
+{
+	// On depth D, the cell indices starts at 2^D-1 and increases by 2^(D+1).
+	uint32 LeafIndex = Column * 2;
+	uint32 K = 1 << Depth;
+	return (LeafIndex & ~K) | (K - 1);
+}
+
+////////////////////////////////////////////////////////////////////////////////////////////////////
+
+uint32 FSbTreeUtils::GetCommonDepth(uint32 ColumnA, uint32 ColumnB)
+{
+	uint32 Xor = ColumnA ^ ColumnB;
+	return FSbTreeUtils::GetMaxDepth(Xor);
+}
+
+////////////////////////////////////////////////////////////////////////////////////////////////////
+
+uint32 FSbTreeUtils::GetCellWidth(uint32 CellIndex)
+{
+	return ((CellIndex ^ (CellIndex + 1)) >> 1) + 1;
+}
+
+////////////////////////////////////////////////////////////////////////////////////////////////////
+// FSbTreeCell
+////////////////////////////////////////////////////////////////////////////////////////////////////
+
+FSbTreeCell::FSbTreeCell(ILinearAllocator& InAllocator)
+	: Allocator(InAllocator)
+	//, Allocs(InAllocator, 1024)
+	, MinStartEventIndex(std::numeric_limits<uint32>::max())
+	, MaxEndEventIndex(0)
+	, MinStartTime(std::numeric_limits<double>::max())
+	, MaxEndTime(0.0)
+{
+}
+
+////////////////////////////////////////////////////////////////////////////////////////////////////
+
+void FSbTreeCell::AddAlloc(const FAllocationItem* Alloc)
+{
+	Allocs.Add(*Alloc);
+
+	if (Alloc->StartEventIndex < MinStartEventIndex)
+	{
+		MinStartEventIndex = Alloc->StartEventIndex;
+	}
+	if (Alloc->EndEventIndex > MaxEndEventIndex)
+	{
+		MaxEndEventIndex = Alloc->EndEventIndex;
+	}
+
+	if (Alloc->StartTime < MinStartTime)
+	{
+		MinStartTime = Alloc->StartTime;
+	}
+	if (Alloc->EndTime > MaxEndTime)
+	{
+		MaxEndTime = Alloc->EndTime;
+	}
+}
+
+////////////////////////////////////////////////////////////////////////////////////////////////////
+
+void FSbTreeCell::Query(TArray<const FAllocationItem*>& OutAllocs, const IAllocationsProvider::FQueryParams& Params) const
+{
+	switch (Params.Rule)
+	{
+	case IAllocationsProvider::EQueryRule::aAf: // active allocs at A
+	{
+		const double Time = Params.TimeA;
+		for (const FAllocationItem& Alloc : Allocs)
+		{
+			if (Alloc.StartTime <= Time && Time <= Alloc.EndTime)
+			{
+				OutAllocs.Add(&Alloc);
+			}
+		}
+	}
+	break;
+
+	case IAllocationsProvider::EQueryRule::afA: // before
+	{
+		const double Time = Params.TimeA;
+		for (const FAllocationItem& Alloc : Allocs)
+		{
+			if (Alloc.EndTime <= Time)
+			{
+				OutAllocs.Add(&Alloc);
+			}
+		}
+	}
+	break;
+
+	case IAllocationsProvider::EQueryRule::Aaf: // after
+	{
+		const double Time = Params.TimeA;
+		for (const FAllocationItem& Alloc : Allocs)
+		{
+			if (Alloc.StartTime >= Time)
+			{
+				OutAllocs.Add(&Alloc);
+			}
+		}
+	}
+	break;
+
+	case IAllocationsProvider::EQueryRule::aAfB: // decline
+	{
+		const double TimeA = Params.TimeA;
+		const double TimeB = Params.TimeB;
+		for (const FAllocationItem& Alloc : Allocs)
+		{
+			if (Alloc.StartTime <= TimeA && Alloc.EndTime >= TimeA && Alloc.EndTime >= TimeB)
+			{
+				OutAllocs.Add(&Alloc);
+			}
+		}
+	}
+	break;
+
+	case IAllocationsProvider::EQueryRule::AaBf: // growth
+	{
+		const double TimeA = Params.TimeA;
+		const double TimeB = Params.TimeB;
+		for (const FAllocationItem& Alloc : Allocs)
+		{
+			if (Alloc.StartTime >= TimeA && Alloc.StartTime <= TimeB && Alloc.EndTime >= TimeB)
+			{
+				OutAllocs.Add(&Alloc);
+			}
+		}
+	}
+	break;
+
+	case IAllocationsProvider::EQueryRule::AafB: // short living allocs
+	{
+		const double TimeA = Params.TimeA;
+		const double TimeB = Params.TimeB;
+		for (const FAllocationItem& Alloc : Allocs)
+		{
+			if (Alloc.StartTime >= TimeA && Alloc.EndTime <= TimeB)
+			{
+				OutAllocs.Add(&Alloc);
+			}
+		}
+	}
+	break;
+
+	case IAllocationsProvider::EQueryRule::aABf: // long living allocs
+	{
+		const double TimeA = Params.TimeA;
+		const double TimeB = Params.TimeB;
+		for (const FAllocationItem& Alloc : Allocs)
+		{
+			if (Alloc.StartTime <= TimeA && Alloc.EndTime >= TimeB)
+			{
+				OutAllocs.Add(&Alloc);
+			}
+		}
+	}
+	break;
+
+	case IAllocationsProvider::EQueryRule::AaBCf: // memory leaks
+	{
+		const double TimeA = Params.TimeA;
+		const double TimeB = Params.TimeB;
+		const double TimeC = Params.TimeC;
+		for (const FAllocationItem& Alloc : Allocs)
+		{
+			if (Alloc.StartTime >= TimeA && Alloc.StartTime <= TimeB && Alloc.EndTime >= TimeC)
+			{
+				OutAllocs.Add(&Alloc);
+			}
+		}
+	}
+	break;
+
+	case IAllocationsProvider::EQueryRule::AaBfC: // limited lifetime
+	{
+		const double TimeA = Params.TimeA;
+		const double TimeB = Params.TimeB;
+		const double TimeC = Params.TimeC;
+		for (const FAllocationItem& Alloc : Allocs)
+		{
+			if (Alloc.StartTime >= TimeA && Alloc.StartTime <= TimeB && Alloc.EndTime >= TimeB && Alloc.EndTime <= TimeC)
+			{
+				OutAllocs.Add(&Alloc);
+			}
+		}
+	}
+	break;
+
+	case IAllocationsProvider::EQueryRule::aABfC: // decline of long living allocs
+	{
+		const double TimeA = Params.TimeA;
+		const double TimeB = Params.TimeB;
+		const double TimeC = Params.TimeC;
+		for (const FAllocationItem& Alloc : Allocs)
+		{
+			if (Alloc.StartTime <= TimeA && Alloc.EndTime >= TimeB && Alloc.EndTime <= TimeC)
+			{
+				OutAllocs.Add(&Alloc);
+			}
+		}
+	}
+	break;
+
+	case IAllocationsProvider::EQueryRule::AaBCfD: // specific lifetime
+	{
+		const double TimeA = Params.TimeA;
+		const double TimeB = Params.TimeB;
+		const double TimeC = Params.TimeC;
+		const double TimeD = Params.TimeD;
+		for (const FAllocationItem& Alloc : Allocs)
+		{
+			if (Alloc.StartTime >= TimeA && Alloc.StartTime <= TimeB && Alloc.EndTime >= TimeC && Alloc.EndTime <= TimeD)
+			{
+				OutAllocs.Add(&Alloc);
+			}
+		}
+	}
+	break;
+	}
+}
+
+////////////////////////////////////////////////////////////////////////////////////////////////////
+// FSbTree
+////////////////////////////////////////////////////////////////////////////////////////////////////
+
+FSbTree::FSbTree(ILinearAllocator& InAllocator, uint32 InColumnShift)
+	: Allocator(InAllocator)
+	, ColumnShift(InColumnShift)
+	, CurrentColumn(0)
+{
+	// Add the first cell.
+	Cells.Add(nullptr);
+
+#if USE_OFFSETTED_CELLS
+	// Add the first offsetted cell.
+	OffsettedCells.Add(nullptr);
+#endif // USE_OFFSETTED_CELLS
+}
+
+////////////////////////////////////////////////////////////////////////////////////////////////////
+
+FSbTree::~FSbTree()
+{
+	for (const FSbTreeCell* CellPtr : Cells)
+	{
+		if (CellPtr != nullptr)
+		{
+			delete CellPtr;
+		}
+	}
+	Cells.Reset();
+
+#if USE_OFFSETTED_CELLS
+	for (const FSbTreeCell* CellPtr : OffsettedCells)
+	{
+		if (CellPtr != nullptr)
+		{
+			delete CellPtr;
+		}
+	}
+	OffsettedCells.Reset();
+#endif // USE_OFFSETTED_CELLS
+}
+
+////////////////////////////////////////////////////////////////////////////////////////////////////
+
+void FSbTree::SetTimeForEvent(uint32 EventIndex, double Time)
+{
+	// Detect the first event of each new column.
+	if ((EventIndex & ((1 << ColumnShift) - 1)) == 0)
+	{
+		check((uint32)ColumnStartTimes.Num() == (EventIndex >> ColumnShift));
+		ColumnStartTimes.Add(Time);
+	}
+}
+
+////////////////////////////////////////////////////////////////////////////////////////////////////
+
+void FSbTree::AddAlloc(const FAllocationItem* Alloc)
+{
+	check(Alloc != nullptr);
+
+	uint32 StartColumn = Alloc->StartEventIndex >> ColumnShift;
+	uint32 EndColumn = Alloc->EndEventIndex >> ColumnShift;
+
+	if (EndColumn > CurrentColumn)
+	{
+		check((uint32)Cells.Num() > (CurrentColumn << 1));
+		check((uint32)ColumnStartTimes.Num() == EndColumn + 1);
+
+		// Adds 2 cells for each new column.
+		uint32 CellsToAdd = (EndColumn - CurrentColumn) << 1;
+
+		// We are only adding cell slots. A cell is created later, only if we add an alloc to it.
+		Cells.AddDefaulted(CellsToAdd);
+#if USE_OFFSETTED_CELLS
+		OffsettedCells.AddDefaulted(CellsToAdd);
+#endif // USE_OFFSETTED_CELLS
+
+		CurrentColumn = EndColumn;
+
+		//TODO: "Close" some of the "open" cells.
+	}
+
+#if USE_OFFSETTED_CELLS
+	const uint32 DeltaColumns = EndColumn - StartColumn;
+	const uint32 Depth = 32 - FMath::CountLeadingZeros(DeltaColumns);
+	const uint32 HalfCellWidth = (1 << Depth) >> 1;
+	const uint32 CellIndex = FSbTreeUtils::GetCellAtDepth(StartColumn & ~HalfCellWidth, Depth);
+	const bool bUseOffsettedCells = ((StartColumn & HalfCellWidth) != 0);
+#else // USE_OFFSETTED_CELLS
+	const uint32 Depth = FSbTreeUtils::GetCommonDepth(StartColumn, EndColumn);
+	const uint32 CellIndex = FSbTreeUtils::GetCellAtDepth(StartColumn, Depth);
+	const bool bUseOffsettedCells = false;
+#endif // USE_OFFSETTED_CELLS
+
+	FSbTreeCell* CellPtr;
+
+	if (!bUseOffsettedCells)
+	{
+		check(CellIndex < (uint32)Cells.Num());
+		CellPtr = Cells[CellIndex];
+		if (CellPtr == nullptr)
+		{
+			// Create new cell.
+			CellPtr = new FSbTreeCell(Allocator);
+			Cells[CellIndex] = CellPtr;
+		}
+	}
+#if USE_OFFSETTED_CELLS
+	else
+	{
+		check((CellIndex & 1) != 0); // depth 0 doesn't use offsetted cells
+
+		check(CellIndex < (uint32)OffsettedCells.Num());
+		CellPtr = OffsettedCells[CellIndex];
+		if (CellPtr == nullptr)
+		{
+			// Create new cell.
+			CellPtr = new FSbTreeCell(Allocator);
+			OffsettedCells[CellIndex] = CellPtr;
+		}
+	}
+#endif // USE_OFFSETTED_CELLS
+
+	CellPtr->AddAlloc(Alloc);
+}
+
+////////////////////////////////////////////////////////////////////////////////////////////////////
+
+int32 FSbTree::GetColumnAtTime(double Time) const
+{
+	const int32 Index = Algo::UpperBound(ColumnStartTimes, Time);
+	return Index - 1;
+}
+
+////////////////////////////////////////////////////////////////////////////////////////////////////
+
+void FSbTree::Query(TArray<const FSbTreeCell*>& OutCells, const IAllocationsProvider::FQueryParams& Params) const
+{
+	switch (Params.Rule)
+	{
+		case IAllocationsProvider::EQueryRule::aAf: // active allocs at A
+		{
+			const int32 ColumnA = GetColumnAtTime(Params.TimeA);
+			IterateCells(OutCells, ColumnA);
+		}
+		break;
+
+		case IAllocationsProvider::EQueryRule::afA: // before
+		{
+			const int32 ColumnA = GetColumnAtTime(Params.TimeA);
+			IterateCells(OutCells, 0, ColumnA);
+		}
+		break;
+
+		case IAllocationsProvider::EQueryRule::Aaf: // after
+		{
+			const int32 ColumnA = GetColumnAtTime(Params.TimeA);
+			IterateCells(OutCells, ColumnA, CurrentColumn);
+		}
+		break;
+	
+		case IAllocationsProvider::EQueryRule::aAfB: // decline
+		{
+			const int32 ColumnA = GetColumnAtTime(Params.TimeA);
+			IterateCells(OutCells, ColumnA);
+		}
+		break;
+	
+		case IAllocationsProvider::EQueryRule::AaBf: // growth
+		{
+			const int32 ColumnB = GetColumnAtTime(Params.TimeB);
+			IterateCells(OutCells, ColumnB);
+		}
+		break;
+
+		case IAllocationsProvider::EQueryRule::AafB: // short living allocs
+		{
+			const int32 ColumnA = GetColumnAtTime(Params.TimeA);
+			const int32 ColumnB = GetColumnAtTime(Params.TimeB);
+			IterateCells(OutCells, ColumnA, ColumnB);
+		}
+		break;
+
+		case IAllocationsProvider::EQueryRule::aABf: // long living allocs
+		{
+			const int32 ColumnA = GetColumnAtTime(Params.TimeA);
+			const int32 ColumnB = GetColumnAtTime(Params.TimeB);
+			IterateCells(OutCells, ColumnA, ColumnB);
+		}
+		break;
+	
+		case IAllocationsProvider::EQueryRule::AaBCf: // memory leaks
+		{
+			const int32 ColumnB = GetColumnAtTime(Params.TimeB);
+			const int32 ColumnC = GetColumnAtTime(Params.TimeC);
+			IterateCells(OutCells, ColumnB, ColumnC);
+		}
+		break;
+
+		case IAllocationsProvider::EQueryRule::AaBfC: // limited lifetime
+		{
+			const int32 ColumnB = GetColumnAtTime(Params.TimeB);
+			IterateCells(OutCells, ColumnB);
+		}
+		break;
+	
+		case IAllocationsProvider::EQueryRule::aABfC: // decline of long living allocs
+		{
+			const int32 ColumnA = GetColumnAtTime(Params.TimeA);
+			const int32 ColumnB = GetColumnAtTime(Params.TimeB);
+			IterateCells(OutCells, ColumnA, ColumnB);
+		}
+		break;
+	
+		case IAllocationsProvider::EQueryRule::AaBCfD: // specific lifetime
+		{
+			const int32 ColumnB = GetColumnAtTime(Params.TimeB);
+			const int32 ColumnC = GetColumnAtTime(Params.TimeC);
+			IterateCells(OutCells, ColumnB, ColumnC);
+		}
+		break;
+	}
+}
+
+////////////////////////////////////////////////////////////////////////////////////////////////////
+
+void FSbTree::IterateCells(TArray<const FSbTreeCell*>& OutCells, int32 Column) const
+{
+	if (Column < 0 || Column > (int32)CurrentColumn)
+	{
+		// not an error; just early out
+		return;
+	}
+
+	const uint32 MaxDepth = FSbTreeUtils::GetMaxDepth(CurrentColumn);
+
+	for (uint32 Depth = 0; Depth <= MaxDepth; ++Depth)
+	{
+		const uint32 CellIndex = FSbTreeUtils::GetCellAtDepth(Column, Depth);
+		if (CellIndex < (uint32)Cells.Num())
+		{
+			const FSbTreeCell* CellPtr = Cells[CellIndex];
+			if (CellPtr)
+			{
+				OutCells.Add(CellPtr);
+			}
+		}
+
+#if USE_OFFSETTED_CELLS
+		const uint32 HalfCellWidth = (1 << Depth) >> 1;
+		const uint32 OffsettedCellIndex = FSbTreeUtils::GetCellAtDepth(Column & ~HalfCellWidth, Depth);
+		if (OffsettedCellIndex < (uint32)OffsettedCells.Num())
+		{
+			const FSbTreeCell* CellPtr = OffsettedCells[OffsettedCellIndex];
+			if (CellPtr)
+			{
+				OutCells.Add(CellPtr);
+			}
+		}
+#endif // USE_OFFSETTED_CELLS
+	}
+}
+
+////////////////////////////////////////////////////////////////////////////////////////////////////
+
+void FSbTree::IterateCells(TArray<const FSbTreeCell*> & OutCells, int32 StartColumn, int32 EndColumn) const
+{
+	if (EndColumn < 0 || StartColumn > (int32)CurrentColumn || StartColumn > EndColumn)
+	{
+		// not an error; just early out
+		return;
+	}
+
+	const uint32 LocalStartColumn = (StartColumn < 0) ? 0 : (uint32)StartColumn;
+	const uint32 LocalEndColumn = ((uint32)EndColumn > CurrentColumn) ? CurrentColumn : (uint32)EndColumn;
+	check(LocalStartColumn <= LocalEndColumn);
+
+	const uint32 MaxDepth = FSbTreeUtils::GetMaxDepth(CurrentColumn);
+
+	for (uint32 Depth = 0; Depth <= MaxDepth; ++Depth)
+	{
+		const uint32 CellOffset = 1 << (Depth + 1);
+
+		const uint32 FirstCellIndex = FSbTreeUtils::GetCellAtDepth(LocalStartColumn, Depth);
+		const uint32 LastCellIndex = FMath::Min(FSbTreeUtils::GetCellAtDepth(LocalEndColumn, Depth) + 1, (uint32)Cells.Num());
+		for (uint32 CellIndex = FirstCellIndex; CellIndex < LastCellIndex; CellIndex += CellOffset)
+		{
+			FSbTreeCell* CellPtr = Cells[CellIndex];
+			if (CellPtr)
+			{
+				OutCells.Add(CellPtr);
+			}
+		}
+
+#if USE_OFFSETTED_CELLS
+		if (Depth > 0) // offsetted cells doesn't exist on Depth 0
+		{
+			uint32 HalfCellWidth = (1 << Depth) >> 1;
+			uint32 FirstOffsettedCellIndex = FSbTreeUtils::GetCellAtDepth(LocalStartColumn & ~HalfCellWidth, Depth);
+			uint32 LastOffsettedCellIndex = FMath::Min(FSbTreeUtils::GetCellAtDepth(LocalEndColumn & ~HalfCellWidth, Depth) + 1, (uint32)OffsettedCells.Num());
+			for (uint32 CellIndex = FirstOffsettedCellIndex; CellIndex < LastOffsettedCellIndex; CellIndex += CellOffset)
+			{
+				FSbTreeCell* CellPtr = OffsettedCells[CellIndex];
+				if (CellPtr)
+				{
+					OutCells.Add(CellPtr);
+				}
+			}
+		}
+#endif // USE_OFFSETTED_CELLS
+	}
+}
+
+////////////////////////////////////////////////////////////////////////////////////////////////////
+
+void FSbTree::DebugPrint() const
+{
+	uint32 TotalColumns = CurrentColumn + 1;
+	uint32 MaxDepth = FSbTreeUtils::GetMaxDepth(TotalColumns);
+
+	uint32 TotalAllocs = 0;
+	uint32 LocalMaxAllocCountPerCell = 0;
+
+	uint32 NotEmptyCellCount = 0;
+	const uint32 CellCount = (uint32)Cells.Num();
+	for (uint32 CellIndex = 0; CellIndex < CellCount; ++CellIndex)
+	{
+		const FSbTreeCell* CellPtr = Cells[CellIndex];
+		if (CellPtr != nullptr)
+		{
+			++NotEmptyCellCount;
+			uint32 CellAllocCount = CellPtr->GetAllocCount();
+			TotalAllocs += CellAllocCount;
+			if (CellAllocCount > LocalMaxAllocCountPerCell)
+			{
+				LocalMaxAllocCountPerCell = CellAllocCount;
+			}
+		}
+	}
+
+	uint32 TotalCells = CellCount;
+
+#if USE_OFFSETTED_CELLS
+	uint32 NotEmptyOffsettedCellCount = 0;
+	const uint32 OffsettedCellCount = (uint32)OffsettedCells.Num();
+	for (uint32 CellIndex = 0; CellIndex < OffsettedCellCount; ++CellIndex)
+	{
+		const FSbTreeCell* CellPtr = OffsettedCells[CellIndex];
+		if (CellPtr != nullptr)
+		{
+			++NotEmptyOffsettedCellCount;
+			uint32 CellAllocCount = CellPtr->GetAllocCount();
+			TotalAllocs += CellAllocCount;
+			if (CellAllocCount > LocalMaxAllocCountPerCell)
+			{
+				LocalMaxAllocCountPerCell = CellAllocCount;
+			}
+		}
+	}
+
+	TotalCells += OffsettedCellCount;
+#endif // USE_OFFSETTED_CELLS
+
+	//#define SbTreePrint(x) FPlatformMisc::LowLevelOutputDebugString(TEXT(x))
+	//#define SbTreePrintF(fmt, ...) FPlatformMisc::LowLevelOutputDebugStringf(TEXT(fmt), __VA_ARGS__)
+
+	FString StringBuffer;
+	#define SbTreePrint(x) StringBuffer += TEXT(x)
+	#define SbTreePrintF(fmt, ...) StringBuffer.Appendf(TEXT(fmt), __VA_ARGS__)
+
+	SbTreePrintF("Column Width:\t%u\n", (1 << ColumnShift));
+	SbTreePrintF("Allocs:\t%u\n", TotalAllocs);
+	SbTreePrintF("Columns:\t%u\n", TotalColumns);
+	SbTreePrintF("Max Depth:\t%u\n", MaxDepth);
+	SbTreePrintF("Cells:\t%u\n", TotalCells);
+	SbTreePrintF("Max Alloc Count Per Cell:\t%u\n", LocalMaxAllocCountPerCell);
+	SbTreePrintF("Not Empty Cells:\t%u\n", NotEmptyCellCount);
+#if USE_OFFSETTED_CELLS
+	SbTreePrintF("Not Empty Offsetted Cells:\t%u\n", NotEmptyOffsettedCellCount);
+#endif // USE_OFFSETTED_CELLS
+
+	SbTreePrint("\n");
+
+	for (uint32 Depth = 0; Depth <= MaxDepth; ++Depth)
+	{
+#if USE_OFFSETTED_CELLS
+		SbTreePrintF("\t%u\t%u*", Depth, Depth);
+#else // USE_OFFSETTED_CELLS
+		SbTreePrintF("\t%u", Depth);
+#endif // USE_OFFSETTED_CELLS
+	}
+	SbTreePrint("\n");
+
+	TArray<uint32> PrevCellIndex;
+	PrevCellIndex.AddDefaulted(MaxDepth + 1);
+	PrevCellIndex[0] = 1;
+
+	for (uint32 Column = 0; Column < TotalColumns; ++Column)
+	{
+		SbTreePrintF("%u", Column);
+
+		for (uint32 Depth = 0; Depth <= MaxDepth; ++Depth)
+		{
+			uint32 CellIndex = FSbTreeUtils::GetCellAtDepth(Column, Depth);
+			if (CellIndex != PrevCellIndex[Depth])
+			{
+				const FSbTreeCell* CellPtr = (CellIndex < CellCount) ? Cells[CellIndex] : nullptr;
+				uint32 AllocCount = (CellPtr != nullptr) ? CellPtr->GetAllocCount() : 0;
+
+#if USE_OFFSETTED_CELLS
+				const FSbTreeCell* OffsettedCellPtr = (CellIndex < OffsettedCellCount) ? OffsettedCells[CellIndex] : nullptr;
+				uint32 OffsettedAllocCount = (OffsettedCellPtr != nullptr) ? OffsettedCellPtr->GetAllocCount() : 0;
+
+				SbTreePrintF("\t%u\t%u", AllocCount, OffsettedAllocCount);
+#else // USE_OFFSETTED_CELLS
+				SbTreePrintF("\t%u", AllocCount);
+#endif // USE_OFFSETTED_CELLS
+
+				PrevCellIndex[Depth] = CellIndex;
+			}
+			else
+			{
+				SbTreePrint("\t\t");
+			}
+		}
+		SbTreePrint("\n");
+	}
+
+	SbTreePrint("\n");
+
+	FString FilePath = TEXT("D:/work/sbif.tab");
+	FFileHelper::SaveStringToFile(StringBuffer, *FilePath);
+
+	#undef SbTreePrint
+	#undef SbTreePrintF
+}
+
+////////////////////////////////////////////////////////////////////////////////////////////////////
+
+#undef USE_OFFSETTED_CELLS
+
+} // namespace TraceServices
diff --git a/Engine/Source/Developer/TraceServices/Private/Model/SbTree.h b/Engine/Source/Developer/TraceServices/Private/Model/SbTree.h
new file mode 100644
--- /dev/null
+++ b/Engine/Source/Developer/TraceServices/Private/Model/SbTree.h
@@ -0,0 +1,86 @@
+// Copyright Epic Games, Inc. All Rights Reserved.
+
+#pragma once
+
+#include "AllocationsProvider.h"
+//#include "Common/PagedArray.h"
+
+namespace TraceServices
+{
+
+class ILinearAllocator;
+
+////////////////////////////////////////////////////////////////////////////////////////////////////
+
+class FSbTreeUtils
+{
+public:
+	static uint32 GetMaxDepth(uint32 TotalColumns);
+	static uint32 GetCellAtDepth(uint32 Column, uint32 Depth);
+	static uint32 GetCommonDepth(uint32 ColumnA, uint32 ColumnB);
+	static uint32 GetCellWidth(uint32 CellIndex);
+};
+
+////////////////////////////////////////////////////////////////////////////////////////////////////
+
+class FSbTreeCell
+{
+public:
+	FSbTreeCell(ILinearAllocator& InAllocator);
+
+	uint32 GetAllocCount() const { return Allocs.Num(); }
+	void AddAlloc(const FAllocationItem* Alloc);
+
+	double GetMinStartTime() const { return MinStartTime; }
+	double GetMaxEndTime() const { return MaxEndTime; }
+
+	void Query(TArray<const FAllocationItem*>& OutAllocs, const IAllocationsProvider::FQueryParams& Params) const;
+
+private:
+	ILinearAllocator& Allocator;
+
+	//TODO: TPagedArray<FAllocationItem> Allocs;
+	TArray<FAllocationItem> Allocs;
+
+	uint32 MinStartEventIndex;
+	uint32 MaxEndEventIndex;
+	double MinStartTime;
+	double MaxEndTime;
+};
+
+////////////////////////////////////////////////////////////////////////////////////////////////////
+
+class FSbTree
+{
+public:
+	FSbTree(ILinearAllocator& InAllocator, uint32 InColumnShift);
+	~FSbTree();
+
+	void SetTimeForEvent(uint32 EventIndex, double Time);
+
+	void AddAlloc(const FAllocationItem* Alloc);
+
+	uint32 GetColumnWidth() const { return 1 << ColumnShift; }
+	uint32 GetCurrentColumn() const { return CurrentColumn; }
+
+	int32 GetColumnAtTime(double Time) const;
+
+	void Query(TArray<const FSbTreeCell*>& OutCells, const IAllocationsProvider::FQueryParams& Params) const;
+
+	void IterateCells(TArray<const FSbTreeCell*>& OutCells, int32 Column) const;
+	void IterateCells(TArray<const FSbTreeCell*>& OutCells, int32 StartColumn, int32 EndColumn) const;
+
+	void DebugPrint() const;
+
+private:
+	ILinearAllocator& Allocator;
+	TArray<FSbTreeCell*> Cells;
+	TArray<FSbTreeCell*> OffsettedCells;
+	TArray<double> ColumnStartTimes;
+	uint32 ColumnShift; // ColumnWidth = 1 << ColumnShift
+	uint32 CurrentColumn;
+};
+
+////////////////////////////////////////////////////////////////////////////////////////////////////
+
+} // namespace TraceServices
diff --git a/Engine/Source/Developer/TraceServices/Private/Modules/MemoryModule.cpp b/Engine/Source/Developer/TraceServices/Private/Modules/MemoryModule.cpp
--- a/Engine/Source/Developer/TraceServices/Private/Modules/MemoryModule.cpp
+++ b/Engine/Source/Developer/TraceServices/Private/Modules/MemoryModule.cpp
@@ -8,24 +8,24 @@
 #include "Model/CallstacksProvider.h"
 #include "TraceServices/Model/AnalysisSession.h"
 
-#if defined(UE_USE_ALLOCATIONS_PROVIDER)
+namespace TraceServices
+{
 
-static void TEMP_TestAllocationsProviderImpl(double TimeA, double TimeB)
+static void TEMP_TestAllocationsProviderImpl(ILinearAllocator& Allocator, double TimeA, double TimeB)
 {
-	FAllocationsProvider Provider;
+	FAllocationsProvider Provider(Allocator);
 
 	// 1. Query is started in response to some form of user interaction.
-	IAllocationsProvider::QueryHandle Query;
+	IAllocationsProvider::FQueryHandle Query;
 	{
-		Query = Provider.StartQuery(TimeA, TimeB, IAllocationsProvider::ECrosses::Xor);
+		const IAllocationsProvider::FQueryParams Params = { IAllocationsProvider::EQueryRule::aABf, TimeA, TimeB, 0.0, 0.0 };
+		Query = Provider.StartQuery(Params);
 	}
 
-
-
 	// 2. At regular intervals the query is polled for results.
 	while (true)
 	{
-		IAllocationsProvider::FQueryStatus Status = Provider.PollQuery(Query);
+		const IAllocationsProvider::FQueryStatus Status = Provider.PollQuery(Query);
 		if (Status.Status <= IAllocationsProvider::EQueryStatus::Done)
 		{
 			break;
@@ -33,33 +33,36 @@ static void TEMP_TestAllocationsProviderImpl(double TimeA, double TimeB)
 
 		// Multiple 'pages' of results will be returned. No guarantees are made
 		// about the order of pages or the allocations they report.
-		IAllocationsProvider::QueryResult Result = Status.NextResult();
+		IAllocationsProvider::FQueryResult Result = Status.NextResult();
 		while (Result.IsValid())
 		{
 			for (uint32 i = 0, n = Result->Num(); i < n; ++i)
 			{
 				const IAllocationsProvider::FAllocation* Allocation = Result->Get(i);
 
+				auto Addr = Allocation->GetAddress();
+
+				auto Size = Allocation->GetSize();
+				auto Alignment = Allocation->GetAlignment();
+				auto Waste = Allocation->GetWaste();
+
 				uint64 BacktraceId = Allocation->GetBacktraceId();
 				/*
 				auto Backtrace = BacktraceProvider.Get(BacktraceId);
 				for (auto& RetAddress : Backtrace)
 				{
-					auto Addr = Allocation->GetAddress();
-					auto Size = Allocation->GetSize();
-					auto LlmTag = Allocation->GetTag();
-
-					... report to user (e.g. build treeview)
 				}
 				*/
+
+				auto Tag = Allocation->GetTag();
+
+				//... report to user (e.g.build treeview)
 			}
 
 			Result = Status.NextResult();
 		}
 	}
 
-
-
 	// 3a. Queries can be cancelled.
 	if (/*HasUserPressedCance()*/ false)
 	{
@@ -70,20 +73,13 @@ static void TEMP_TestAllocationsProviderImpl(double TimeA, double TimeB)
 	/* Queries will automatically clean themselves up */
 }
 
-static void TEMP_TestAllocationsProvider()
+static void TEMP_TestAllocationsProvider(ILinearAllocator& Allocator)
 {
 	double TimeA = 3.0;
 	double TimeB = 11.0;
-	TEMP_TestAllocationsProviderImpl(TimeA, TimeB);
+	TEMP_TestAllocationsProviderImpl(Allocator, TimeA, TimeB);
 }
 
-#endif // UE_USE_ALLOCATIONS_PROVIDER
-
-
-
-namespace TraceServices
-{
-
 static const FName MemoryModuleName("TraceModule_Memory");
 
 void FMemoryModule::GetModuleInfo(FModuleInfo& OutModuleInfo)
@@ -94,13 +90,11 @@ void FMemoryModule::GetModuleInfo(FModuleInfo& OutModuleInfo)
 
 void FMemoryModule::OnAnalysisBegin(IAnalysisSession& Session)
 {
-#if defined(UE_USE_ALLOCATIONS_PROVIDER)
-	FAllocationsProvider* AllocationsProvider = new FAllocationsProvider();
+	FAllocationsProvider* AllocationsProvider = new FAllocationsProvider(Session.GetLinearAllocator());
 	Session.AddProvider(AllocationsProvider->GetName(), AllocationsProvider);
-	//Session.AddAnalyzer(new FAllocationsAnalyzer(Session));
+	Session.AddAnalyzer(new FAllocationsAnalyzer(Session, *AllocationsProvider));
 
-	TEMP_TestAllocationsProvider();
-#endif // UE_USE_ALLOCATIONS_PROVIDER
+	//TEMP_TestAllocationsProvider(Session.GetLinearAllocator());
 
 	FCallstacksProvider* CallstacksProvider = new FCallstacksProvider(Session);
 	Session.AddProvider(CallstacksProvider->GetName(), CallstacksProvider);
diff --git a/Engine/Source/Developer/TraceServices/Public/TraceServices/Model/AllocationsProvider.h b/Engine/Source/Developer/TraceServices/Public/TraceServices/Model/AllocationsProvider.h
--- a/Engine/Source/Developer/TraceServices/Public/TraceServices/Model/AllocationsProvider.h
+++ b/Engine/Source/Developer/TraceServices/Public/TraceServices/Model/AllocationsProvider.h
@@ -2,80 +2,75 @@
 
 #pragma once
 
-#if defined(UE_USE_ALLOCATIONS_PROVIDER)
-
-namespace TraceServices
-{
-
 #include "AnalysisSession.h"
 #include "CoreTypes.h"
 #include "Templates/UniquePtr.h"
 
 #include <new>
 
-/* {{{1
-
-1. *A*    Query(A, A, Crosses_Or)             Show all live at time A
-   **A    Query(0, A, Crosses_None)           Show all freed before time A
-   A**    Query(A, INF, Crosses_None)         Show all allocated after time A
-2. A\B    Query(A, B, Crosses_A)              Live at time A, but not at time B
-   B\A    Query(A, B, Crosses_B)              Live at time B, but not at time A
-   A&B    Query(A, B, Crosses_And)            Live at time A and also at time B
-   A^B    Query(A, A, Crosses_Xor)            Live either at time A or at time B
-   A|B    Query(A, B, Crosses_Or)             All at time A and all at time B with diffs highlighted
-   A<>B   Query(A, B, Crosses_None)           Allocated and freed during AB time interval
-   <AB>   [same as A&B]                       Allocated before time A and freed after time B
-   <A>B   [same as A\B]                       Allocated before time A and freed during AB time interval
-   A<B>   [same as B\A]                       Allocated during AB time interval and freed after time B
-3. A*BC*  Query(B, C, Crosses_And).Start > A  Allocated during AB time interval and still live after time C
-   *AB*C  Query(A, B, Crosses_And).End < C    Allocated before A and freed during BC time interval
-
-                     [-]
-   (A)        (B)    < >   ~Or             /---- Xor&A
-    |          |     A B   None Xor And Or A B-- Xor&B
-    | [------] |     0 0   1    0   0   0  0 0
-    |      [---+--]  0 1   0    1   0   1  0 1
- [--+---]      |     1 0   0    1   0   1  1 0
-  [-+----------+-]   1 1   0    0   1   1  0 0
-    |          |
-
-}}} */
-
-////////////////////////////////////////////////////////////////////////////////
-class IAllocationsProvider
-	: public UE::Trace::IProvider
+namespace TraceServices
+{
+
+////////////////////////////////////////////////////////////////////////////////////////////////////
+
+class IAllocationsProvider : public IProvider
 {
 public:
-	struct FAllocation
+	// Allocation query rules.
+	// The enum uses the following naming convention:
+	//     A, B, C, D = time markers
+	//     a = time when "alloc" event occurs
+	//     f = time when "free" event occurs (can be infinite)
+	// Ex.: "AaBf" means "all memory allocations allocated between time A and time B and freed after time B".
+	enum class EQueryRule
 	{
-		double		GetStartTime() const;
-		double		GetEndTime() const;
-		uint64		GetAddress() const;
-		uint64		GetSize() const;
-		uint64		GetBacktraceId() const;
-		uint32		GetTag() const;
+		aAf,     // active allocs at A
+		afA,     // before
+		Aaf,     // after
+		aAfB,    // decline
+		AaBf,    // growth
+		AafB,    // short living allocs
+		aABf,    // long living allocs
+		AaBCf,   // memory leaks
+		AaBfC,   // limited lifetime
+		aABfC,   // decline of long living allocs
+		AaBCfD,  // specific lifetime
+		//A_vs_B,  // compare A vs. B; {aAf} vs. {aBf}
+		//A_or_B,  // live at A or at B; {aAf} U {aBf}
+		//A_xor_B, // live either at A or at B; ({aAf} U {aBf}) \ {aABf}
 	};
 
-	class FAllocations
+	struct TRACESERVICES_API FQueryParams
 	{
-	public:
-		void				operator delete (void* Address);
-		uint32				Num() const;
-		const FAllocation*	Get(uint32 Index) const;
+		EQueryRule Rule;
+		double TimeA;
+		double TimeB;
+		double TimeC;
+		double TimeD;
 	};
 
-	typedef TUniquePtr<const FAllocations> QueryResult;
+	struct TRACESERVICES_API FAllocation
+	{
+		double GetStartTime() const;
+		double GetEndTime() const;
+		uint64 GetAddress() const;
+		uint32 GetSize() const;
+		uint8 GetAlignment() const;
+		uint8 GetWaste() const;
+		uint64 GetBacktraceId() const;
+		uint32 GetTag() const;
+	};
 
-	enum class ECrosses
+	class TRACESERVICES_API FAllocations
 	{
-		None,
-		//A, /* ...off temporarily */
-		//B, /* ...off temporarily */
-		Xor,
-		And,
-		Or,
+	public:
+		void operator delete (void* Address);
+		uint32 Num() const;
+		const FAllocation* Get(uint32 Index) const;
 	};
 
+	typedef TUniquePtr<const FAllocations> FQueryResult;
+
 	enum class EQueryStatus
 	{
 		Unknown,
@@ -84,20 +79,24 @@ public:
 		Available,
 	};
 
-	struct FQueryStatus
+	struct TRACESERVICES_API FQueryStatus
 	{
-		QueryResult			NextResult() const;
-		EQueryStatus		Status;
-		mutable UPTRINT		Handle;
+		FQueryResult NextResult() const;
+
+		EQueryStatus Status;
+		mutable UPTRINT Handle;
 	};
 
-	typedef UPTRINT				QueryHandle;
-	typedef const FQueryStatus	QueryStatus;
-	virtual QueryHandle			StartQuery(double TimeA, double TimeB, ECrosses Crosses) = 0;
-	virtual void				CancelQuery(QueryHandle Query) = 0;
-	virtual QueryStatus			PollQuery(QueryHandle Query) = 0;
+	typedef UPTRINT FQueryHandle;
+
+public:
+	virtual bool IsInitialized() const = 0;
+
+	virtual FQueryHandle StartQuery(const FQueryParams& Params) const = 0;
+	virtual void CancelQuery(FQueryHandle Query) const = 0;
+	virtual const FQueryStatus PollQuery(FQueryHandle Query) const = 0;
 };
 
-} // namespace TraceServices
+TRACESERVICES_API const IAllocationsProvider* ReadAllocationsProvider(const IAnalysisSession& Session);
 
-#endif // UE_USE_ALLOCATIONS_PROVIDER
+} // namespace TraceServices
