Input stage of an alloc/free trace event pipeline that matches allocations to their frees as a series of jobs, sorting the results in an order suitable for addition to a "Swept Binary Interval Forest" data structure.

diff --git a/Engine/Source/Developer/TraceServices/Private/Analyzers/Memory/Config.h b/Engine/Source/Developer/TraceServices/Private/Analyzers/Memory/Config.h
new file mode 100644
--- /dev/null
+++ b/Engine/Source/Developer/TraceServices/Private/Analyzers/Memory/Config.h
@@ -0,0 +1,26 @@
+// Copyright Epic Games, Inc. All Rights Reserved.
+
+#pragma once
+
+#include "CoreTypes.h"
+
+namespace Trace {
+namespace TraceServices {
+
+namespace FTrackerConfig
+{
+	static const uint32 NumLanes			= 32;
+	static const uint32 MaxLaneInputItems	= 16 << 10;
+	static const uint32 ActiveSetPageSize	= 32 << 10;
+
+	/* The following are assumptions for optimisation and not configurable */
+	static const uint32 MaxSerialBits		= 20;	// = 64 - addressable_bits - min_align = 64 - 47 - 3
+
+	static_assert(MaxLaneInputItems < ActiveSetPageSize, "");
+	static_assert(MaxLaneInputItems * NumLanes <= 1 << MaxSerialBits, "");
+};
+
+} // namespace TraceServices
+} // namespace Trace
+
+/* vim: set noet : */
diff --git a/Engine/Source/Developer/TraceServices/Private/Analyzers/Memory/Lane.h b/Engine/Source/Developer/TraceServices/Private/Analyzers/Memory/Lane.h
new file mode 100644
--- /dev/null
+++ b/Engine/Source/Developer/TraceServices/Private/Analyzers/Memory/Lane.h
@@ -0,0 +1,132 @@
+// Copyright Epic Games, Inc. All Rights Reserved.
+
+#pragma once
+
+#include "Containers/ArrayView.h"
+#include "LaneItemSet.h"
+
+namespace Trace {
+namespace TraceServices {
+
+////////////////////////////////////////////////////////////////////////////////
+class FAddressSerial
+{
+public:
+	void Set(uint32 Serial, uint64 Address)
+	{
+		Address_Serial = Serial & (SerialMax - 1);
+		Address_Serial += Address << UnusedHighBits;
+	}
+
+	uint64	GetAddress() const		{ return (Address_Serial >> UnusedHighBits) & ~AlignBits; }
+	uint32	GetBiasedSerial() const	{ return Address_Serial & (SerialMax - 1); }
+
+protected:
+	enum
+	{
+		AlignBits		= 3,
+		AddressBits		= 47,
+		UnusedHighBits	= 64 - 47,
+		SerialBits		= UnusedHighBits + AlignBits,
+		SerialMax		= 1 << SerialBits,
+	};
+
+	uint64		Address_Serial; // packed for sorting purposes
+};
+
+
+
+////////////////////////////////////////////////////////////////////////////////
+class alignas(16) FLaneItem
+	: public FAddressSerial
+{
+public:
+	void Set(uint32 Serial, uint64 Address, uint32 MetadataId)
+	{
+		FAddressSerial::Set(Serial, Address);
+		Id = MetadataId;
+	}
+
+	uint32	GetSerial(uint32 Bias) const				{ return GetBiasedSerial() + Bias; }
+	void	SetActiveIndex(uint32 Index)				{ ActiveIndex = Index; }
+	uint32	GetActiveIndex() const						{ return ActiveIndex; }
+	bool	operator < (const FLaneItem& Rhs) const		{ return Address_Serial < Rhs.Address_Serial; }
+	bool	IsSameAddress(const FLaneItem& Rhs) const	{ return (Address_Serial ^ Rhs.Address_Serial) < SerialMax; }
+	uint32	GetMetadataId() const						{ return Id; }
+	bool	HasMetadata() const							{ return Id != 0; }
+
+private:
+	uint32		Id;
+	uint32		ActiveIndex;
+};
+static_assert(sizeof(FLaneItem) == 16, "Deliberately XMM sized");
+
+
+
+////////////////////////////////////////////////////////////////////////////////
+typedef TArrayView<FLaneItem> FLaneItemView;
+
+
+
+////////////////////////////////////////////////////////////////////////////////
+class FLaneInput
+{
+public:
+	bool				AddAlloc(uint64 Address, uint32 Serial, uint32 MetadataId);
+	bool				AddFree(uint64 Address, uint32 Serial);
+	uint32				GetNum() const	{ return Num; }
+	FLaneItemView		GetAllocs()		{ return FLaneItemView(Items, NumAllocs); }
+	FLaneItemView		GetFrees()		{ return FLaneItemView(Items + Max - NumFrees, NumFrees); }
+
+private:
+	uint32				Max;
+	uint32				NumAllocs = 0;
+	uint32				NumFrees = 0;
+	uint32				Num = 0;
+	FLaneItem			Items[];
+
+	static_assert(sizeof(FLaneItem) == 16, "Deliberately XMM-sized");
+	friend class FTrackerBuffer;
+};
+
+////////////////////////////////////////////////////////////////////////////////
+inline bool FLaneInput::AddAlloc(uint64 Address, uint32 Serial, uint32 MetadataId)
+{
+	FLaneItem* Item = Items + NumAllocs;
+	Item->Set(Serial, Address, MetadataId);
+
+	++NumAllocs;
+	++Num;
+	return Num >= Max;
+}
+
+////////////////////////////////////////////////////////////////////////////////
+inline bool FLaneInput::AddFree(uint64 Address, uint32 Serial)
+{
+	FLaneItem* Item = Items + Max - NumFrees - 1;
+	Item->Set(Serial, Address, 0);
+
+	++NumFrees;
+	++Num;
+	return Num >= Max;
+}
+
+
+
+////////////////////////////////////////////////////////////////////////////////
+class FLane
+{
+public:
+	void				LockWrite()				{}
+	void				UnlockWrite()			{}
+	FLaneItemSet&		GetActiveSet()			{ return ActiveSet; }
+	const FLaneItemSet& GetActiveSet() const	{ return ActiveSet; }
+
+private:
+	FLaneItemSet		ActiveSet;
+};
+
+} // namespace TraceServices
+} // namespace Trace
+
+/* vim: set noet : */
diff --git a/Engine/Source/Developer/TraceServices/Private/Analyzers/Memory/LaneItemSet.cpp b/Engine/Source/Developer/TraceServices/Private/Analyzers/Memory/LaneItemSet.cpp
new file mode 100644
--- /dev/null
+++ b/Engine/Source/Developer/TraceServices/Private/Analyzers/Memory/LaneItemSet.cpp
@@ -0,0 +1,245 @@
+// Copyright Epic Games, Inc. All Rights Reserved.
+
+#include "LaneItemSet.h"
+#include "Support.h"
+
+namespace Trace {
+namespace TraceServices {
+
+////////////////////////////////////////////////////////////////////////////////
+FLaneItemSet::FLaneItemSet()
+{
+}
+
+////////////////////////////////////////////////////////////////////////////////
+FLaneItemSet::~FLaneItemSet()
+{
+	Clear();
+}
+
+////////////////////////////////////////////////////////////////////////////////
+int32 FLaneItemSet::Find(uint64 Address) const
+{
+	if (Num == 0)
+	{
+		return -1;
+	}
+
+	uint64 Hash = GetHash(Address);
+	uint32 BucketIndex = GetBucketIndex(Hash);
+
+	while (true)
+	{
+		FEntry* __restrict Cursor = Data + BucketIndex;
+		for (uint32 i = BucketSize / 4;;)
+		{
+			uint32 Rax;
+			Rax  = (Cursor[3].Address == Address)|(!Cursor[3].Address); Rax <<= 1;
+			Rax |= (Cursor[2].Address == Address)|(!Cursor[2].Address); Rax <<= 1;
+			Rax |= (Cursor[1].Address == Address)|(!Cursor[1].Address); Rax <<= 1;
+			Rax |= (Cursor[0].Address == Address)|(!Cursor[0].Address);
+
+			if (Rax)
+			{
+				Rax &= -int32(Rax);
+				uint32 Index = 31 - UnsafeCountLeadingZeros(Rax);
+				return Cursor[Index].Address ? int32(UPTRINT((Cursor + Index) - Data)) : -1;
+			}
+
+			Cursor += 4;
+
+			--i;
+			if (i == 0)
+			{
+				break;
+			}
+		}
+
+		BucketIndex += BucketSize;
+		if (BucketIndex >= Capacity)
+		{
+			BucketIndex = 0;
+		}
+	}
+
+	return -1;
+}
+
+////////////////////////////////////////////////////////////////////////////////
+void FLaneItemSet::Add(uint64 Address, uint32 Serial, uint32 MetadataId)
+{
+	if (Num >= Capacity - 1) // Leave at least one empty to prevent infinite loops
+	{
+		return;
+	}
+
+	return AddInternal(Address, Serial, MetadataId);
+}
+
+////////////////////////////////////////////////////////////////////////////////
+void FLaneItemSet::AddInternal(uint64 Address, uint32 Serial, uint32 MetadataId)
+{
+	uint64 Hash = GetHash(Address);
+	uint32 BucketIndex = GetBucketIndex(Hash);
+
+	while (true)
+	{
+		FEntry* __restrict Cursor = Data + BucketIndex;
+		for (uint32 i = BucketSize / 4;;)
+		{
+			uint32 Rax;
+			Rax  = (Cursor[3].Address <= Tomb); Rax <<= 1;
+			Rax |= (Cursor[2].Address <= Tomb); Rax <<= 1;
+			Rax |= (Cursor[1].Address <= Tomb); Rax <<= 1;
+			Rax |= (Cursor[0].Address <= Tomb);
+
+			if (Rax)
+			{
+				Rax &= -int32(Rax);
+				uint32 Index = 31 - UnsafeCountLeadingZeros(Rax);
+
+				NumTombs -= (Cursor[Index].Address == Tomb);
+				Cursor[Index] = { Address, Serial, MetadataId };
+
+				++Num;
+				return;
+			}
+
+			Cursor += 4;
+
+			--i;
+			if (i == 0)
+			{
+				break;
+			}
+		}
+
+		BucketIndex += BucketSize;
+		if (BucketIndex >= Capacity)
+		{
+			BucketIndex = 0;
+		}
+	}
+}
+
+////////////////////////////////////////////////////////////////////////////////
+bool FLaneItemSet::Update(uint32 Index, uint32 Serial, uint32 MetadataId)
+{
+	if (Index >= Capacity)
+	{
+		return false;
+	}
+
+	FEntry& Entry = Data[Index];
+	Entry.Serial = Serial;
+	Entry.MetadataId = MetadataId;
+	return true;
+}
+
+////////////////////////////////////////////////////////////////////////////////
+void FLaneItemSet::Remove(uint32 Index)
+{
+	if (Index >= Capacity)
+	{
+		return;
+	}
+
+	Data[Index].Address = Tomb;
+	--Num;
+	++NumTombs;
+}
+
+////////////////////////////////////////////////////////////////////////////////
+FLaneItemSet::FItemHandle FLaneItemSet::ReadItems() const
+{
+	auto Handle = FItemHandle(Data - 1);
+	return NextItem(Handle);
+}
+
+////////////////////////////////////////////////////////////////////////////////
+FLaneItemSet::FItemHandle FLaneItemSet::NextItem(FItemHandle Handle) const
+{
+	const auto *Entry = ((const FEntry*)Handle) + 1, *End = Data + Capacity;
+	do
+	{
+		if (Entry->Address > Tomb)
+		{
+			return FItemHandle(Entry);
+		}
+
+		if (Entry->Address == 0)
+		{
+			Entry += BucketMask - (UPTRINT(Entry - Data) & BucketMask);
+		}
+
+		++Entry;
+	}
+	while (Entry < End);
+
+	return 0;
+}
+
+////////////////////////////////////////////////////////////////////////////////
+int32 FLaneItemSet::GetItemIndex(FItemHandle Handle) const
+{
+	return PTRINT((FEntry*)Handle - Data);
+}
+
+////////////////////////////////////////////////////////////////////////////////
+void FLaneItemSet::Clear()
+{
+	FTrackerBuffer::Free(Data);
+	Num = Capacity = 0;
+}
+
+////////////////////////////////////////////////////////////////////////////////
+void FLaneItemSet::Rehash(uint32 GrowSize)
+{
+	FEntry* __restrict PrevData = Data;
+	FEntry* __restrict PrevEnd = Data + Capacity;
+
+	NumTombs = 0;
+	Capacity += GrowSize;
+	SIZE_T Size = sizeof(FEntry) * Capacity;
+	Data = FTrackerBuffer::AllocRaw<FEntry>(Size, EntryAlign);
+	memset(Data, 0, Size);
+
+	Num = 0;
+	for (FEntry* __restrict Entry = PrevData; Entry < PrevEnd; ++Entry)
+	{
+		if (Entry->Address <= Tomb)
+		{
+			if (Entry->Address == 0)
+			{
+				Entry += BucketMask - (UPTRINT(Entry - PrevData) & BucketMask);
+			}
+
+			continue;
+		}
+
+		AddInternal(Entry->Address, Entry->Serial, Entry->MetadataId);
+	}
+
+	FTrackerBuffer::Free(PrevData);
+}
+
+////////////////////////////////////////////////////////////////////////////////
+uint64 FLaneItemSet::GetHash(uint64 Address)
+{
+	static const uint64 Mix = 0x30be8efa499c249dull;
+	uint64 Hash = Address >> 3;
+	Hash += Address & 0xffff; Hash *= Mix; Address >>= 16;
+	Hash += Address & 0xffff; Hash *= Mix;
+	return Hash;
+}
+
+////////////////////////////////////////////////////////////////////////////////
+uint32 FLaneItemSet::GetBucketIndex(uint64 Hash) const
+{
+	return uint32((uint64(uint32(Hash)) * Capacity) >> 32ull) & ~BucketMask;
+}
+
+} // namespace TraceServices
+} // namespace Trace
+
+/* vim: set noet : */
diff --git a/Engine/Source/Developer/TraceServices/Private/Analyzers/Memory/LaneItemSet.h b/Engine/Source/Developer/TraceServices/Private/Analyzers/Memory/LaneItemSet.h
new file mode 100644
--- /dev/null
+++ b/Engine/Source/Developer/TraceServices/Private/Analyzers/Memory/LaneItemSet.h
@@ -0,0 +1,80 @@
+// Copyright Epic Games, Inc. All Rights Reserved.
+
+#pragma once
+
+#include "CoreTypes.h"
+
+namespace Trace {
+namespace TraceServices {
+
+////////////////////////////////////////////////////////////////////////////////
+class FLaneItemSet
+{
+public:
+	typedef UPTRINT	FItemHandle;
+					FLaneItemSet();
+					~FLaneItemSet();
+	void			Clear();
+	void			Rehash(uint32 GrowSize);
+	uint32			GetNum() const	{ return Num; }
+	uint32			GetCapacity() const { return Capacity; }
+	uint32			GetNumTombs() const { return NumTombs; }
+	int32			Find(uint64 Address) const;
+	uint64			GetAddress(uint32 Index) const;
+	uint32			GetSerial(uint32 Index) const;
+	uint32			GetMetadataId(uint32 Index) const;
+	void			Add(uint64 Address, uint32 Serial, uint32 MetadataId);
+	bool			Update(uint32 Index, uint32 Serial, uint32 MetadataId);
+	void			Remove(uint32 Index);
+	FItemHandle		ReadItems() const;
+	FItemHandle		NextItem(FItemHandle Handle) const;
+	int32			GetItemIndex(FItemHandle Handle) const;
+
+private:
+	struct FEntry
+	{
+		uint64		Address;
+		uint32		Serial;		// this may overflow on long runs. There's looooads of bits free in 'Address' though.
+		uint32		MetadataId;
+	};
+
+	enum
+	{
+		BucketShift = 8,
+		BucketSize	= 1 << BucketShift,
+		BucketMask	= BucketSize - 1,
+		Tomb		= 1,
+		EntryAlign	= 64,
+	};
+
+	static uint64	GetHash(uint64 Address);
+	void			AddInternal(uint64 Address, uint32 Serial, uint32 MetadataId);
+	uint32			GetBucketIndex(uint64 Hash) const;
+	FEntry*			Data = nullptr;
+	uint32			Num = 0;
+	uint32			Capacity = 0;
+	uint32			NumTombs = 0;
+};
+
+////////////////////////////////////////////////////////////////////////////////
+inline uint64 FLaneItemSet::GetAddress(uint32 Index) const
+{
+	return Data[Index].Address;
+}
+
+////////////////////////////////////////////////////////////////////////////////
+inline uint32 FLaneItemSet::GetSerial(uint32 Index) const
+{
+	return Data[Index].Serial;
+}
+
+////////////////////////////////////////////////////////////////////////////////
+inline uint32 FLaneItemSet::GetMetadataId(uint32 Index) const
+{
+	return Data[Index].MetadataId;
+}
+
+} // namespace TraceServices
+} // namespace Trace
+
+/* vim: set noet : */
diff --git a/Engine/Source/Developer/TraceServices/Private/Analyzers/Memory/MetadataDb.cpp b/Engine/Source/Developer/TraceServices/Private/Analyzers/Memory/MetadataDb.cpp
new file mode 100644
--- /dev/null
+++ b/Engine/Source/Developer/TraceServices/Private/Analyzers/Memory/MetadataDb.cpp
@@ -0,0 +1,103 @@
+// Copyright Epic Games, Inc. All Rights Reserved.
+
+#include "MetadataDb.h"
+#include "Support.h"
+
+namespace Trace {
+namespace TraceServices {
+
+////////////////////////////////////////////////////////////////////////////////
+bool FMetadataDb::FEntryInternal::operator == (const FEntryInternal& Rhs) const
+{
+	return (Owner == Rhs.Owner) | (Size == Rhs.Size);
+}
+
+////////////////////////////////////////////////////////////////////////////////
+FMetadataDb::FMetadataDb()
+{
+}
+
+////////////////////////////////////////////////////////////////////////////////
+FMetadataDb::~FMetadataDb()
+{
+}
+
+////////////////////////////////////////////////////////////////////////////////
+uint32 FMetadataDb::GetNum() const
+{
+	return uint32(Map.Num());
+}
+
+////////////////////////////////////////////////////////////////////////////////
+uint32 FMetadataDb::GetCollisionNum() const
+{
+	return Collisions;
+}
+
+////////////////////////////////////////////////////////////////////////////////
+uint32 FMetadataDb::Add(uint64 Owner, uint64 Size, uint32 Alignment, bool bRealloc)
+{
+	static const uint64 Mix = 0x00000100000001b3;
+
+	uint64 Id = 0xcbf29ce484222325;
+	Id ^= Owner;	Id *= Mix;
+	Id ^= Size;		Id *= Mix;
+	Id ^= Alignment;Id *= Mix;
+
+	uint32 SizeTribble = uint32(Size & 0x07);
+	uint32 AlignmentPow2 = uint8(31 - UnsafeCountLeadingZeros(Alignment));
+
+	FEntryInternal Entry;
+	Entry.Id = 0;
+	Entry.Size = uint32(Size >> 3);
+	Entry.Owner = Owner;
+	Entry.Alignment_SizeTribble = uint8((AlignmentPow2 << 3) | SizeTribble);
+	Entry.Alignment_SizeTribble |= (bRealloc == true);
+
+	return AddInternal(Id, Entry);
+}
+
+////////////////////////////////////////////////////////////////////////////////
+uint32 FMetadataDb::AddInternal(uint64 Id, const FEntryInternal& Entry)
+{
+	/* To make this thread safe, only do the look-up in this function. If the
+	 * entry needs adding queue the add up to one side. Flush the queue when it
+	 * is guaranteed there is no other reading going on. Removes need to lock */
+
+	uint32 SmallerId = uint32(Id);
+
+	FEntryInternal& Out = Map.FindOrAdd(SmallerId);
+	if (LIKELY(Out == Entry))
+	{
+		return Id;
+	}
+
+	if (Out.Id == 0)
+	{
+		Out = Entry;
+		Out.Id = SmallerId;
+		return SmallerId;
+	}
+
+	++Collisions;
+	Id = Remix(Id);
+	return AddInternal(Id, Entry); // tail call opt. one hopes...
+}
+
+////////////////////////////////////////////////////////////////////////////////
+const FMetadataDb::FEntry* FMetadataDb::Get(uint32 Id) const
+{
+	const FEntryInternal* Iter = Map.Find(Id);
+	return (const FEntry*)Iter;
+}
+
+////////////////////////////////////////////////////////////////////////////////
+uint64 FMetadataDb::Remix(uint64 Id)
+{
+	return Id * 0x30be8efa499c249dull;
+}
+
+} // namespace TraceServices
+} // namespace Trace
+
+/* vim: set noet : */
diff --git a/Engine/Source/Developer/TraceServices/Private/Analyzers/Memory/MetadataDb.h b/Engine/Source/Developer/TraceServices/Private/Analyzers/Memory/MetadataDb.h
new file mode 100644
--- /dev/null
+++ b/Engine/Source/Developer/TraceServices/Private/Analyzers/Memory/MetadataDb.h
@@ -0,0 +1,85 @@
+// Copyright Epic Games, Inc. All Rights Reserved.
+
+#pragma once
+
+#include "Containers/Map.h"
+
+namespace Trace {
+namespace TraceServices {
+
+////////////////////////////////////////////////////////////////////////////////
+class FMetadataDb
+{
+public:
+	struct FEntry
+	{
+		uint64			GetOwner() const;
+		uint64			GetSize() const;
+		uint32			GetAlignment() const;
+		bool			IsRealloc() const;
+	};
+
+						FMetadataDb();
+						~FMetadataDb();
+	uint32				GetNum() const;
+	uint32				GetCollisionNum() const;
+	uint32				Add(uint64 Owner, uint64 Size, uint32 Alignment, bool bRealloc);
+	const FEntry*		Get(uint32 Id) const;
+
+private:
+	struct FEntryInternal
+	{
+		bool			operator == (const FEntryInternal& Rhs) const;
+		uint32			Id;
+		uint32			Size;
+		union
+		{
+			uint64		Owner;
+			struct
+			{
+				uint8	_Padding[7];
+				uint8	Alignment_SizeTribble; // alignment = pow2
+			};
+		};
+	};
+
+	typedef TMap<uint32, FEntryInternal> MetadataMap;
+
+	static uint64		Remix(uint64 Id);
+	uint32				AddInternal(uint64 Id, const FEntryInternal& Entry);
+	uint32				Collisions = 0;
+	MetadataMap			Map;
+};
+
+////////////////////////////////////////////////////////////////////////////////
+inline uint64 FMetadataDb::FEntry::GetOwner() const
+{
+	const auto* Internal = (FEntryInternal*)this;
+	return Internal->Owner & 0x00ff'ffff'ffff'ffff;
+}
+
+////////////////////////////////////////////////////////////////////////////////
+inline uint64 FMetadataDb::FEntry::GetSize() const
+{
+	const auto* Internal = (FEntryInternal*)this;
+	return (uint64(Internal->Size) << 3) + (Internal->Alignment_SizeTribble & 0x07);
+}
+
+////////////////////////////////////////////////////////////////////////////////
+inline uint32 FMetadataDb::FEntry::GetAlignment() const
+{
+	const auto* Internal = (FEntryInternal*)this;
+	return 1 << (Internal->Alignment_SizeTribble >> 3) & ~1;
+}
+
+////////////////////////////////////////////////////////////////////////////////
+inline bool FMetadataDb::FEntry::IsRealloc() const
+{
+	const auto* Internal = (FEntryInternal*)this;
+	return (Internal->Alignment_SizeTribble & (1 << 3)) != 0;
+}
+
+} // namespace TraceServices
+} // namespace Trace
+
+/* vim: set noet : */
diff --git a/Engine/Source/Developer/TraceServices/Private/Analyzers/Memory/SbifBuilder.h b/Engine/Source/Developer/TraceServices/Private/Analyzers/Memory/SbifBuilder.h
new file mode 100644
--- /dev/null
+++ b/Engine/Source/Developer/TraceServices/Private/Analyzers/Memory/SbifBuilder.h
@@ -0,0 +1,140 @@
+// Copyright Epic Games, Inc. All Rights Reserved.
+
+#pragma once
+
+#include "SbifIdentities.inl"
+
+namespace Trace {
+namespace TraceServices {
+
+class FMetadataDb;
+
+////////////////////////////////////////////////////////////////////////////////
+class FSbifContext
+{
+public:
+				FSbifContext(uint32 CellIndex, uint32 ColumnShift);
+	uint32		GetColumn() const;
+	uint32		GetDepth() const;
+	uint32		GetCellIndex() const;
+
+private:
+				FSbifContext() = default;
+	uint32		Column;
+	uint32		Depth;
+	uint32		BaseSerial;
+
+	friend		class FTracker;
+	friend		class FSbifRetiree;
+};
+
+////////////////////////////////////////////////////////////////////////////////
+inline FSbifContext::FSbifContext(uint32 CellIndex, uint32 ColumnShift)
+{
+	Sbif_GetColumnDepthFromCell(CellIndex, &Column, &Depth);
+
+	uint32 BaseColumn = Sbif_GetBaseColumn(Column, Depth);
+	BaseSerial = BaseColumn << ColumnShift;
+}
+
+////////////////////////////////////////////////////////////////////////////////
+inline uint32 FSbifContext::GetColumn() const
+{
+	return Column;
+}
+
+////////////////////////////////////////////////////////////////////////////////
+inline uint32 FSbifContext::GetDepth() const
+{
+	return Depth;
+}
+
+////////////////////////////////////////////////////////////////////////////////
+inline uint32 FSbifContext::GetCellIndex() const
+{
+	return Sbif_GetCellAtDepth(Column, Depth);
+}
+
+
+
+////////////////////////////////////////////////////////////////////////////////
+class FSbifRetiree
+{
+public:
+	uint64				GetAddress() const;
+	uint32				GetStartSerial(const FSbifContext& Context) const;
+	uint32				GetEndSerial(const FSbifContext& Context) const;
+	uint32				GetDuration(const FSbifContext& Context) const;
+	uint32				GetMetadataId() const;
+
+private:
+	friend				class FTracker;
+	uint32				_SpaceForAddress; // change metadata from an ID to an index to make space.
+	uint32				StartSerial;
+	uint32				EndSerial;
+	uint32				MetadataId;
+};
+static_assert(sizeof(FSbifRetiree), "Designed to be XMM-sized");
+
+////////////////////////////////////////////////////////////////////////////////
+inline uint64 FSbifRetiree::GetAddress() const
+{
+	// More work required to free up 12 more bits to store 44 bits active in an
+	// allocation's address. Where do we get 12 more bits from? Refactor the
+	// metadata ID to be an index instead of a hash. Use Context to leverage
+	// detail to encode start/end better.
+	return _SpaceForAddress;
+}
+
+////////////////////////////////////////////////////////////////////////////////
+/*
+inline uint32 FSbifRetiree::GetMidSerial(const FSbifContext& Context) const
+{
+	return Context.BaseSerial + MidSerialOffset;
+}
+*/
+
+////////////////////////////////////////////////////////////////////////////////
+inline uint32 FSbifRetiree::GetDuration(const FSbifContext& Context) const
+{
+	// return Context.MinDuration + DurationTopup;
+	return GetEndSerial(Context) - GetStartSerial(Context);
+}
+
+////////////////////////////////////////////////////////////////////////////////
+inline uint32 FSbifRetiree::GetStartSerial(const FSbifContext& Context) const
+{
+	//return GetMidSerial(Context) - ((GetDuration(Context) + 1) >> 1);
+	return StartSerial + Context.BaseSerial;
+}
+
+////////////////////////////////////////////////////////////////////////////////
+inline uint32 FSbifRetiree::GetEndSerial(const FSbifContext& Context) const
+{
+	//return GetMidSerial(Context) + (GetDuration(Context) >> 1);
+	return EndSerial + Context.BaseSerial;
+}
+
+////////////////////////////////////////////////////////////////////////////////
+inline uint32 FSbifRetiree::GetMetadataId() const
+{
+	return MetadataId;
+}
+
+
+
+////////////////////////////////////////////////////////////////////////////////
+class ISbifBuilder
+{
+public:
+	virtual uint32		GetEventsPerColumn() const = 0;
+	virtual void		Begin(const FMetadataDb* MetadataDb) = 0;
+	virtual void		End() = 0;
+	virtual void		AddColumn() = 0;
+	virtual void		AddRetirees(const FSbifContext* Context, const FSbifRetiree* Retirees, uint32 Num) = 0;
+};
+
+} // namespace TraceServices
+} // namespace Trace
+
+/* vim: set noet : */
diff --git a/Engine/Source/Developer/TraceServices/Private/Analyzers/Memory/SbifIdentities.inl b/Engine/Source/Developer/TraceServices/Private/Analyzers/Memory/SbifIdentities.inl
new file mode 100644
--- /dev/null
+++ b/Engine/Source/Developer/TraceServices/Private/Analyzers/Memory/SbifIdentities.inl
@@ -0,0 +1,56 @@
+// Copyright Epic Games, Inc. All Rights Reserved.
+
+#pragma once
+
+#include "Support.h"
+
+namespace Trace {
+namespace TraceServices {
+
+////////////////////////////////////////////////////////////////////////////////
+inline uint32 Sbif_GetMaxDepth(uint32 TotalColumns)
+{
+	/* we're always expecting a non-zero value */
+	return 32 - UnsafeCountLeadingZeros(TotalColumns);
+}
+
+////////////////////////////////////////////////////////////////////////////////
+inline uint32 Sbif_GetCellAtDepth(uint32 Column, uint32 Depth)
+{
+	uint32 LeafIndex = Column * 2;
+	uint32 k = 1 << Depth;
+	return (LeafIndex & ~k) | (k - 1);
+}
+
+////////////////////////////////////////////////////////////////////////////////
+inline uint32 Sbif_GetCommonDepth(uint32 ColumnA, uint32 ColumnB)
+{
+	uint32 Xor = ColumnA ^ ColumnB;
+	return Sbif_GetMaxDepth(Xor);
+}
+
+////////////////////////////////////////////////////////////////////////////////
+inline uint32 Sbif_GetCellWidth(uint32 CellIndex)
+{
+	return ((CellIndex ^ (CellIndex + 1)) >> 1) + 1;
+}
+
+////////////////////////////////////////////////////////////////////////////////
+inline uint32 Sbif_GetBaseColumn(uint32 Column, uint32 Depth)
+{
+	return Column & ~((1 << Depth) - 1);
+}
+
+////////////////////////////////////////////////////////////////////////////////
+inline void Sbif_GetColumnDepthFromCell(uint32 CellIndex, uint32* Column, uint32* Depth)
+{
+	uint32 PlusOne = CellIndex + 1;
+	uint32 LeastUnset = PlusOne ^ -int32(PlusOne);
+	*Column = (CellIndex & PlusOne) >> 1;
+	*Depth = 31 - UnsafeCountLeadingZeros(LeastUnset);
+}
+
+} // namespace TraceServices
+} // namespace Trace
+
+/* vim: set noet : */
diff --git a/Engine/Source/Developer/TraceServices/Private/Analyzers/Memory/Support.h b/Engine/Source/Developer/TraceServices/Private/Analyzers/Memory/Support.h
new file mode 100644
--- /dev/null
+++ b/Engine/Source/Developer/TraceServices/Private/Analyzers/Memory/Support.h
@@ -0,0 +1,110 @@
+// Copyright Epic Games, Inc. All Rights Reserved.
+
+#pragma once
+
+#include "HAL/UnrealMemory.h"
+
+#if defined(_MSC_VER)
+#	include <intrin.h>
+#endif
+
+namespace Trace {
+namespace TraceServices {
+
+////////////////////////////////////////////////////////////////////////////////
+inline uint32 UnsafeCountLeadingZeros(uint32 Value)
+{
+	/* it is assumed that the caller knows not to counting leading zeros of 0 */
+#ifdef _MSC_VER
+	return __lzcnt(Value);
+#else
+	return __builtin_clz(Value);
+#endif
+}
+
+#define PROF_SCOPE(...)
+
+////////////////////////////////////////////////////////////////////////////////
+class FTrackerBuffer
+{
+public:
+	template <class T> static T*	AllocTemp();
+	template <class T> static T*	CallocTemp(uint32 Max);
+	static void						FreeTemp(void*);
+
+	template <class T> static T*	Alloc();
+	template <class T> static T*	AllocRaw(SIZE_T Size, uint32 Alignment=alignof(T));
+	template <class T> static T*	Calloc(uint32 Max);
+	template <class T> static T*	Realloc(T* Prev, uint32 Max);
+	static void						Free(void*);
+
+private:
+};
+
+////////////////////////////////////////////////////////////////////////////////
+template <class T>
+inline T* FTrackerBuffer::AllocTemp()
+{
+	return Alloc<T>();
+}
+
+////////////////////////////////////////////////////////////////////////////////
+template <class T>
+inline T* FTrackerBuffer::CallocTemp(uint32 Num)
+{
+	return Calloc<T>(Num);
+}
+
+////////////////////////////////////////////////////////////////////////////////
+inline void FTrackerBuffer::FreeTemp(void* Ptr)
+{
+	Free(Ptr);
+}
+
+////////////////////////////////////////////////////////////////////////////////
+template <class T>
+inline T* FTrackerBuffer::Alloc()
+{
+	void* Out = FMemory::Malloc(sizeof(T), alignof(T));
+	return new (Out) T();
+}
+
+////////////////////////////////////////////////////////////////////////////////
+template <class T>
+inline T* FTrackerBuffer::AllocRaw(SIZE_T Size, uint32 Alignment)
+{
+	void* Out = FMemory::Malloc(Size, Alignment);
+	return (T*)Out;
+}
+
+////////////////////////////////////////////////////////////////////////////////
+template <class T>
+inline T* FTrackerBuffer::Calloc(uint32 Num)
+{
+	SIZE_T Bytes = sizeof(T) + (sizeof(T::Items[0]) * Num);
+	void* Ptr = FMemory::Malloc(Bytes, alignof(T));
+	T* Out = new (Ptr) T();
+	Out->Max = Num;
+	return Out;
+}
+
+////////////////////////////////////////////////////////////////////////////////
+template <class T>
+inline T* FTrackerBuffer::Realloc(T* Prev, uint32 Num)
+{
+	SIZE_T Bytes = sizeof(T) + (sizeof(T::Items[0]) * Num);
+	auto* Out = (T*)FMemory::Realloc(Prev, Bytes, alignof(T));
+	Out->Max = Num;
+	return Out;
+}
+
+////////////////////////////////////////////////////////////////////////////////
+inline void FTrackerBuffer::Free(void* Ptr)
+{
+	FMemory::Free(Ptr);
+}
+
+} // namespace TraceServices
+} // namespace Trace
+
+/* vim: set noet : */
diff --git a/Engine/Source/Developer/TraceServices/Private/Analyzers/Memory/Tracker.cpp b/Engine/Source/Developer/TraceServices/Private/Analyzers/Memory/Tracker.cpp
new file mode 100644
--- /dev/null
+++ b/Engine/Source/Developer/TraceServices/Private/Analyzers/Memory/Tracker.cpp
@@ -0,0 +1,554 @@
+// Copyright Epic Games, Inc. All Rights Reserved.
+
+#include "Tracker.h"
+#include "CoreMinimal.h"
+#include "Algo/BinarySearch.h"
+#include "Lane.h"
+#include "TrackerJobs.h"
+#include "SbifBuilder.h"
+#include "Support.h"
+
+namespace Trace {
+namespace TraceServices {
+
+////////////////////////////////////////////////////////////////////////////////
+namespace FTrackerScheduler
+{
+	typedef UPTRINT			FJobHandle;
+	typedef UPTRINT			FWaitHandle;
+	void					SubmitJob(FJobHandle Job) {}
+	FJobHandle				CreateJob(const ANSICHAR* Name) { return 0; }
+	void					DependsOn(FJobHandle This, FJobHandle That) {}
+	void					StartAfter(FJobHandle This, FJobHandle RunsFirst) {}
+	FWaitHandle				MakeWaitable(FJobHandle Job, void* Data=nullptr) { return UPTRINT(Data); }
+	template <class T> T*	Wait(FWaitHandle Waitable) { return (T*)Waitable; }
+
+	template <class T, class U>
+	FJobHandle CreateJob(const ANSICHAR* Name, void (*EntryFunc)(T*), U* Data=nullptr)
+	{
+		EntryFunc(Data);
+		return 0;
+	}
+};
+
+
+
+////////////////////////////////////////////////////////////////////////////////
+FTracker::FTracker(ISbifBuilder* InSbifBuilder)
+: SbifBuilder(InSbifBuilder)
+{
+	for (FLane*& Lane : Lanes)
+	{
+		Lane = new FLane();
+	}
+}
+
+////////////////////////////////////////////////////////////////////////////////
+FTracker::~FTracker()
+{
+	for (FLane* Lane : Lanes)
+	{
+		delete Lane;
+	}
+}
+
+////////////////////////////////////////////////////////////////////////////////
+void FTracker::Begin()
+{
+	SerialBias = 0;
+	Serial = 0;
+	SyncWait = 0;
+	NumColumns = 1;
+
+	uint32 EventsPerColumn = SbifBuilder->GetEventsPerColumn();
+	if (EventsPerColumn == 0)
+	{
+		EventsPerColumn = 128 << 10;
+	}
+	ColumnShift = 31 - UnsafeCountLeadingZeros(EventsPerColumn);
+	ColumnMask = (1 << ColumnShift) - 1;
+	check(ColumnShift <= FTrackerConfig::MaxSerialBits);
+
+	SbifBuilder->Begin(&MetadataDb);
+
+	Provision();
+}
+
+////////////////////////////////////////////////////////////////////////////////
+void FTracker::End()
+{
+	Finalize();
+	SbifBuilder->End();
+}
+
+////////////////////////////////////////////////////////////////////////////////
+uint32 FTracker::GetCurrentSerial() const
+{
+	return Serial + SerialBias;
+}
+
+////////////////////////////////////////////////////////////////////////////////
+void FTracker::AddAlloc(uint64 Address, const FMetadata& Metadata)
+{
+#if defined(PROF_ON)
+	++AllocDelta;
+#endif
+
+	uint32 MetadataId = MetadataDb.Add(
+		Metadata.Owner,
+		Metadata.Size,
+		Metadata.Alignment,
+		Metadata.bIsRealloc
+	);
+
+	FLaneInput* Input = GetLaneInput(Address);
+	bool bLaneFull = Input->AddAlloc(Address, Serial, MetadataId);
+	Update(bLaneFull);
+}
+
+////////////////////////////////////////////////////////////////////////////////
+void FTracker::AddFree(uint64 Address)
+{
+#if defined(PROF_ON)
+	++FreeDelta;
+#endif
+
+	FLaneInput* Input = GetLaneInput(Address);
+	bool bLaneFull = Input->AddFree(Address, Serial);
+	Update(bLaneFull);
+}
+
+////////////////////////////////////////////////////////////////////////////////
+void FTracker::Update(bool bDispatch)
+{
+	++Serial;
+	
+	if (!bDispatch)
+	{
+		return;
+	}
+
+	Dispatch();
+	Provision();
+}
+
+////////////////////////////////////////////////////////////////////////////////
+void FTracker::Finalize()
+{
+	// Flush lanes. Provision() usually tweaks SerialBias but we don't want to
+	// provision so we'll do it here ourselves.
+	Dispatch(false);
+
+	SerialBias += Serial;
+	Serial = 0;
+
+	if (FLaneJobData* DoneJobData = Sync())
+	{
+		FinalizeWork(DoneJobData);
+	}
+
+	// Reuse the last serial submitted to the Sbif for leaks. Saves dealing with
+	// column overflow by one serial
+	SerialBias -= 1;
+
+	// Collect leaks
+	FTrackerScheduler::FJobHandle SyncJob = FTrackerScheduler::CreateJob("Sync");
+	FTrackerScheduler::FJobHandle RootJob = FTrackerScheduler::CreateJob("Root");
+
+	FLeakJobData* LastJobData = nullptr;
+	for (uint32 i = 0; i < FTrackerConfig::NumLanes; ++i)
+	{
+		const FLane* Lane = Lanes[i];
+		const FLaneItemSet& ActiveSet = Lane->GetActiveSet();
+
+		auto* JobData = FTrackerBuffer::AllocTemp<FLeakJobData>();
+		JobData->Next = LastJobData;
+		JobData->ActiveSet = &ActiveSet;
+		JobData->SerialBias = SerialBias;
+		JobData->ColumnShift = ColumnShift;
+		LastJobData = JobData;
+
+		auto LeakJob = FTrackerScheduler::CreateJob("LaneLeaks", LaneLeaksJob, JobData);
+		FTrackerScheduler::StartAfter(LeakJob, RootJob);
+		FTrackerScheduler::DependsOn(SyncJob, LeakJob);
+	}
+
+	FTrackerScheduler::FWaitHandle LeakWait = FTrackerScheduler::MakeWaitable(SyncJob, LastJobData);
+	FTrackerScheduler::SubmitJob(RootJob);
+	auto* JobData = FTrackerScheduler::Wait<FLeakJobData>(LeakWait);
+
+	do
+	{
+		if (JobData->Retirees != nullptr)
+		{
+			ProcessRetirees(JobData);
+		}
+
+		FLeakJobData* NextData = JobData->Next;
+		FTrackerBuffer::FreeTemp(JobData->Retirees);
+		FTrackerBuffer::FreeTemp(JobData);
+		JobData = NextData;
+	}
+	while (JobData != nullptr);
+
+	// Finish off the Sbif
+	/*
+	uint32 MaxCell = (NumColumns << 1) - 1;
+	for (uint32 i = 0, n = Sbif_GetMaxDepth(NumColumns); i <= n; ++i)
+	{
+		uint32 Delta = 2 << i;
+		uint32 CellIndex = Sbif_GetCellAtDepth(NumColumns - 1, i);
+		for (auto j : {0, 1})
+		{
+			if (CellIndex < MaxCell)
+			{
+				SbifBuilder->CloseCell(CellIndex);
+			}
+			CellIndex -= Delta;
+		}
+	}
+	*/
+}
+
+////////////////////////////////////////////////////////////////////////////////
+void FTracker::Provision()
+{
+	for (FLaneInput*& LaneInput : LaneInputs)
+	{
+		LaneInput = FTrackerBuffer::CallocTemp<FLaneInput>(FTrackerConfig::MaxLaneInputItems);
+	}
+
+	SerialBias += Serial;
+	Serial = 0;
+}
+
+////////////////////////////////////////////////////////////////////////////////
+FLaneJobData* FTracker::Sync()
+{
+	if (!SyncWait)
+	{
+		return nullptr;
+	}
+
+	PROF_SCOPE("Sync");
+
+	auto* Data = FTrackerScheduler::Wait<FLaneJobData>(SyncWait);
+	SyncWait = 0;
+	return Data;
+}
+
+////////////////////////////////////////////////////////////////////////////////
+void FTracker::FinalizeWork(FLaneJobData* Data)
+{
+	// SerialBias is the start of the next batch. The batch we're are finalising
+	// (e.g. Data) is up to but not including SerialBias so that's how many
+	// columns we're going to need.
+	while (((SerialBias + ColumnMask) >> ColumnShift) > NumColumns)
+	{
+		NumColumns += 1;
+		SbifBuilder->AddColumn();
+	}
+
+	do
+	{
+		if (Data->Retirees->Num)
+		{
+			ProcessRetirees(Data);
+		}
+
+		FLaneJobData* NextData = Data->Next;
+		FTrackerBuffer::FreeTemp(Data->Input);
+		FTrackerBuffer::FreeTemp(Data->Retirees);
+		FTrackerBuffer::FreeTemp(Data);
+		Data = NextData;
+	}
+	while (Data != nullptr);
+	
+	/*
+	for (; ClosedColumns < NumColumns - 2; ++ClosedColumns)
+	{
+		// Columns to close lag two behind the active column. OEIS' A091090 tells us
+		// how many cells are due to close. Cell indices step by 3i, weirdly enough
+
+		uint32 CellIndex = ClosedColumns << 1;
+
+		uint32 C1 = ClosedColumns + 1;
+		uint32 C2 = ClosedColumns + 2;
+		uint32 A = 31 - UnsafeCountLeadingZeros(C1 ^ C2);
+		uint32 B = (C1 & C2) != 0;
+
+		for (uint32 i = 0, n = A + B, d = 3; i < n; ++i, d <<= 1)
+		{
+			SbifBuilder->CloseCell(CellIndex);
+			CellIndex -= d;
+		}
+		
+		/ *
+		uint32 CellIndex = ClosedColumns << 1;
+		uint32 Num = ClosedColumns + 2;
+		Num = (Num & -Num); // extracts lowest set bit
+		for (uint32 d = 3; ; d <<= 1) // As if by magic a multiple of three turns up!!
+		{
+			SbifBuilder->CloseCell(CellIndex);
+			CellIndex -= d;
+
+			Num >>= 1;
+			if (Num == 0 || int32(CellIndex) < 0)
+			{
+				break;
+			}
+		}
+		* /
+	}
+	*/
+}
+
+PRAGMA_DISABLE_OPTIMIZATION
+
+////////////////////////////////////////////////////////////////////////////////
+void FTracker::RetireAllocs(
+	const FRetireeJobData* Data,
+	uint32 Column,
+	uint32 Depth,
+	const FRetiree* __restrict Left,
+	const FRetiree* __restrict Right)
+{
+	static_assert(sizeof(FRetiree) == sizeof(FSbifRetiree), "These are converted up in place so must be the same size");
+
+	uint32 BaseColumn = Sbif_GetBaseColumn(Column, Depth);
+	uint32 BaseSerial = BaseColumn << ColumnShift;
+
+	uint32 Num = uint32(UPTRINT(Right - Left));
+	for (uint32 i = 0; i < Num; ++i)
+	{
+		const FRetiree* __restrict In = Left + i;
+
+		uint64 Address = In->GetAddress();
+		uint32 MetadataId = In->GetMetadataId();
+		uint32 StartSerial = In->GetStartSerial();
+		uint32 EndSerial = In->GetEndSerial(Data->SerialBias);
+
+		auto* __restrict Out = (FSbifRetiree* __restrict)In;
+		Out->_SpaceForAddress = 0xfedbab1e;
+		Out->StartSerial = StartSerial - BaseSerial;
+		Out->EndSerial = EndSerial - BaseSerial;
+		Out->MetadataId = MetadataId;
+
+		check(int32(Out->StartSerial) >= 0);
+		check(int32(Out->EndSerial) >= 0);
+	}
+
+	const auto* __restrict Out = (FSbifRetiree*)Left;
+
+	FSbifContext Context;
+	Context.Column = Column;
+	Context.Depth = Depth;
+	Context.BaseSerial = BaseSerial;
+	SbifBuilder->AddRetirees(&Context, Out, Num);
+}
+
+////////////////////////////////////////////////////////////////////////////////
+void FTracker::ProcessRetirees(const FRetireeJobData* Data)
+{
+	PROF_SCOPE("ProcessRetirees");
+
+	const FRetiree* __restrict Cursor = Data->Retirees->Items;
+	const FRetiree* __restrict End = Cursor + Data->Retirees->Num;
+
+	auto GetColumn = [Data] (const FRetiree* __restrict Retiree)
+	{
+		uint32 EndSerial = Retiree->GetEndSerial(Data->SerialBias);
+		return EndSerial >> Data->ColumnShift;
+	};
+
+	auto GetDepth = [Data] (const FRetiree* __restrict Retiree)
+	{
+		return Sbif_GetCommonDepth(
+			Retiree->GetStartSerial() >> Data->ColumnShift,
+			Retiree->GetEndSerial(Data->SerialBias) >> Data->ColumnShift
+		);
+	};
+
+	auto ByColumn = [&] (uint32 Column, const FRetiree& __restrict Retiree) -> uint32
+	{
+		return Column < GetColumn(&Retiree);
+	};
+
+	auto ByDepth = [&] (uint32 Depth, const FRetiree& __restrict Retiree) -> uint32
+	{
+		return Depth < GetDepth(&Retiree);
+	};
+
+	do
+	{
+		uint32 Depth = GetDepth(Cursor);
+		TArrayView<const FRetiree> Range0(Cursor, int32(UPTRINT(End - Cursor)));
+		const FRetiree* __restrict Next = Cursor + Algo::UpperBound(Range0, Depth, ByDepth);
+
+		uint32 Column = GetColumn(Cursor);
+		if (Column == GetColumn(Next - 1))
+		{
+			/*
+			check(Sbif_GetCellAtDepth(Column, Depth) < (NumColumns * 2) - 1);
+			for (; Cursor < Next; ++Cursor)
+			{
+				uint32 a = Cursor->GetStartSerial();
+				uint32 b = Cursor->GetEndSerial(Data->SerialBias);
+				uint32 bs = Cursor->GetBiasedSerial();
+				check(GetColumn(Cursor) == Column);
+				check(GetDepth(Cursor) == Depth);
+			}
+			*/
+			RetireAllocs(Data, Column, Depth, Cursor, Next);
+		}
+		else
+		{
+			do
+			{
+				TArrayView<const FRetiree> Range1(Cursor, int32(UPTRINT(Next - Cursor)));
+				const FRetiree* __restrict SubNext = Cursor + Algo::UpperBound(Range1, Column, ByColumn);
+
+				/*
+				check(Sbif_GetCellAtDepth(Column, Depth) < (NumColumns * 2) - 1);
+				for (; Cursor < SubNext; ++Cursor)
+				{
+					uint32 a = Cursor->GetStartSerial();
+					uint32 b = Cursor->GetEndSerial(Data->SerialBias);
+					uint32 bs = Cursor->GetBiasedSerial();
+					check(GetColumn(Cursor) == Column);
+					check(GetDepth(Cursor) == Depth);
+				}
+				*/
+				RetireAllocs(Data, Column, Depth, Cursor, SubNext);
+
+				Cursor = SubNext;
+				if (Cursor >= Next)
+				{
+					break;
+				}
+			}
+			while (Column = GetColumn(Cursor), true);
+		}
+
+		Cursor = Next;
+	}
+	while (Cursor < End);
+}
+
+////////////////////////////////////////////////////////////////////////////////
+void FTracker::Dispatch(bool bDoRehash)
+{
+#if defined(PROF_ON)
+	Running += AllocDelta - FreeDelta;
+	PROF_COUNTER("stats", {
+		{ "allocs", AllocDelta },
+		{ "frees", FreeDelta },
+		{ "alive", Running },
+	});
+
+	PROF_COUNTER("metadata", {
+		{ "#", int64(MetadataDb.GetNum()) },
+		{ "collisions", int64(MetadataDb.GetCollisionNum()) },
+	});
+
+	AllocDelta = FreeDelta = 0;
+#endif
+
+	FTrackerScheduler::FJobHandle LaneJobs[FTrackerConfig::NumLanes] = {};
+	if (bDoRehash)
+	{
+		for (uint32 i = 0; i < FTrackerConfig::NumLanes; ++i)
+		{
+			FLaneItemSet& ActiveSet = Lanes[i]->GetActiveSet();
+			uint32 Capacity = ActiveSet.GetCapacity();
+			uint32 Num = ActiveSet.GetNum();
+			uint32 Tombs = ActiveSet.GetNumTombs();
+
+			/* The below needs more thought */
+
+			uint32 AllocNum = LaneInputs[i]->GetAllocs().Num();
+
+			bool bGrow = (Capacity - Num) <= AllocNum; // Not enough space to fit potential allocs
+
+			bGrow |= !Capacity || ((100 * Num) / Capacity) > 88;				// Load % is getting too high
+			bool bRehash = Num && ((100 * Tombs) / (Tombs + Num)) > 50;// Too tombstoney
+
+			if (bGrow | bRehash)
+			{
+				UPTRINT JobData = UPTRINT(Lanes[i]);
+				JobData ^= bGrow ? ~0ull : 0ull;
+				LaneJobs[i] = FTrackerScheduler::CreateJob("LaneRehash", LaneRehashJob, (void*)JobData);
+			}
+		}
+	}
+
+	FTrackerScheduler::FJobHandle SyncJob = FTrackerScheduler::CreateJob("Sync");
+
+	// Lane input jobs
+	FLaneJobData* LastJobData = nullptr;
+	for (uint32 i = 0; i < FTrackerConfig::NumLanes; ++i)
+	{
+		FLaneInput* LaneInput = LaneInputs[i];
+		if (LaneInput->GetNum() == 0)
+		{
+			continue;
+		}
+
+		auto* JobData = FTrackerBuffer::AllocTemp<FLaneJobData>();
+		JobData->Next = LastJobData;
+		JobData->Lane = Lanes[i];
+		JobData->Input = LaneInput;
+		JobData->SerialBias = SerialBias;
+		JobData->ColumnShift = ColumnShift;
+		LastJobData = JobData;
+
+		auto InputJob = FTrackerScheduler::CreateJob("LaneInput", LaneInputJob, JobData);
+		auto UpdateJob = FTrackerScheduler::CreateJob("LaneUpdate", LaneUpdateJob, JobData);
+		auto RetireeJob = FTrackerScheduler::CreateJob("LaneRetiree", LaneRetireeJob, JobData);
+
+		FTrackerScheduler::StartAfter(UpdateJob, InputJob);
+		FTrackerScheduler::StartAfter(RetireeJob, InputJob);
+
+		FTrackerScheduler::DependsOn(SyncJob, UpdateJob);
+		FTrackerScheduler::DependsOn(SyncJob, RetireeJob);
+
+		if (LaneJobs[i])
+		{
+			FTrackerScheduler::StartAfter(InputJob, LaneJobs[i]);
+		}
+		else
+		{
+			LaneJobs[i] = InputJob;
+		}
+	}
+
+	FTrackerScheduler::FJobHandle RootJob = FTrackerScheduler::CreateJob("Root");
+	for (FTrackerScheduler::FJobHandle LaneJob : LaneJobs)
+	{
+		if (LaneJob)
+		{
+			FTrackerScheduler::StartAfter(LaneJob, RootJob);
+		}
+	}
+
+	FLaneJobData* DoneJobData = Sync();
+	SyncWait = FTrackerScheduler::MakeWaitable(SyncJob, LastJobData);
+	FTrackerScheduler::SubmitJob(RootJob);
+	if (DoneJobData != nullptr)
+	{
+		FinalizeWork(DoneJobData);
+	}
+}
+
+////////////////////////////////////////////////////////////////////////////////
+FLaneInput* FTracker::GetLaneInput(uint64 Address)
+{
+	uint32 SmallKey = uint32(Address >> 4);
+	uint32 Index = ((SmallKey * 0x397666d) >> 8) & (FTrackerConfig::NumLanes - 1);
+	return LaneInputs[Index];
+}
+
+} // namespace TraceServices
+} // namespace Trace
+
+/* vim: set noet : */
diff --git a/Engine/Source/Developer/TraceServices/Private/Analyzers/Memory/Tracker.h b/Engine/Source/Developer/TraceServices/Private/Analyzers/Memory/Tracker.h
new file mode 100644
--- /dev/null
+++ b/Engine/Source/Developer/TraceServices/Private/Analyzers/Memory/Tracker.h
@@ -0,0 +1,69 @@
+// Copyright Epic Games, Inc. All Rights Reserved.
+
+#pragma once
+
+#include "Config.h"
+#include "MetadataDb.h"
+
+namespace Trace {
+namespace TraceServices {
+
+class FLane;
+class FLaneInput;
+struct FLaneJobData;
+class FRetiree;
+struct FRetireeJobData;
+class ISbifBuilder;
+
+////////////////////////////////////////////////////////////////////////////////
+class FTracker
+{
+public:
+	struct FMetadata
+	{
+		uint64				Owner;
+		uint64				Size; 
+		uint32				Alignment;
+		bool				bIsRealloc;
+	};
+
+							FTracker(ISbifBuilder* InSbifBuilder);
+							~FTracker();
+	void					Begin();
+	void					End();
+	void					AddAlloc(uint64 Address, const FMetadata& TeamTada);
+	void					AddFree(uint64 Address);
+	uint32					GetCurrentSerial() const;
+
+private:
+	FLaneInput*				GetLaneInput(uint64 Address);
+	void					Update(bool bDispatch);
+	void					ProcessRetirees(const FRetireeJobData* JobData);
+	void					RetireAllocs(const FRetireeJobData* Data, uint32 Column, uint32 Depth, const FRetiree* __restrict Left, const FRetiree* __restrict Right);
+	FLaneJobData*			Sync();
+	void					Dispatch(bool bDoRehash=true);
+	void					Provision();
+	void					FinalizeWork(FLaneJobData* Data);
+	void					Finalize();
+	FMetadataDb				MetadataDb;
+	FLane*					Lanes[FTrackerConfig::NumLanes];
+	FLaneInput*				LaneInputs[FTrackerConfig::NumLanes];
+	ISbifBuilder*			SbifBuilder;
+	UPTRINT					SyncWait;
+	uint32					Serial;
+	uint32					SerialBias;
+	uint32					ColumnShift;
+	uint32					ColumnMask;
+	uint32					NumColumns;
+
+#if defined(PROF_ON)
+	int32					AllocDelta = 0;
+	int32					FreeDelta = 0;
+	int32					Running = 0;
+#endif
+};
+
+} // namespace TraceServices
+} // namespace Trace
+
+/* vim: set noet : */
diff --git a/Engine/Source/Developer/TraceServices/Private/Analyzers/Memory/TrackerJobs.cpp b/Engine/Source/Developer/TraceServices/Private/Analyzers/Memory/TrackerJobs.cpp
new file mode 100644
--- /dev/null
+++ b/Engine/Source/Developer/TraceServices/Private/Analyzers/Memory/TrackerJobs.cpp
@@ -0,0 +1,315 @@
+// Copyright Epic Games, Inc. All Rights Reserved.
+
+#include "TrackerJobs.h"
+#include "Algo/Sort.h"
+#include "Config.h"
+#include "SbifIdentities.inl"
+#include "Support.h"
+
+namespace Trace {
+namespace TraceServices {
+
+////////////////////////////////////////////////////////////////////////////////
+void LaneRehashJob(void* Data)
+{
+	bool bGrow = PTRINT(Data) < 0;
+	UPTRINT Ptr = UPTRINT(Data) ^ (bGrow ? ~0ull : 0ull);
+	FLane* Lane = (FLane*)Ptr;
+
+	{
+		PROF_SCOPE(bGrow ? "Grow" : "Rehash");
+		Lane->GetActiveSet().Rehash(bGrow ? FTrackerConfig::ActiveSetPageSize : 0);
+	}
+}
+
+////////////////////////////////////////////////////////////////////////////////
+void LaneInputJob(FLaneJobData* Data)
+{
+	FLaneItemView Allocs = Data->Input->GetAllocs();
+	FLaneItemView Frees = Data->Input->GetFrees();
+
+	// Get everything in address order.
+	{
+		PROF_SCOPE("AllocSort", {{"#", Allocs.Num}});
+		Algo::Sort(Allocs);
+	}
+	{
+		PROF_SCOPE("FreeSort", {{"#", Frees.Num}});
+		Algo::Sort(Frees);
+	}
+
+	auto* Retirees = FTrackerBuffer::CallocTemp<FRetirees>(Frees.Num());
+
+	struct FLaneItemIter
+	{
+		FLaneItem const* __restrict	Cursor;
+		FLaneItem* __restrict		Write;
+		FLaneItem const* __restrict	End;
+	};
+
+	FLaneItemIter Iters[] = {
+		{ Allocs.GetData(), Allocs.GetData(), Allocs.GetData() + Allocs.Num() },
+		{ Frees.GetData(),  Frees.GetData(),  Frees.GetData() + Frees.Num() },
+	};
+
+	if (Allocs.Num() && Frees.Num())
+	{
+		PROF_SCOPE("Short allocs");
+
+		while (true)
+		{
+			const FLaneItem& AllocItem = *(Iters[0].Cursor);
+			const FLaneItem& FreeItem = *(Iters[1].Cursor);
+
+			int32 Parity = (FreeItem < AllocItem);
+			if (AllocItem.IsSameAddress(FreeItem) && !Parity)
+			{
+				uint64 Address = AllocItem.GetAddress();
+				uint32 Start = AllocItem.GetSerial(Data->SerialBias);
+				uint32 EndBiased = FreeItem.GetBiasedSerial();
+				uint32 MetadataId = AllocItem.GetMetadataId();
+
+				FRetiree Retiree;
+				Retiree.Set(Start, EndBiased, Address, MetadataId);
+				uint32 Index = Retirees->Num++;
+				Retirees->Items[Index] = Retiree;
+
+				int Eof = 0;
+				Eof += ++Iters[0].Cursor >= Iters[0].End;
+				Eof += ++Iters[1].Cursor >= Iters[1].End;
+				if (Eof)
+				{
+					break;
+				}
+				continue;
+			}
+
+			FLaneItemIter& Iter = Iters[Parity];
+
+			Iter.Write[0] = Iter.Cursor[0];
+			Iter.Write++;
+
+			if (++Iter.Cursor >= Iter.End)
+			{
+				break;
+			}
+		}
+	}
+
+	// Consolidate remaining allocs
+	FLaneItemIter& AllocIter = Iters[0];
+	if (uint32 Remaining = uint32(UPTRINT(AllocIter.End - AllocIter.Cursor)))
+	{
+		::memmove(AllocIter.Write, AllocIter.Cursor, sizeof(FLaneItem) * Remaining);
+		AllocIter.Write += Remaining;
+	}
+
+	// Consolidate remaining frees
+	FLaneItemIter& FreeIter = Iters[1];
+	if (uint32 Remaining = uint32(UPTRINT(FreeIter.End - FreeIter.Cursor)))
+	{
+		::memmove(FreeIter.Write, FreeIter.Cursor, sizeof(FLaneItem) * Remaining);
+		FreeIter.Write += Remaining;
+	}
+
+	// Resolve remaining frees.
+	FLaneItem* __restrict AllocWrite = AllocIter.Write;
+	FLaneItem* __restrict FreeWrite = FreeIter.Write;
+	{
+		uint32 FreesNum = uint32(UPTRINT(FreeWrite - Frees.GetData()));
+
+		PROF_SCOPE("RemainingFrees");
+
+		FLaneItemSet& ActiveSet = Data->Lane->GetActiveSet();
+
+		for (uint32 i = 0, n = FreesNum; i < n; ++i)
+		{
+			FLaneItem& FreeItem = Frees[i];
+
+			uint64 Address = FreeItem.GetAddress();
+			int32 ActiveIndex = ActiveSet.Find(Address);
+			if (ActiveIndex < 0)
+			{
+				continue;
+			}
+
+			uint32 Start = ActiveSet.GetSerial(ActiveIndex);
+			uint32 EndBiased = FreeItem.GetBiasedSerial();
+			uint32 MetadataId = ActiveSet.GetMetadataId(ActiveIndex);
+
+			FRetiree Retiree;
+			Retiree.Set(Start, EndBiased, Address, MetadataId);
+			int32 RetireeIndex = Retirees->Num++;
+			Retirees->Items[RetireeIndex] = Retiree;
+
+			AllocWrite[0] = FreeItem;
+			AllocWrite->SetActiveIndex(ActiveIndex);
+			AllocWrite++;
+		}
+	}
+
+	Data->Retirees = Retirees;
+
+	uint32 AllocsNum = uint32(UPTRINT(AllocWrite - Allocs.GetData()));
+	Data->SetUpdates = FLaneItemView(Allocs.GetData(), AllocsNum);
+}
+
+////////////////////////////////////////////////////////////////////////////////
+void LaneUpdateJob(FLaneJobData* Data)
+{
+	FLaneItemView SetUpdates = Data->SetUpdates;
+
+	if (SetUpdates.Num() == 0)
+	{
+		return;
+	}
+
+	{
+		PROF_SCOPE("UpdateSort");
+
+		auto Predicate = [] (const FLaneItem& Lhs, const FLaneItem& Rhs)
+		{
+			if (Lhs.IsSameAddress(Rhs))
+			{
+				return Rhs.HasMetadata();
+			}
+
+			return Lhs < Rhs;
+		};
+		Algo::Sort(SetUpdates, Predicate);
+	}
+
+	FLaneItemSet& ActiveSet = Data->Lane->GetActiveSet();
+	const FLaneItem* __restrict Ptr = SetUpdates.GetData();
+
+	PROF_SCOPE("SetUpdate", {
+		{"update", SetUpdates.Num()},
+		{"set", ActiveSet.GetNum()},
+	});
+
+	auto ApplyToSet = [&ActiveSet, Data] (const FLaneItem& Item)
+	{
+		if (Item.HasMetadata())
+		{
+			uint64 Address = Item.GetAddress();
+			uint32 Serial = Item.GetSerial(Data->SerialBias);
+			uint32 MetadataId = Item.GetMetadataId();
+			ActiveSet.Add(Address, Serial, MetadataId);
+		}
+		else
+		{
+			uint32 Index = Item.GetActiveIndex();
+			ActiveSet.Remove(Index);
+		}
+	};
+
+	Data->Lane->LockWrite();
+	if (SetUpdates.Num() == 1)
+	{
+		ApplyToSet(*Ptr);
+	}
+	else
+	{
+		const FLaneItem* __restrict Prev = Ptr;
+
+		uint32 i = 1, n = SetUpdates.Num();
+		do
+		{
+			const FLaneItem& Item = Ptr[i];
+
+			if (Item.IsSameAddress(*Prev))
+			{
+				uint32 Index = Prev->GetActiveIndex();
+				uint32 Serial = Item.GetSerial(Data->SerialBias);
+				uint32 MetadataId = Item.GetMetadataId();
+				ActiveSet.Update(Index, Serial, MetadataId);
+
+				++i;
+				Prev = Ptr + i;
+			}
+			else
+			{
+				ApplyToSet(*Prev);
+				Prev = &Item;
+			}
+		}
+		while (++i < n);
+
+		if (i == n)
+		{
+			ApplyToSet(*Prev);
+		}
+	}
+	Data->Lane->UnlockWrite();
+}
+
+////////////////////////////////////////////////////////////////////////////////
+void LaneRetireeJob(FRetireeJobData* Data)
+{
+	FRetirees* Retirees = Data->Retirees;
+
+	auto GetDepth = [Data] (const FRetiree& Retiree)
+	{
+		return Sbif_GetCommonDepth(
+			Retiree.GetStartSerial() >> Data->ColumnShift,
+			Retiree.GetEndSerial(Data->SerialBias) >> Data->ColumnShift
+		);
+	};
+
+	auto Predicate = [&GetDepth] (const FRetiree& Lhs, const FRetiree& Rhs)
+	{
+		uint32 LhsDepth = GetDepth(Lhs);
+		uint32 RhsDepth = GetDepth(Rhs);
+		if (LhsDepth != RhsDepth)
+		{
+			return LhsDepth < RhsDepth;
+		}
+
+		return Lhs.GetBiasedSerial() < Rhs.GetBiasedSerial();
+	};
+
+	TArrayView<FRetiree> Range(Retirees->Items, Retirees->Num);
+	Algo::Sort(Range, Predicate);
+}
+
+////////////////////////////////////////////////////////////////////////////////
+void LaneLeaksJob(FLeakJobData* Data)
+{
+	const FLaneItemSet* ActiveSet = Data->ActiveSet;
+
+	uint32 Num = ActiveSet->GetNum();
+
+	if (Num == 0)
+	{
+		Data->Retirees = nullptr;
+		return;
+	}
+
+	auto* Retirees = FTrackerBuffer::CallocTemp<FRetirees>(Num);
+	Data->Retirees = Retirees;
+
+	FLaneItemSet::FItemHandle SetItem = ActiveSet->ReadItems();
+	for (uint32 i = 0; i < Num; ++i)
+	{
+		uint32 Index = ActiveSet->GetItemIndex(SetItem);
+
+		uint64 Address = ActiveSet->GetAddress(Index);
+		uint32 Start = ActiveSet->GetSerial(Index);
+		uint32 MetadataId = ActiveSet->GetMetadataId(Index);
+
+		FRetiree Retiree;
+		Retiree.Set(Start, 0, Address, MetadataId);
+		int32 RetireeIndex = Retirees->Num++;
+		Retirees->Items[RetireeIndex] = Retiree;
+
+		SetItem = ActiveSet->NextItem(SetItem);
+	}
+
+	PROF_SCOPE("Retirees");
+	LaneRetireeJob(Data);
+}
+
+} // namespace TraceServices
+} // namespace Trace
+
+/* vim: set noet : */
diff --git a/Engine/Source/Developer/TraceServices/Private/Analyzers/Memory/TrackerJobs.h b/Engine/Source/Developer/TraceServices/Private/Analyzers/Memory/TrackerJobs.h
new file mode 100644
--- /dev/null
+++ b/Engine/Source/Developer/TraceServices/Private/Analyzers/Memory/TrackerJobs.h
@@ -0,0 +1,81 @@
+// Copyright Epic Games, Inc. All Rights Reserved.
+
+#pragma once
+
+#include "Lane.h"
+
+namespace Trace {
+namespace TraceServices {
+
+////////////////////////////////////////////////////////////////////////////////
+class alignas(16) FRetiree
+	: public FAddressSerial
+{
+public:
+	void			Set(uint32 Start, uint32 EndBiased, uint64 Address, uint32 MetadataId);
+	uint32			GetStartSerial() const			{ return StartSerial; }
+	uint32			GetEndSerial(uint32 Bias) const	{ return GetBiasedSerial() + Bias; }
+	uint32			GetMetadataId() const			{ return Id; }
+
+private:
+	uint32			StartSerial;
+	uint32			Id;
+};
+
+////////////////////////////////////////////////////////////////////////////////
+inline void FRetiree::Set(uint32 Start, uint32 EndBiased, uint64 Address, uint32 MetadataId)
+{
+	FAddressSerial::Set(EndBiased, Address);
+	StartSerial = Start;
+	Id = MetadataId;
+}
+
+
+
+////////////////////////////////////////////////////////////////////////////////
+struct FRetirees
+{
+	static_assert(sizeof(FRetiree) == 16, "");
+	
+	uint32			Max;
+	uint32			Num = 0;
+	FRetiree		Items[];
+};
+
+////////////////////////////////////////////////////////////////////////////////
+struct FRetireeJobData
+{
+	FRetirees*			Retirees;
+	uint32				SerialBias;
+	uint32				ColumnShift;
+};
+
+////////////////////////////////////////////////////////////////////////////////
+struct FLaneJobData
+	: public FRetireeJobData
+{
+	FLaneJobData*		Next;
+	FLane*				Lane;
+	FLaneInput*			Input;
+	FLaneItemView		SetUpdates;
+};
+
+////////////////////////////////////////////////////////////////////////////////
+struct FLeakJobData
+	: public FRetireeJobData
+{
+	FLeakJobData*		Next;
+	const FLaneItemSet* ActiveSet;
+};
+
+////////////////////////////////////////////////////////////////////////////////
+void LaneRehashJob(void* Data);
+void LaneInputJob(FLaneJobData* Data);
+void LaneUpdateJob(FLaneJobData* Data);
+void LaneRetireeJob(FRetireeJobData* Data);
+void LaneLeaksJob(FLeakJobData* Data);
+
+} // namespace TraceServices
+} // namespace Trace
+
+/* vim: set noet : */
diff --git a/Engine/Source/Developer/TraceServices/Private/Analyzers/Memory/TrackerTime.h b/Engine/Source/Developer/TraceServices/Private/Analyzers/Memory/TrackerTime.h
new file mode 100644
--- /dev/null
+++ b/Engine/Source/Developer/TraceServices/Private/Analyzers/Memory/TrackerTime.h
@@ -0,0 +1,26 @@
+// Copyright Epic Games, Inc. All Rights Reserved.
+
+#pragma once
+
+namespace Trace {
+namespace TraceServices {
+
+////////////////////////////////////////////////////////////////////////////////
+class FTrackerTime
+{
+public:
+					FTrackerTime(uint32 Period) {}
+	void			AddSample(uint64 Cycles) {}
+	double			GetTimeBySerial(uint32 Serial, bool bInterpolate=false) const;
+	uint64			GetCykelBySerial(uint32 Serial, bool bInterpolate=false) const;
+	uint32			GetSerialByCycle(uint64 Cycle, bool bInterpolate=false) const;
+
+private:
+	TArray<uint64>	Samples;
+	uint32			Period;
+};
+
+} // namespace TraceServices
+} // namespace Trace
+
+/* vim: set noet : */
