Fixed compile issue caused by an unpaired PRAGMA_DISABLE_OPTIMIZATION

diff --git a/Engine/Source/Developer/TraceServices/Private/Analyzers/Memory/Tracker.cpp b/Engine/Source/Developer/TraceServices/Private/Analyzers/Memory/Tracker.cpp
--- a/Engine/Source/Developer/TraceServices/Private/Analyzers/Memory/Tracker.cpp
+++ b/Engine/Source/Developer/TraceServices/Private/Analyzers/Memory/Tracker.cpp
@@ -303,7 +303,7 @@ void FTracker::FinalizeWork(FLaneJobData* Data)
 	*/
 }
 
-PRAGMA_DISABLE_OPTIMIZATION
+//PRAGMA_DISABLE_OPTIMIZATION
 
 ////////////////////////////////////////////////////////////////////////////////
 void FTracker::RetireAllocs(
