Removed commented out PRAGMA_DISABLE_OPT.

diff --git a/Engine/Source/Developer/TraceServices/Private/Analyzers/Memory/Tracker.cpp b/Engine/Source/Developer/TraceServices/Private/Analyzers/Memory/Tracker.cpp
--- a/Engine/Source/Developer/TraceServices/Private/Analyzers/Memory/Tracker.cpp
+++ b/Engine/Source/Developer/TraceServices/Private/Analyzers/Memory/Tracker.cpp
@@ -303,8 +303,6 @@ void FTracker::FinalizeWork(FLaneJobData* Data)
 	*/
 }
 
-//PRAGMA_DISABLE_OPTIMIZATION
-
 ////////////////////////////////////////////////////////////////////////////////
 void FTracker::RetireAllocs(
 	const FRetireeJobData* Data,
