Use period timestamps for allocations in memory analysis

diff --git a/Engine/Source/Developer/TraceServices/Private/Analyzers/AllocationsAnalysis.cpp b/Engine/Source/Developer/TraceServices/Private/Analyzers/AllocationsAnalysis.cpp
--- a/Engine/Source/Developer/TraceServices/Private/Analyzers/AllocationsAnalysis.cpp
+++ b/Engine/Source/Developer/TraceServices/Private/Analyzers/AllocationsAnalysis.cpp
@@ -33,6 +33,7 @@ void FAllocationsAnalyzer::OnAnalysisBegin(const FOnAnalysisContext& Context)
 	Builder.RouteEvent(RouteId_Alloc,      "Memory", "Alloc");
 	Builder.RouteEvent(RouteId_Realloc,    "Memory", "Realloc");
 	Builder.RouteEvent(RouteId_Free,       "Memory", "Free");
+	Builder.RouteEvent(RouteId_Marker,	   "Memory", "Marker");
 }
 
 ////////////////////////////////////////////////////////////////////////////////////////////////////
@@ -57,6 +58,9 @@ bool FAllocationsAnalyzer::OnEvent(uint16 RouteId, EStyle Style, const FOnEventC
 			const uint8 SummarySizeShift = EventData.GetValue<uint8>("SummarySizeShift");
 			const uint8 Mode = EventData.GetValue<uint8>("Mode");
 
+			BaseCycle = EventData.GetValue<uint64>("BaseCycle");
+			MarkerPeriod = EventData.GetValue<uint32>("MarkerPeriod");
+
 			FAllocationsProvider::FEditScopeLock _(AllocationsProvider);
 			AllocationsProvider.EditInit(Time, MinAlignment, SizeShift, SummarySizeShift, Mode);
 			break;
@@ -123,6 +127,12 @@ bool FAllocationsAnalyzer::OnEvent(uint16 RouteId, EStyle Style, const FOnEventC
 			AllocationsProvider.EditFree(Time, Owner, FreeAddress);
 			break;
 		}
+		case RouteId_Marker:
+		{
+			const uint64 LastCycle = BaseCycle + EventData.GetValue<uint32>("Cycle");
+			LastMarkerSeconds = Context.EventTime.AsSeconds(LastCycle);
+			break;
+		}
 	}
 
 	return true;
@@ -132,10 +142,7 @@ bool FAllocationsAnalyzer::OnEvent(uint16 RouteId, EStyle Style, const FOnEventC
 
 double FAllocationsAnalyzer::GetCurrentTime() const
 {
-	// TODO: In order to avoid the session lock, either add a time param in allocation events
-	//       or add a handler for a specific time event.
-	FAnalysisSessionEditScope _(Session);
-	return Session.GetDurationSeconds();
+	return LastMarkerSeconds;
 }
 
 ////////////////////////////////////////////////////////////////////////////////////////////////////
diff --git a/Engine/Source/Developer/TraceServices/Private/Analyzers/AllocationsAnalysis.h b/Engine/Source/Developer/TraceServices/Private/Analyzers/AllocationsAnalysis.h
--- a/Engine/Source/Developer/TraceServices/Private/Analyzers/AllocationsAnalysis.h
+++ b/Engine/Source/Developer/TraceServices/Private/Analyzers/AllocationsAnalysis.h
@@ -22,6 +22,7 @@ private:
 		RouteId_Alloc,
 		RouteId_Realloc,
 		RouteId_Free,
+		RouteId_Marker,
 	};
 
 public:
@@ -36,6 +37,9 @@ private:
 	IAnalysisSession& Session;
 	FAllocationsProvider& AllocationsProvider;
 	uint64 EventCount = 0; // debug
+	uint64 BaseCycle;
+	uint32 MarkerPeriod;
+	double LastMarkerSeconds;
 };
 
 } // namespace TraceServices
