Fixed few build errors when -nopch -disableunity.

diff --git a/Engine/Source/Developer/TraceServices/Private/Model/AllocationsProvider.cpp b/Engine/Source/Developer/TraceServices/Private/Model/AllocationsProvider.cpp
--- a/Engine/Source/Developer/TraceServices/Private/Model/AllocationsProvider.cpp
+++ b/Engine/Source/Developer/TraceServices/Private/Model/AllocationsProvider.cpp
@@ -3,6 +3,7 @@
 #include "AllocationsProvider.h"
 
 #include "AllocationsQuery.h"
+#include "Containers/ArrayView.h"
 #include "SbTree.h"
 #include "TraceServices/Containers/Allocators.h"
 
diff --git a/Engine/Source/Developer/TraceServices/Public/TraceServices/Containers/Allocators.h b/Engine/Source/Developer/TraceServices/Public/TraceServices/Containers/Allocators.h
--- a/Engine/Source/Developer/TraceServices/Public/TraceServices/Containers/Allocators.h
+++ b/Engine/Source/Developer/TraceServices/Public/TraceServices/Containers/Allocators.h
@@ -2,6 +2,8 @@
 
 #pragma once
 
+#include "HAL/Platform.h"
+
 namespace TraceServices
 {
 	
diff --git a/Engine/Source/Runtime/Core/Public/ProfilingDebugging/TagTrace.h b/Engine/Source/Runtime/Core/Public/ProfilingDebugging/TagTrace.h
--- a/Engine/Source/Runtime/Core/Public/ProfilingDebugging/TagTrace.h
+++ b/Engine/Source/Runtime/Core/Public/ProfilingDebugging/TagTrace.h
@@ -1,6 +1,7 @@
 // Copyright Epic Games, Inc. All Rights Reserved.
 #pragma once
 #include "CoreTypes.h"
+#include "HAL/LowLevelMemTracker.h"
 #include "UObject/NameTypes.h"
 #include "Trace/Trace.h"
 #include "Trace/Detail/LogScope.h"
