Fix Insights crash when encountering a GPU event that does not have an event spec

diff --git a/Engine/Source/Developer/TraceServices/Private/Analyzers/GpuProfilerTraceAnalysis.cpp b/Engine/Source/Developer/TraceServices/Private/Analyzers/GpuProfilerTraceAnalysis.cpp
--- a/Engine/Source/Developer/TraceServices/Private/Analyzers/GpuProfilerTraceAnalysis.cpp
+++ b/Engine/Source/Developer/TraceServices/Private/Analyzers/GpuProfilerTraceAnalysis.cpp
@@ -78,6 +78,13 @@ bool FGpuProfilerAnalyzer::OnEvent(uint16 RouteId, EStyle Style, const FOnEventC
 					Event.TimerIndex = EventTypeMap[EventType];
 					Timeline.AppendBeginEvent(LastTime, Event);
 				}
+				else
+				{
+					FTimingProfilerEvent Event;
+					Event.TimerIndex = TimingProfilerProvider.AddGpuTimer(TEXT("<unknown>"));
+					EventTypeMap.Add(EventType, Event.TimerIndex);
+					Timeline.AppendBeginEvent(LastTime, Event);
+				}
 				++CurrentDepth;
 			}
 			else
