Changed metadata IDs to be indices. This frees up some bits so we can better pack structs.

diff --git a/Engine/Source/Developer/TraceServices/Private/Analyzers/Memory/MetadataDb.cpp b/Engine/Source/Developer/TraceServices/Private/Analyzers/Memory/MetadataDb.cpp
--- a/Engine/Source/Developer/TraceServices/Private/Analyzers/Memory/MetadataDb.cpp
+++ b/Engine/Source/Developer/TraceServices/Private/Analyzers/Memory/MetadataDb.cpp
@@ -3,13 +3,15 @@
 #include "MetadataDb.h"
 #include "Support.h"
 
+PRAGMA_DISABLE_OPTIMIZATION
+
 namespace Trace {
 namespace TraceServices {
 
 ////////////////////////////////////////////////////////////////////////////////
 bool FMetadataDb::FEntryInternal::operator == (const FEntryInternal& Rhs) const
 {
-	return (Owner == Rhs.Owner) | (Size == Rhs.Size);
+	return FMemory::Memcmp(this, &Rhs, sizeof(*this)) == 0;
 }
 
 ////////////////////////////////////////////////////////////////////////////////
@@ -40,19 +42,16 @@ uint32 FMetadataDb::Add(uint64 Owner, uint64 Size, uint32 Alignment, bool bReall
 	static const uint64 Mix = 0x00000100000001b3;
 
 	uint64 Id = 0xcbf29ce484222325;
-	Id ^= Owner;	Id *= Mix;
-	Id ^= Size;		Id *= Mix;
-	Id ^= Alignment;Id *= Mix;
-
-	uint32 SizeTribble = uint32(Size & 0x07);
-	uint32 AlignmentPow2 = uint8(31 - UnsafeCountLeadingZeros(Alignment));
+	Id ^= Owner;		Id *= Mix;
+	Id ^= Size;			Id *= Mix;
+	Id ^= Alignment;	Id *= Mix;
+	Id ^= !!bRealloc;	Id *= Mix;
 
-	FEntryInternal Entry;
-	Entry.Id = 0;
-	Entry.Size = uint32(Size >> 3);
+	FEntryInternal Entry = {};
+	Entry.Size = Size;
 	Entry.Owner = Owner;
-	Entry.Alignment_SizeTribble = uint8((AlignmentPow2 << 3) | SizeTribble);
-	Entry.Alignment_SizeTribble |= (bRealloc == true);
+	Entry.Alignment = uint16(Alignment);
+	Entry.bIsRealloc = !!bRealloc;
 
 	return AddInternal(Id, Entry);
 }
@@ -66,17 +65,17 @@ uint32 FMetadataDb::AddInternal(uint64 Id, const FEntryInternal& Entry)
 
 	uint32 SmallerId = uint32(Id);
 
-	FEntryInternal& Out = Map.FindOrAdd(SmallerId);
-	if (LIKELY(Out == Entry))
+	uint32& Index = Map.FindOrAdd(SmallerId, ~0u);
+	if (UNLIKELY(Index == ~0u))
 	{
-		return Id;
+		Index = Entries.Num();
+		Entries.Add(Entry);
+		return Index;
 	}
 
-	if (Out.Id == 0)
+	if (LIKELY(Entries[Index] == Entry))
 	{
-		Out = Entry;
-		Out.Id = SmallerId;
-		return SmallerId;
+		return Index;
 	}
 
 	++Collisions;
@@ -87,8 +86,7 @@ uint32 FMetadataDb::AddInternal(uint64 Id, const FEntryInternal& Entry)
 ////////////////////////////////////////////////////////////////////////////////
 const FMetadataDb::FEntry* FMetadataDb::Get(uint32 Id) const
 {
-	const FEntryInternal* Iter = Map.Find(Id);
-	return (const FEntry*)Iter;
+	return (const FEntry*)((Id < uint32(Entries.Num())) ? &(Entries[Id]) : nullptr);
 }
 
 ////////////////////////////////////////////////////////////////////////////////
diff --git a/Engine/Source/Developer/TraceServices/Private/Analyzers/Memory/MetadataDb.h b/Engine/Source/Developer/TraceServices/Private/Analyzers/Memory/MetadataDb.h
--- a/Engine/Source/Developer/TraceServices/Private/Analyzers/Memory/MetadataDb.h
+++ b/Engine/Source/Developer/TraceServices/Private/Analyzers/Memory/MetadataDb.h
@@ -2,6 +2,7 @@
 
 #pragma once
 
+#include "Containers/Array.h"
 #include "Containers/Map.h"
 
 namespace Trace {
@@ -30,53 +31,51 @@ private:
 	struct FEntryInternal
 	{
 		bool			operator == (const FEntryInternal& Rhs) const;
-		uint32			Id;
-		uint32			Size;
-		union
-		{
-			uint64		Owner;
-			struct
-			{
-				uint8	_Padding[7];
-				uint8	Alignment_SizeTribble; // alignment = pow2
-			};
-		};
+		uint64			Owner;
+		uint64			Size;
+		uint32			Tag;
+		uint16			Alignment;
+		uint8			bIsRealloc	: 1;
+		uint8			_Padding0	: 7;
+		uint8			_Padding1;
 	};
 
-	typedef TMap<uint32, FEntryInternal> MetadataMap;
+	using MetadataMap	= TMap<uint32, uint32>;
+	using MetadataArray	= TArray<FEntryInternal>;
 
 	static uint64		Remix(uint64 Id);
 	uint32				AddInternal(uint64 Id, const FEntryInternal& Entry);
 	uint32				Collisions = 0;
 	MetadataMap			Map;
+	MetadataArray		Entries;
 };
 
 ////////////////////////////////////////////////////////////////////////////////
 inline uint64 FMetadataDb::FEntry::GetOwner() const
 {
 	const auto* Internal = (FEntryInternal*)this;
-	return Internal->Owner & 0x00ff'ffff'ffff'ffff;
+	return Internal->Owner;
 }
 
 ////////////////////////////////////////////////////////////////////////////////
 inline uint64 FMetadataDb::FEntry::GetSize() const
 {
 	const auto* Internal = (FEntryInternal*)this;
-	return (uint64(Internal->Size) << 3) + (Internal->Alignment_SizeTribble & 0x07);
+	return Internal->Size;
 }
 
 ////////////////////////////////////////////////////////////////////////////////
 inline uint32 FMetadataDb::FEntry::GetAlignment() const
 {
 	const auto* Internal = (FEntryInternal*)this;
-	return 1 << (Internal->Alignment_SizeTribble >> 3) & ~1;
+	return Internal->Alignment;
 }
 
 ////////////////////////////////////////////////////////////////////////////////
 inline bool FMetadataDb::FEntry::IsRealloc() const
 {
 	const auto* Internal = (FEntryInternal*)this;
-	return (Internal->Alignment_SizeTribble & (1 << 3)) != 0;
+	return (Internal->bIsRealloc != 0);
 }
 
 } // namespace TraceServices
