Repacked retirees struct so there's plenty of bits for start serials.

diff --git a/Engine/Source/Developer/TraceServices/Private/Analyzers/Memory/TrackerJobs.cpp b/Engine/Source/Developer/TraceServices/Private/Analyzers/Memory/TrackerJobs.cpp
--- a/Engine/Source/Developer/TraceServices/Private/Analyzers/Memory/TrackerJobs.cpp
+++ b/Engine/Source/Developer/TraceServices/Private/Analyzers/Memory/TrackerJobs.cpp
@@ -265,7 +265,7 @@ void LaneRetireeJob(FRetireeJobData* Data)
 			return LhsDepth < RhsDepth;
 		}
 
-		return Lhs.GetBiasedSerial() < Rhs.GetBiasedSerial();
+		return Lhs.GetEndSerialBiased() < Rhs.GetEndSerialBiased();
 	};
 
 	TArrayView<FRetiree> Range(Retirees->Items, Retirees->Num);
diff --git a/Engine/Source/Developer/TraceServices/Private/Analyzers/Memory/TrackerJobs.h b/Engine/Source/Developer/TraceServices/Private/Analyzers/Memory/TrackerJobs.h
--- a/Engine/Source/Developer/TraceServices/Private/Analyzers/Memory/TrackerJobs.h
+++ b/Engine/Source/Developer/TraceServices/Private/Analyzers/Memory/TrackerJobs.h
@@ -9,25 +9,29 @@ namespace TraceServices {
 
 ////////////////////////////////////////////////////////////////////////////////
 class alignas(16) FRetiree
-	: public FAddressSerial
 {
 public:
-	void			Set(uint32 Start, uint32 EndBiased, uint64 Address, uint32 MetadataId);
-	uint32			GetStartSerial() const			{ return StartSerial; }
-	uint32			GetEndSerial(uint32 Bias) const	{ return GetBiasedSerial() + Bias; }
-	uint32			GetMetadataId() const			{ return Id; }
+	void			Set(uint32 Start, uint32 EndBiased, uint64 InAddress, uint32 InMetadataId);
+	uint64			GetAddress() const				{ return Address; }
+	uint64			GetStartSerial() const			{ return StartSerial; }
+	uint32			GetEndSerial(uint32 Bias) const	{ return EndSerialBiased + Bias; }
+	uint32			GetEndSerialBiased() const		{ return EndSerialBiased; }
+	uint32			GetMetadataId() const			{ return MetadataId; }
 
 private:
-	uint32			StartSerial;
-	uint32			Id;
+	uint64			Address			: 44;
+	uint64			EndSerialBiased	: 20;
+	uint64			StartSerial		: 36;
+	uint64			MetadataId		: 28;
 };
 
 ////////////////////////////////////////////////////////////////////////////////
-inline void FRetiree::Set(uint32 Start, uint32 EndBiased, uint64 Address, uint32 MetadataId)
+inline void FRetiree::Set(uint32 Start, uint32 EndBiased, uint64 InAddress, uint32 InMetadataId)
 {
-	FAddressSerial::Set(EndBiased, Address);
+	Address = InAddress;
+	EndSerialBiased = EndBiased;
 	StartSerial = Start;
-	Id = MetadataId;
+	MetadataId = InMetadataId;
 }
 
 
