Remove PRAGMA_DISABLE_OPTIMIZATION

diff --git a/Engine/Source/Developer/TraceServices/Private/Analyzers/Memory/MetadataDb.cpp b/Engine/Source/Developer/TraceServices/Private/Analyzers/Memory/MetadataDb.cpp
--- a/Engine/Source/Developer/TraceServices/Private/Analyzers/Memory/MetadataDb.cpp
+++ b/Engine/Source/Developer/TraceServices/Private/Analyzers/Memory/MetadataDb.cpp
@@ -3,8 +3,6 @@
 #include "MetadataDb.h"
 #include "Support.h"
 
-PRAGMA_DISABLE_OPTIMIZATION
-
 namespace Trace {
 namespace TraceServices {
 
