[Insights]
  - Implemented mem alloc queries also for live allocs.
  - Implemented analysis for llm tags associated with allocations (based on initial work by [REDACTED_PERSON]).
  - Changed granularity of memory allocations timelines to 0.1ms (from 1ms).
  - Fixed allotment of allocations in SbTree's offsetted cells.
  - Fixed queries to handle also cases where multiple cells have exact same timestamps.
  - Added more debug/validation code for AllocationsProvider.

diff --git a/Engine/Source/Developer/TraceInsights/Private/Insights/MemoryProfiler/Widgets/SMemAllocTableTreeView.cpp b/Engine/Source/Developer/TraceInsights/Private/Insights/MemoryProfiler/Widgets/SMemAllocTableTreeView.cpp
--- a/Engine/Source/Developer/TraceInsights/Private/Insights/MemoryProfiler/Widgets/SMemAllocTableTreeView.cpp
+++ b/Engine/Source/Developer/TraceInsights/Private/Insights/MemoryProfiler/Widgets/SMemAllocTableTreeView.cpp
@@ -94,9 +94,6 @@ void SMemAllocTableTreeView::RebuildTree(bool bResync)
 
 	const int32 PreviousNodeCount = TableTreeNodes.Num();
 
-	//check(Table->Is<Insights::FMemAllocTable>());
-	//TSharedPtr<Insights::FMemAllocTable> MemAllocTable = StaticCastSharedPtr<Insights::FMemAllocTable>(Table);
-
 	TSharedPtr<Insights::FMemAllocTable> MemAllocTable = GetMemAllocTable();
 
 	SyncStopwatch.Start();
@@ -107,6 +104,7 @@ void SMemAllocTableTreeView::RebuildTree(bool bResync)
 
 		if (QueryStatus == TraceServices::IAllocationsProvider::EQueryStatus::Done)
 		{
+			UE_LOG(MemoryProfiler, Log, TEXT("[MemAlloc] Rebuilding tree..."));
 			TArray<FMemoryAlloc>& Allocs = MemAllocTable->GetAllocs();
 
 			const int32 TotalAllocCount = Allocs.Num();
@@ -242,7 +240,6 @@ void SMemAllocTableTreeView::UpdateQuery(TraceServices::IAllocationsProvider::EQ
 		return;
 	}
 	const TraceServices::IAllocationsProvider& Provider = *AllocationsProvider;
-	TraceServices::IAllocationsProvider::FReadScopeLock _(Provider);
 
 	constexpr double MaxPollTime = 0.03; // Stop getting results after 30 ms so we don't tank the frame rate too much.
 	FStopwatch TotalStopwatch;
@@ -251,8 +248,8 @@ void SMemAllocTableTreeView::UpdateQuery(TraceServices::IAllocationsProvider::EQ
 	do
 	{
 		TraceServices::IAllocationsProvider::FQueryStatus Status = Provider.PollQuery(Query);
-
 		OutStatus = Status.Status;
+
 		if (Status.Status <= TraceServices::IAllocationsProvider::EQueryStatus::Done)
 		{
 			UE_LOG(MemoryProfiler, Log, TEXT("[MemAlloc] Query completed."));
@@ -261,27 +258,38 @@ void SMemAllocTableTreeView::UpdateQuery(TraceServices::IAllocationsProvider::EQ
 			return;
 		}
 
+		if (Status.Status == IAllocationsProvider::EQueryStatus::Working)
+		{
+			break;
+		}
+
+		check(Status.Status == IAllocationsProvider::EQueryStatus::Available);
+
 		TSharedPtr<Insights::FMemAllocTable> MemAllocTable = GetMemAllocTable();
 		if (MemAllocTable)
 		{
+			TraceServices::IAllocationsProvider::FReadScopeLock _(Provider);
+
 			TArray<FMemoryAlloc>& Allocs = MemAllocTable->GetAllocs();
 
 			FStopwatch ResultStopwatch;
 			FStopwatch PageStopwatch;
 			ResultStopwatch.Start();
 			uint32 PageCount = 0;
-			uint32 TotatAllocCount = 0;
+			uint32 TotalAllocCount = 0;
 
 			// Multiple 'pages' of results will be returned. No guarantees are made
 			// about the order of pages or the allocations they report.
 			TraceServices::IAllocationsProvider::FQueryResult Result = Status.NextResult();
 			while (Result.IsValid())
 			{
+				UE_LOG(MemoryProfiler, Log, TEXT("[MemAlloc] Page with %u allocs..."), Result->Num());
+
 				++PageCount;
 				PageStopwatch.Restart();
 
 				const uint32 AllocCount = Result->Num();
-				TotatAllocCount += AllocCount;
+				TotalAllocCount += AllocCount;
 #if 0
 				Allocs.Reserve(Allocs.Num() + AllocCount);
 
@@ -330,13 +338,14 @@ void SMemAllocTableTreeView::UpdateQuery(TraceServices::IAllocationsProvider::EQ
 			const double TotalTime = ResultStopwatch.GetAccumulatedTime();
 			if (TotalTime > 0.01)
 			{
-				const double Speed = (TotalTime * 1000000.0) / TotatAllocCount;
-				UE_LOG(MemoryProfiler, Log, TEXT("[MemAlloc] Query results (%u pages, %u allocs, slack=%u) retrieved in %.3fs (speed: %.3f seconds per 1M allocs)."), PageCount, TotatAllocCount, Allocs.GetSlack(), TotalTime, Speed);
+				const double Speed = (TotalTime * 1000000.0) / TotalAllocCount;
+				UE_LOG(MemoryProfiler, Log, TEXT("[MemAlloc] Query results (%u pages, %u allocs, slack=%u) retrieved in %.3fs (speed: %.3f seconds per 1M allocs)."), PageCount, TotalAllocCount, Allocs.GetSlack(), TotalTime, Speed);
 			}
 		}
 
 		TotalStopwatch.Update();
-	} while (OutStatus == TraceServices::IAllocationsProvider::EQueryStatus::Available && TotalStopwatch.GetAccumulatedTime() < MaxPollTime);
+	}
+	while (OutStatus == TraceServices::IAllocationsProvider::EQueryStatus::Available && TotalStopwatch.GetAccumulatedTime() < MaxPollTime);
 
 	TotalStopwatch.Stop();
 }
@@ -414,7 +423,7 @@ TSharedRef<ITableCellValueFormatter> SMemAllocTableTreeView::CreateCachedLlmTagV
 				{
 					return FText::FromString(FString::Printf(TEXT("%lli (%s)"), MemTagId, *TagIdToStatNameMap[MemTagId]));
 				}
-				
+
 				return FText::FromString(FString::Printf(TEXT("%lli ()"), MemTagId));
 			}
 			return FText::GetEmpty();
diff --git a/Engine/Source/Developer/TraceServices/Private/Analyzers/AllocationsAnalysis.cpp b/Engine/Source/Developer/TraceServices/Private/Analyzers/AllocationsAnalysis.cpp
--- a/Engine/Source/Developer/TraceServices/Private/Analyzers/AllocationsAnalysis.cpp
+++ b/Engine/Source/Developer/TraceServices/Private/Analyzers/AllocationsAnalysis.cpp
@@ -33,15 +33,25 @@ void FAllocationsAnalyzer::OnAnalysisBegin(const FOnAnalysisContext& Context)
 	Builder.RouteEvent(RouteId_Alloc,        "Memory", "Alloc");
 	Builder.RouteEvent(RouteId_Free,         "Memory", "Free");
 	Builder.RouteEvent(RouteId_ReallocAlloc, "Memory", "ReallocAlloc");
-	Builder.RouteEvent(RouteId_Marker,	     "Memory", "Marker");
+	Builder.RouteEvent(RouteId_Marker,       "Memory", "Marker");
+	Builder.RouteEvent(RouteId_TagSpec,      "Memory", "TagSpec");
+
+	Builder.RouteLoggerEvents(RouteId_MemScope, "Memory", true);
 }
 
 ////////////////////////////////////////////////////////////////////////////////////////////////////
 
 void FAllocationsAnalyzer::OnAnalysisEnd()
 {
+	double SessionTime;
+	{
+		FAnalysisSessionEditScope _(Session);
+		SessionTime = Session.GetDurationSeconds();
+	}
+	const double Time = FMath::Max(SessionTime, LastMarkerSeconds);
+
 	FAllocationsProvider::FEditScopeLock _(AllocationsProvider);
-	AllocationsProvider.EditOnAnalysisCompleted();
+	AllocationsProvider.EditOnAnalysisCompleted(Time);
 }
 
 ////////////////////////////////////////////////////////////////////////////////////////////////////
@@ -97,10 +107,12 @@ bool FAllocationsAnalyzer::OnEvent(uint16 RouteId, EStyle Style, const FOnEventC
 			const uint64 Address = EventData.GetValue<uint64>("Address");
 			const uint32 Size = EventData.GetValue<uint32>("Size");
 			const uint8 Alignment_SizeLower = EventData.GetValue<uint8>("Alignment_SizeLower");
-			const uint32 Tag = 0; //TODO
+
+			const uint32 ThreadId = Context.ThreadInfo.GetSystemId();
+			const uint8 Tracker = 0; // We only care about the default tracker for now.
 
 			FAllocationsProvider::FEditScopeLock _(AllocationsProvider);
-			AllocationsProvider.EditAlloc(Time, Owner, Address, Size, Alignment_SizeLower, Tag);
+			AllocationsProvider.EditAlloc(Time, Owner, Address, Size, Alignment_SizeLower, ThreadId, Tracker);
 			break;
 		}
 		case RouteId_Free:
@@ -119,6 +131,53 @@ bool FAllocationsAnalyzer::OnEvent(uint16 RouteId, EStyle Style, const FOnEventC
 			LastMarkerSeconds = Context.EventTime.AsSeconds(LastCycle);
 			break;
 		}
+		case RouteId_TagSpec:
+		{
+			const int32 Tag = Context.EventData.GetValue<int32>("Tag");
+			const int32 Parent = Context.EventData.GetValue<int32>("Parent");
+
+			FString Display;
+			Context.EventData.GetString("Display", Display);
+			const TCHAR* DisplayString = Session.StoreString(*Display);
+
+			FAllocationsProvider::FEditScopeLock _(AllocationsProvider);
+			AllocationsProvider.EditAddTagSpec(Tag, Parent, DisplayString);
+			break;
+		}
+		case RouteId_MemScope:
+		{
+			const uint32 ThreadId = Context.ThreadInfo.GetSystemId();
+			const uint8 Tracker = 0; // We only care about the default tracker for now.
+
+			if (Style == EStyle::EnterScope)
+			{
+				if (FCStringAnsi::Strlen(Context.EventData.GetTypeInfo().GetName()) == 11) // "MemoryScope"
+				{
+					const int32 Tag = Context.EventData.GetValue<int32>("Tag");
+					FAllocationsProvider::FEditScopeLock _(AllocationsProvider);
+					AllocationsProvider.EditPushTag(ThreadId, Tracker, Tag);
+				}
+				else // "MemoryScopeRealloc"
+				{
+					const uint64 Ptr = Context.EventData.GetValue<int32>("Ptr");
+					FAllocationsProvider::FEditScopeLock _(AllocationsProvider);
+					AllocationsProvider.EditPushRealloc(ThreadId, Tracker, Ptr);
+				}
+			}
+			else // EStyle::LeaveScope
+			{
+				FAllocationsProvider::FEditScopeLock _(AllocationsProvider);
+				if (AllocationsProvider.HasReallocScope(ThreadId, Tracker)) // Is realloc scope active?
+				{
+					AllocationsProvider.EditPopRealloc(ThreadId, Tracker);
+				}
+				else
+				{
+					AllocationsProvider.EditPopTag(ThreadId, Tracker);
+				}
+			}
+			break;
+		}
 	}
 
 	return true;
diff --git a/Engine/Source/Developer/TraceServices/Private/Analyzers/AllocationsAnalysis.h b/Engine/Source/Developer/TraceServices/Private/Analyzers/AllocationsAnalysis.h
--- a/Engine/Source/Developer/TraceServices/Private/Analyzers/AllocationsAnalysis.h
+++ b/Engine/Source/Developer/TraceServices/Private/Analyzers/AllocationsAnalysis.h
@@ -24,6 +24,8 @@ private:
 		RouteId_ReallocAlloc,
 		RouteId_ReallocFree,
 		RouteId_Marker,
+		RouteId_TagSpec,
+		RouteId_MemScope,
 	};
 
 public:
@@ -37,7 +39,6 @@ public:
 private:
 	IAnalysisSession& Session;
 	FAllocationsProvider& AllocationsProvider;
-	uint64 EventCount = 0; // debug
 	uint64 BaseCycle;
 	uint32 MarkerPeriod;
 	double LastMarkerSeconds;
diff --git a/Engine/Source/Developer/TraceServices/Private/Model/AllocationsProvider.cpp b/Engine/Source/Developer/TraceServices/Private/Model/AllocationsProvider.cpp
--- a/Engine/Source/Developer/TraceServices/Private/Model/AllocationsProvider.cpp
+++ b/Engine/Source/Developer/TraceServices/Private/Model/AllocationsProvider.cpp
@@ -74,6 +74,114 @@ void FAllocationsProviderLock::EndWrite()
 	}
 }
 
+////////////////////////////////////////////////////////////////////////////////////////////////////
+// FTagTracker
+////////////////////////////////////////////////////////////////////////////////////////////////////
+
+void FTagTracker::AddTagSpec(uint32 InTag, uint32 InParentTag, const TCHAR* InDisplay)
+{
+	if (ensure(!TagMap.Contains(InTag)))
+	{
+		TagMap.Emplace(InTag, TagEntry{ InDisplay, InParentTag });
+	}
+	else
+	{
+		++NumErrors;
+		//TODO: error message
+	}
+}
+
+////////////////////////////////////////////////////////////////////////////////////////////////////
+
+void FTagTracker::PushTag(uint32 InThreadId, uint8 InTracker, uint32 InTag)
+{
+	const uint32 TrackerThreadId = GetTrackerThreadId(InThreadId, InTracker);
+	ThreadState& State = TrackerThreadStates.FindOrAdd(TrackerThreadId);
+	State.TagStack.Push(InTag);
+}
+
+////////////////////////////////////////////////////////////////////////////////////////////////////
+
+void FTagTracker::PopTag(uint32 InThreadId, uint8 InTracker)
+{
+	const uint32 TrackerThreadId = GetTrackerThreadId(InThreadId, InTracker);
+	ThreadState* State = TrackerThreadStates.Find(TrackerThreadId);
+	if (ensure(State && !State->TagStack.IsEmpty()))
+	{
+		State->TagStack.Pop();
+	}
+	else
+	{
+		++NumErrors;
+		//TODO: error message
+	}
+}
+
+////////////////////////////////////////////////////////////////////////////////////////////////////
+
+uint32 FTagTracker::GetCurrentTag(uint32 InThreadId, uint8 InTracker) const
+{
+	const uint32 TrackerThreadId = GetTrackerThreadId(InThreadId, InTracker);
+	const ThreadState* State = TrackerThreadStates.Find(TrackerThreadId);
+	if (!State || State->TagStack.IsEmpty())
+	{
+		return 0; // Untagged
+	}
+	return State->TagStack.Top();
+}
+
+////////////////////////////////////////////////////////////////////////////////////////////////////
+
+const TCHAR* FTagTracker::GetTagString(uint32 InTag) const
+{
+	const TagEntry* Entry = TagMap.Find(InTag);
+	return Entry ? Entry->Display : nullptr;
+}
+
+////////////////////////////////////////////////////////////////////////////////////////////////////
+
+void FTagTracker::PushRealloc(uint32 InThreadId, uint8 InTracker, uint64 InPtr)
+{
+	const uint32 TrackerThreadId = GetTrackerThreadId(InThreadId, InTracker);
+	ThreadState& State = TrackerThreadStates.FindOrAdd(TrackerThreadId);
+	State.ReallocPtr.Push(InPtr);
+}
+
+////////////////////////////////////////////////////////////////////////////////////////////////////
+
+void FTagTracker::PopRealloc(uint32 InThreadId, uint8 InTracker)
+{
+	const uint32 TrackerThreadId = GetTrackerThreadId(InThreadId, InTracker);
+	ThreadState* State = TrackerThreadStates.Find(TrackerThreadId);
+	if (ensure(State && !State->ReallocPtr.IsEmpty()))
+	{
+		State->ReallocPtr.Pop();
+	}
+	else
+	{
+		++NumErrors;
+		//TODO: error message
+	}
+}
+
+////////////////////////////////////////////////////////////////////////////////////////////////////
+
+bool FTagTracker::HasReallocScope(uint32 InThreadId, uint8 InTracker) const
+{
+	const uint32 TrackerThreadId = GetTrackerThreadId(InThreadId, InTracker);
+	const ThreadState* State = TrackerThreadStates.Find(TrackerThreadId);
+	return State && !State->ReallocPtr.IsEmpty();
+}
+
+////////////////////////////////////////////////////////////////////////////////////////////////////
+
+bool FTagTracker::IsReallocScope(uint64 InSourcePtr, uint32 InThreadId, uint8 InTracker) const
+{
+	const uint32 TrackerThreadId = GetTrackerThreadId(InThreadId, InTracker);
+	const ThreadState* State = TrackerThreadStates.Find(TrackerThreadId);
+	return State && !State->ReallocPtr.IsEmpty() && State->ReallocPtr.Contains(InSourcePtr);
+}
+
 ////////////////////////////////////////////////////////////////////////////////////////////////////
 // IAllocationsProvider::FAllocation
 ////////////////////////////////////////////////////////////////////////////////////////////////////
@@ -268,7 +376,7 @@ void FAllocationsProvider::EditRemoveCore(double Time, uint64 Owner, uint64 Base
 
 ////////////////////////////////////////////////////////////////////////////////////////////////////
 
-void FAllocationsProvider::EditAlloc(double Time, uint64 Owner, uint64 Address, uint32 InSize, uint8 InAlignmentAndSizeLower, uint32 Tag)
+void FAllocationsProvider::EditAlloc(double Time, uint64 Owner, uint64 Address, uint32 InSize, uint8 InAlignmentAndSizeLower, uint32 ThreadId, uint8 Tracker)
 {
 	Lock.WriteAccessCheck();
 
@@ -298,6 +406,8 @@ void FAllocationsProvider::EditAlloc(double Time, uint64 Owner, uint64 Address,
 
 		const uint64 Size = (static_cast<uint64>(InSize) << SizeShift) | static_cast<uint64>(InAlignmentAndSizeLower & SizeLowerMask);
 
+		const uint32 Tag = TagTracker.GetCurrentTag(ThreadId, Tracker);
+
 		FAllocationItem Allocation =
 		{
 			EventIndex,
@@ -432,10 +542,8 @@ void FAllocationsProvider::UpdateHistogramByEventDistance(uint32 EventDistance)
 
 void FAllocationsProvider::AdvanceTimelines(double Time)
 {
-	constexpr double SampleGranularity = 0.001; // 1ms
-
 	// If enough time has passed (since the current sample is started)...
-	if (Time - SampleStartTimestamp > SampleGranularity)
+	if (Time - SampleStartTimestamp > DefaultTimelineSampleGranularity)
 	{
 		// Add the current sample to the timelines.
 		Timeline.EmplaceBack(SampleStartTimestamp);
@@ -457,7 +565,7 @@ void FAllocationsProvider::AdvanceTimelines(double Time)
 		SampleFreeEvents = 0;
 
 		// If the previous sample is well distanced in time...
-		if (Time - SampleEndTimestamp > SampleGranularity) 
+		if (Time - SampleEndTimestamp > DefaultTimelineSampleGranularity)
 		{
 			// Add an intermediate "flat region" sample.
 			Timeline.EmplaceBack(SampleEndTimestamp);
@@ -475,7 +583,7 @@ void FAllocationsProvider::AdvanceTimelines(double Time)
 
 ////////////////////////////////////////////////////////////////////////////////////////////////////
 
-void FAllocationsProvider::EditOnAnalysisCompleted()
+void FAllocationsProvider::EditOnAnalysisCompleted(double Time)
 {
 	Lock.WriteAccessCheck();
 
@@ -484,9 +592,50 @@ void FAllocationsProvider::EditOnAnalysisCompleted()
 		return;
 	}
 
+#if 0
+	const bool bResetTimelineAtEnd = false;
+
+	// Add all live allocs to SbTree (with infinite end time).
+	uint64 LiveAllocsTotalSize = 0;
+	for (TPair<uint64, FAllocationItem>& KV : LiveAllocs)
+	{
+		FAllocationItem* AllocationPtr = &KV.Value;
+
+		LiveAllocsTotalSize += AllocationPtr->GetSize();
+
+		// Assign same event index to all live allocs at the end of the session.
+		AllocationPtr->EndEventIndex = EventIndex;
+
+		SbTree->AddAlloc(AllocationPtr);
+
+		uint32 EventDistance = AllocationPtr->EndEventIndex - AllocationPtr->StartEventIndex;
+		UpdateHistogramByEventDistance(EventDistance);
+	}
+	check(TotalAllocatedMemory == LiveAllocsTotalSize);
+
+	if (bResetTimelineAtEnd)
+	{
+		AdvanceTimelines(Time + 10 * DefaultTimelineSampleGranularity);
+
+		const uint32 LiveAllocsTotalCount = (uint32)LiveAllocs.Num();
+		LiveAllocs.Empty();
+
+		// Update stats for the last timeline sample (reset to zero).
+		TotalAllocatedMemory = 0;
+		SampleMinTotalAllocatedMemory = 0;
+		SampleMinLiveAllocations = 0;
+		SampleFreeEvents += LiveAllocsTotalCount;
+	}
+#endif
+
+	// Flush the last cached timeline sample.
 	AdvanceTimelines(std::numeric_limits<double>::infinity());
 
-	//DebugPrint();
+#if 0
+	DebugPrint();
+#endif
+
+	SbTree->Validate();
 }
 
 ////////////////////////////////////////////////////////////////////////////////////////////////////
diff --git a/Engine/Source/Developer/TraceServices/Private/Model/AllocationsProvider.h b/Engine/Source/Developer/TraceServices/Private/Model/AllocationsProvider.h
--- a/Engine/Source/Developer/TraceServices/Private/Model/AllocationsProvider.h
+++ b/Engine/Source/Developer/TraceServices/Private/Model/AllocationsProvider.h
@@ -33,6 +33,49 @@ private:
 
 ////////////////////////////////////////////////////////////////////////////////////////////////////
 
+class FTagTracker
+{
+private:
+	static constexpr uint32 TrackerIdShift = 24;
+	static constexpr uint32 TrackerIdMask = 0xFF000000;
+
+	struct ThreadState
+	{
+		TArray<uint32> TagStack;
+		TArray<uint64> ReallocPtr;
+	};
+
+	struct TagEntry
+	{
+		const TCHAR* Display;
+		uint32 ParentTag;
+	};
+
+public:
+	void AddTagSpec(uint32 Tag, uint32 ParentTag, const TCHAR* Display);
+	void PushTag(uint32 ThreadId, uint8 Tracker, uint32 Tag);
+	void PopTag(uint32 ThreadId, uint8 Tracker);
+	uint32 GetCurrentTag(uint32 ThreadId, uint8 Tracker) const;
+	const TCHAR* GetTagString(uint32 Tag) const;
+
+	void PushRealloc(uint32 ThreadId, uint8 Tracker, uint64 Ptr);
+	void PopRealloc(uint32 ThreadId, uint8 Tracker);
+	bool HasReallocScope(uint32 ThreadId, uint8 Tracker) const;
+	bool IsReallocScope(uint64 SourcePtr, uint32 ThreadId, uint8 Tracker) const;
+
+private:
+	inline uint32 GetTrackerThreadId(uint32 ThreadId, uint8 Tracker) const
+	{
+		return (Tracker << TrackerIdShift) & (~TrackerIdMask & ThreadId);
+	}
+
+	TMap<uint32, ThreadState> TrackerThreadStates;
+	TMap<uint32, TagEntry> TagMap;
+	uint32 NumErrors = 0;
+};
+
+////////////////////////////////////////////////////////////////////////////////////////////////////
+
 struct FAllocationItem
 {
 	static constexpr uint32 AlignmentBits = 8;
@@ -77,6 +120,9 @@ struct FAllocationCoreItem
 
 class FAllocationsProvider : public IAllocationsProvider
 {
+private:
+	static constexpr double DefaultTimelineSampleGranularity = 0.0001; // 0.1ms
+
 public:
 	FAllocationsProvider(ILinearAllocator& InAllocator);
 	virtual ~FAllocationsProvider();
@@ -111,18 +157,32 @@ public:
 
 	const FSbTree* GetSbTree() const { ReadAccessCheck(); return SbTree; }
 
+	const TMap<uint64, FAllocationItem>& GetLiveAllocs()  const { ReadAccessCheck(); return LiveAllocs; }
+	uint32 GetNumLiveAllocs() const { ReadAccessCheck(); return LiveAllocs.Num(); }
+
+	bool HasReallocScope(uint32 ThreadId, uint8 Tracker) const { ReadAccessCheck(); return TagTracker.HasReallocScope(ThreadId, Tracker); }
+
 	void DebugPrint() const;
 
 	//////////////////////////////////////////////////
 	// Edit operations
 
 	void EditInit(double Time, uint8 MinAlignment, uint8 SizeShift, uint8 SummarySizeShift);
+
 	void EditAddCore(double Time, uint64 Owner, uint64 Base, uint32 Size);
 	void EditRemoveCore(double Time, uint64 Owner, uint64 Base, uint32 Size);
-	void EditAlloc(double Time, uint64 Owner, uint64 Address, uint32 Size, uint8 AlignmentAndSizeLower, uint32 Tag);
+
+	void EditAlloc(double Time, uint64 Owner, uint64 Address, uint32 Size, uint8 AlignmentAndSizeLower, uint32 ThreadId, uint8 Tracker);
 	void EditFree(double Time, uint64 Address);
 
-	void EditOnAnalysisCompleted();
+	void EditAddTagSpec(int32 Tag, uint32 ParentTag, const TCHAR* Display) { EditAccessCheck(); TagTracker.AddTagSpec(Tag, ParentTag, Display); }
+	void EditPushTag(uint32 ThreadId, uint8 Tracker, uint32 Tag)           { EditAccessCheck(); TagTracker.PushTag(ThreadId, Tracker, Tag); }
+	void EditPopTag(uint32 ThreadId, uint8 Tracker)                        { EditAccessCheck(); TagTracker.PopTag(ThreadId, Tracker); }
+
+	void EditPushRealloc(uint32 ThreadId, uint8 Tracker, uint64 Ptr)       { EditAccessCheck(); TagTracker.PushRealloc(ThreadId, Tracker, Ptr); }
+	void EditPopRealloc(uint32 ThreadId, uint8 Tracker)                    { EditAccessCheck(); TagTracker.PopRealloc(ThreadId, Tracker); }
+
+	void EditOnAnalysisCompleted(double Time);
 
 	//////////////////////////////////////////////////
 
@@ -142,6 +202,8 @@ private:
 	uint8 SummarySizeShift = 0;
 	bool bInitialized = false;
 
+	FTagTracker TagTracker;
+
 	uint32 EventIndex = 0;
 
 	uint64 AllocCount = 0;
diff --git a/Engine/Source/Developer/TraceServices/Private/Model/AllocationsQuery.cpp b/Engine/Source/Developer/TraceServices/Private/Model/AllocationsQuery.cpp
--- a/Engine/Source/Developer/TraceServices/Private/Model/AllocationsQuery.cpp
+++ b/Engine/Source/Developer/TraceServices/Private/Model/AllocationsQuery.cpp
@@ -11,21 +11,6 @@
 namespace TraceServices
 {
 
-////////////////////////////////////////////////////////////////////////////////////////////////////
-// FAllocationsImpl
-////////////////////////////////////////////////////////////////////////////////////////////////////
-
-FAllocationsImpl::FAllocationsImpl(uint32 NumItems)
-	: Next(nullptr)
-{
-}
-
-////////////////////////////////////////////////////////////////////////////////////////////////////
-
-FAllocationsImpl::~FAllocationsImpl()
-{
-}
-
 ////////////////////////////////////////////////////////////////////////////////////////////////////
 // FAllocationsQueryAsyncTask
 ////////////////////////////////////////////////////////////////////////////////////////////////////
@@ -86,7 +71,6 @@ void FAllocationsQuery::Cancel()
 IAllocationsProvider::FQueryStatus FAllocationsQuery::Poll()
 {
 	IAllocationsProvider::FQueryStatus Status = {};
-	Status.Handle = 0;
 
 	FAllocationsImpl* Allocations = nullptr;
 	if (Results.Dequeue(Allocations))
@@ -99,13 +83,13 @@ IAllocationsProvider::FQueryStatus FAllocationsQuery::Poll()
 	if (IsWorking)
 	{
 		Status.Status = IAllocationsProvider::EQueryStatus::Working;
-		return Status;
 	}
 	else
 	{
 		Status.Status = IAllocationsProvider::EQueryStatus::Done;
-		return Status;
 	}
+	Status.Handle = 0;
+	return Status;
 }
 
 ////////////////////////////////////////////////////////////////////////////////////////////////////
@@ -125,6 +109,25 @@ void FAllocationsQuery::Run()
 	const FSbTree* SbTree = Provider.GetSbTree();
 	if (SbTree)
 	{
+		UE_LOG(LogTraceServices, Log, TEXT("[MemAlloc] Processing %d live allocs..."), Provider.GetNumLiveAllocs());
+
+		FAllocationsImpl* LiveAllocsResult = new FAllocationsImpl();
+		QueryLiveAllocs(LiveAllocsResult->Items);
+
+		const uint32 NumLiveAllocs = static_cast<uint32>(LiveAllocsResult->Items.Num());
+		if (NumLiveAllocs != 0)
+		{
+			UE_LOG(LogTraceServices, Log, TEXT("[MemAlloc] Enqueue %u live allocs..."), NumLiveAllocs);
+			TotalAllocationCount += NumLiveAllocs;
+			Results.Enqueue(LiveAllocsResult);
+		}
+		else
+		{
+			delete LiveAllocsResult;
+		}
+
+		UE_LOG(LogTraceServices, Log, TEXT("[MemAlloc] Detecting cells..."));
+
 		TArray<const FSbTreeCell*> Cells;
 
 		if (!IsCanceling)
@@ -133,17 +136,36 @@ void FAllocationsQuery::Run()
 			CellCount += Cells.Num();
 		}
 
+		UE_LOG(LogTraceServices, Log, TEXT("[MemAlloc] %d cells to process"), Cells.Num());
+		uint32 CellIndex = 0;
+
 		for (const FSbTreeCell* Cell : Cells)
 		{
 			if (IsCanceling)
 			{
 				break;
 			}
-			FAllocationsImpl* Result = new FAllocationsImpl(0);
-			Cell->Query(Result->Items, Params);
-			TotalAllocationCount += Result->Items.Num();
-			Results.Enqueue(Result);
+
+			++CellIndex;
+			UE_LOG(LogTraceServices, Log, TEXT("[MemAlloc] Processing cell %u (%u allocs)..."), CellIndex, Cell->GetAllocCount());
+
+			FAllocationsImpl* CellResult = new FAllocationsImpl();
+			Cell->Query(CellResult->Items, Params);
+
+			const uint32 NumAllocs = static_cast<uint32>(CellResult->Items.Num());
+			if (NumAllocs != 0)
+			{
+				UE_LOG(LogTraceServices, Log, TEXT("[MemAlloc] Enqueue %u allocs..."), NumAllocs);
+				TotalAllocationCount += NumAllocs;
+				Results.Enqueue(CellResult);
+			}
+			else
+			{
+				delete CellResult;
+			}
 		}
+
+		UE_LOG(LogTraceServices, Log, TEXT("[MemAlloc] Done"));
 	}
 
 	Provider.EndRead();
@@ -161,4 +183,185 @@ void FAllocationsQuery::Run()
 
 ////////////////////////////////////////////////////////////////////////////////////////////////////
 
+void FAllocationsQuery::QueryLiveAllocs(TArray<const FAllocationItem*>& OutAllocs) const
+{
+	const TMap<uint64, FAllocationItem>& LiveAllocs = Provider.GetLiveAllocs();
+
+	// TODO: Live allocs have the end time set to std::numeric_limits<double>::infinity(),
+	//       so the conditions below could be simplified to remove the end time checks.
+	//       For now, the code remains "unoptimized" for debugging purposes.
+
+	switch (Params.Rule)
+	{
+	case IAllocationsProvider::EQueryRule::aAf: // active allocs at A
+	{
+		const double Time = Params.TimeA;
+		for (const auto& KV : LiveAllocs)
+		{
+			const FAllocationItem& Alloc = KV.Value;
+			if (Alloc.StartTime <= Time && Time <= Alloc.EndTime)
+			{
+				OutAllocs.Add(&Alloc);
+			}
+		}
+	}
+	break;
+
+	case IAllocationsProvider::EQueryRule::afA: // before
+	{
+		const double Time = Params.TimeA;
+		for (const auto& KV : LiveAllocs)
+		{
+			const FAllocationItem& Alloc = KV.Value;
+			if (Alloc.EndTime <= Time)
+			{
+				OutAllocs.Add(&Alloc);
+			}
+		}
+	}
+	break;
+
+	case IAllocationsProvider::EQueryRule::Aaf: // after
+	{
+		const double Time = Params.TimeA;
+		for (const auto& KV : LiveAllocs)
+		{
+			const FAllocationItem& Alloc = KV.Value;
+			if (Alloc.StartTime >= Time)
+			{
+				OutAllocs.Add(&Alloc);
+			}
+		}
+	}
+	break;
+
+	case IAllocationsProvider::EQueryRule::aAfB: // decline
+	{
+		const double TimeA = Params.TimeA;
+		const double TimeB = Params.TimeB;
+		for (const auto& KV : LiveAllocs)
+		{
+			const FAllocationItem& Alloc = KV.Value;
+			if (Alloc.StartTime <= TimeA && Alloc.EndTime >= TimeA && Alloc.EndTime >= TimeB)
+			{
+				OutAllocs.Add(&Alloc);
+			}
+		}
+	}
+	break;
+
+	case IAllocationsProvider::EQueryRule::AaBf: // growth
+	{
+		const double TimeA = Params.TimeA;
+		const double TimeB = Params.TimeB;
+		for (const auto& KV : LiveAllocs)
+		{
+			const FAllocationItem& Alloc = KV.Value;
+			if (Alloc.StartTime >= TimeA && Alloc.StartTime <= TimeB && Alloc.EndTime >= TimeB)
+			{
+				OutAllocs.Add(&Alloc);
+			}
+		}
+	}
+	break;
+
+	case IAllocationsProvider::EQueryRule::AafB: // short living allocs
+	{
+		const double TimeA = Params.TimeA;
+		const double TimeB = Params.TimeB;
+		for (const auto& KV : LiveAllocs)
+		{
+			const FAllocationItem& Alloc = KV.Value;
+			if (Alloc.StartTime >= TimeA && Alloc.EndTime <= TimeB)
+			{
+				OutAllocs.Add(&Alloc);
+			}
+		}
+	}
+	break;
+
+	case IAllocationsProvider::EQueryRule::aABf: // long living allocs
+	{
+		const double TimeA = Params.TimeA;
+		const double TimeB = Params.TimeB;
+		for (const auto& KV : LiveAllocs)
+		{
+			const FAllocationItem& Alloc = KV.Value;
+			if (Alloc.StartTime <= TimeA && Alloc.EndTime >= TimeB)
+			{
+				OutAllocs.Add(&Alloc);
+			}
+		}
+	}
+	break;
+
+	case IAllocationsProvider::EQueryRule::AaBCf: // memory leaks
+	{
+		const double TimeA = Params.TimeA;
+		const double TimeB = Params.TimeB;
+		const double TimeC = Params.TimeC;
+		for (const auto& KV : LiveAllocs)
+		{
+			const FAllocationItem& Alloc = KV.Value;
+			if (Alloc.StartTime >= TimeA && Alloc.StartTime <= TimeB && Alloc.EndTime >= TimeC)
+			{
+				OutAllocs.Add(&Alloc);
+			}
+		}
+	}
+	break;
+
+	case IAllocationsProvider::EQueryRule::AaBfC: // limited lifetime
+	{
+		const double TimeA = Params.TimeA;
+		const double TimeB = Params.TimeB;
+		const double TimeC = Params.TimeC;
+		for (const auto& KV : LiveAllocs)
+		{
+			const FAllocationItem& Alloc = KV.Value;
+			if (Alloc.StartTime >= TimeA && Alloc.StartTime <= TimeB && Alloc.EndTime >= TimeB && Alloc.EndTime <= TimeC)
+			{
+				OutAllocs.Add(&Alloc);
+			}
+		}
+	}
+	break;
+
+	case IAllocationsProvider::EQueryRule::aABfC: // decline of long living allocs
+	{
+		const double TimeA = Params.TimeA;
+		const double TimeB = Params.TimeB;
+		const double TimeC = Params.TimeC;
+		for (const auto& KV : LiveAllocs)
+		{
+			const FAllocationItem& Alloc = KV.Value;
+			if (Alloc.StartTime <= TimeA && Alloc.EndTime >= TimeB && Alloc.EndTime <= TimeC)
+			{
+				OutAllocs.Add(&Alloc);
+			}
+		}
+	}
+	break;
+
+	case IAllocationsProvider::EQueryRule::AaBCfD: // specific lifetime
+	{
+		const double TimeA = Params.TimeA;
+		const double TimeB = Params.TimeB;
+		const double TimeC = Params.TimeC;
+		const double TimeD = Params.TimeD;
+		for (const auto& KV : LiveAllocs)
+		{
+			const FAllocationItem& Alloc = KV.Value;
+			if (Alloc.StartTime >= TimeA && Alloc.StartTime <= TimeB && Alloc.EndTime >= TimeC && Alloc.EndTime <= TimeD)
+			{
+				OutAllocs.Add(&Alloc);
+			}
+		}
+	}
+	break;
+	}
+}
+
+////////////////////////////////////////////////////////////////////////////////////////////////////
+
 } // namespace TraceServices
diff --git a/Engine/Source/Developer/TraceServices/Private/Model/AllocationsQuery.h b/Engine/Source/Developer/TraceServices/Private/Model/AllocationsQuery.h
--- a/Engine/Source/Developer/TraceServices/Private/Model/AllocationsQuery.h
+++ b/Engine/Source/Developer/TraceServices/Private/Model/AllocationsQuery.h
@@ -16,10 +16,7 @@ class ILinearAllocator;
 
 struct FAllocationsImpl
 {
-	FAllocationsImpl(uint32 NumItems);
-	~FAllocationsImpl();
-
-	FAllocationsImpl* Next;
+	FAllocationsImpl* Next = nullptr;
 	TArray<const FAllocationItem*> Items;
 };
 
@@ -35,6 +32,7 @@ public:
 
 private:
 	void Run();
+	void QueryLiveAllocs(TArray<const FAllocationItem*>& OutAllocs) const;
 
 private:
 	const FAllocationsProvider& Provider;
diff --git a/Engine/Source/Developer/TraceServices/Private/Model/SbTree.cpp b/Engine/Source/Developer/TraceServices/Private/Model/SbTree.cpp
--- a/Engine/Source/Developer/TraceServices/Private/Model/SbTree.cpp
+++ b/Engine/Source/Developer/TraceServices/Private/Model/SbTree.cpp
@@ -299,7 +299,7 @@ void FSbTree::SetTimeForEvent(uint32 EventIndex, double Time)
 	// Detect the first event of each new column.
 	if ((EventIndex & ((1 << ColumnShift) - 1)) == 0)
 	{
-		check((uint32)ColumnStartTimes.Num() == (EventIndex >> ColumnShift));
+		check(static_cast<uint32>(ColumnStartTimes.Num()) == (EventIndex >> ColumnShift));
 		ColumnStartTimes.Add(Time);
 	}
 }
@@ -310,13 +310,16 @@ void FSbTree::AddAlloc(const FAllocationItem* Alloc)
 {
 	check(Alloc != nullptr);
 
+	check(Alloc->EndTime >= LastAllocEndTime);
+	LastAllocEndTime = Alloc->EndTime;
+
 	uint32 StartColumn = Alloc->StartEventIndex >> ColumnShift;
 	uint32 EndColumn = Alloc->EndEventIndex >> ColumnShift;
 
 	if (EndColumn > CurrentColumn)
 	{
-		check((uint32)Cells.Num() > (CurrentColumn << 1));
-		check((uint32)ColumnStartTimes.Num() == EndColumn + 1);
+		check(static_cast<uint32>(Cells.Num()) > (CurrentColumn << 1));
+		check(static_cast<uint32>(ColumnStartTimes.Num()) == EndColumn + 1);
 
 		// Adds 2 cells for each new column.
 		uint32 CellsToAdd = (EndColumn - CurrentColumn) << 1;
@@ -333,11 +336,39 @@ void FSbTree::AddAlloc(const FAllocationItem* Alloc)
 	}
 
 #if USE_OFFSETTED_CELLS
+	uint32 Depth;
+	uint32 CellIndex;
+	bool bUseOffsettedCells;
+
 	const uint32 DeltaColumns = EndColumn - StartColumn;
-	const uint32 Depth = 32 - FMath::CountLeadingZeros(DeltaColumns);
-	const uint32 HalfCellWidth = (1 << Depth) >> 1;
-	const uint32 CellIndex = FSbTreeUtils::GetCellAtDepth(StartColumn & ~HalfCellWidth, Depth);
-	const bool bUseOffsettedCells = ((StartColumn & HalfCellWidth) != 0);
+	if (DeltaColumns == 0)
+	{
+		Depth = 0;
+		CellIndex = StartColumn << 1;
+		bUseOffsettedCells = false;
+	}
+	else
+	{
+		Depth = 32 - FMath::CountLeadingZeros(DeltaColumns);
+		uint32 HalfCellWidth = (1 << Depth) >> 1;
+
+		// In a cell, each alloc has start event column only in the first half of the cell and the end event column in the second half of the cell.
+		// For a column, the bit indicating the "type of the half cell" at a certain depth ("start columns half cell" or "end columns half cell")
+		// is Column & HalfCellWidth. This is also true for offsetted cells (just that for normal cells 0 means "start columns half cell"
+		// while in offsetted cells 1 means "start columns half cell").
+
+		// If the "type of the half cell" bit is the same for StartColumn and for EndColumn on the min depth (computed based on the column width for this alloc)
+		// it means that the alloc doesn't fit in the (alligned) cells on this depth (neither in offsetted cells on this depth).
+		// Those allocs are pushed to the next depth.
+		if ((StartColumn & HalfCellWidth) == (EndColumn & HalfCellWidth))
+		{
+			++Depth;
+			HalfCellWidth <<= 1;
+		}
+
+		CellIndex = FSbTreeUtils::GetCellAtDepth(StartColumn & ~HalfCellWidth, Depth);
+		bUseOffsettedCells = ((StartColumn & HalfCellWidth) != 0);
+	}
 #else // USE_OFFSETTED_CELLS
 	const uint32 Depth = FSbTreeUtils::GetCommonDepth(StartColumn, EndColumn);
 	const uint32 CellIndex = FSbTreeUtils::GetCellAtDepth(StartColumn, Depth);
@@ -348,7 +379,7 @@ void FSbTree::AddAlloc(const FAllocationItem* Alloc)
 
 	if (!bUseOffsettedCells)
 	{
-		check(CellIndex < (uint32)Cells.Num());
+		check(CellIndex < static_cast<uint32>(Cells.Num()));
 		CellPtr = Cells[CellIndex];
 		if (CellPtr == nullptr)
 		{
@@ -362,7 +393,7 @@ void FSbTree::AddAlloc(const FAllocationItem* Alloc)
 	{
 		check((CellIndex & 1) != 0); // depth 0 doesn't use offsetted cells
 
-		check(CellIndex < (uint32)OffsettedCells.Num());
+		check(CellIndex < static_cast<uint32>(OffsettedCells.Num()));
 		CellPtr = OffsettedCells[CellIndex];
 		if (CellPtr == nullptr)
 		{
@@ -378,10 +409,29 @@ void FSbTree::AddAlloc(const FAllocationItem* Alloc)
 
 ////////////////////////////////////////////////////////////////////////////////////////////////////
 
-int32 FSbTree::GetColumnAtTime(double Time) const
+int32 FSbTree::GetColumnsAtTime(double Time, int32* StartColumnPtr, int32* EndColumnPtr) const
 {
-	const int32 Index = Algo::UpperBound(ColumnStartTimes, Time);
-	return Index - 1;
+	// Returns the first column with T >= Time.
+	// [0 .. N]
+	const int32 LowerBoundColumn = Algo::LowerBound(ColumnStartTimes, Time);
+
+	if (StartColumnPtr)
+	{
+		*StartColumnPtr = LowerBoundColumn - 1; // [-1 .. N-1]
+	}
+
+	if (EndColumnPtr)
+	{
+		int32 Column = LowerBoundColumn - 1; // [-1 .. N-1]
+		const int32 Last = ColumnStartTimes.Num() - 1; // N-1
+		while (Column < Last && ColumnStartTimes[Column + 1] == Time)
+		{
+			++Column;
+		}
+		*EndColumnPtr = Column; // [-1 .. N-1]
+	}
+
+	return LowerBoundColumn;
 }
 
 ////////////////////////////////////////////////////////////////////////////////////////////////////
@@ -392,83 +442,99 @@ void FSbTree::Query(TArray<const FSbTreeCell*>& OutCells, const IAllocationsProv
 	{
 		case IAllocationsProvider::EQueryRule::aAf: // active allocs at A
 		{
-			const int32 ColumnA = GetColumnAtTime(Params.TimeA);
-			IterateCells(OutCells, ColumnA);
+			int32 Column1, Column2;
+			GetColumnsAtTime(Params.TimeA, &Column1, &Column2);
+			IterateCells(OutCells, Column1, Column2);
 		}
 		break;
 
 		case IAllocationsProvider::EQueryRule::afA: // before
 		{
-			const int32 ColumnA = GetColumnAtTime(Params.TimeA);
-			IterateCells(OutCells, 0, ColumnA);
+			int32 Column2;
+			GetColumnsAtTime(Params.TimeA, nullptr, &Column2);
+			IterateCells(OutCells, 0, Column2);
 		}
 		break;
 
 		case IAllocationsProvider::EQueryRule::Aaf: // after
 		{
-			const int32 ColumnA = GetColumnAtTime(Params.TimeA);
-			IterateCells(OutCells, ColumnA, CurrentColumn);
+			int32 Column1;
+			GetColumnsAtTime(Params.TimeA, &Column1, nullptr);
+			IterateCells(OutCells, Column1, CurrentColumn);
 		}
 		break;
-	
+
 		case IAllocationsProvider::EQueryRule::aAfB: // decline
 		{
-			const int32 ColumnA = GetColumnAtTime(Params.TimeA);
-			IterateCells(OutCells, ColumnA);
+			int32 Column1;
+			GetColumnsAtTime(Params.TimeA, &Column1, nullptr);
+			int32 Column2;
+			GetColumnsAtTime(Params.TimeB, nullptr, &Column2);
+			IterateCells(OutCells, Column1, Column2);
 		}
 		break;
-	
+
 		case IAllocationsProvider::EQueryRule::AaBf: // growth
 		{
-			const int32 ColumnB = GetColumnAtTime(Params.TimeB);
-			IterateCells(OutCells, ColumnB);
+			int32 Column1;
+			GetColumnsAtTime(Params.TimeB, &Column1, nullptr);
+			IterateCells(OutCells, Column1, CurrentColumn);
 		}
 		break;
 
 		case IAllocationsProvider::EQueryRule::AafB: // short living allocs
 		{
-			const int32 ColumnA = GetColumnAtTime(Params.TimeA);
-			const int32 ColumnB = GetColumnAtTime(Params.TimeB);
-			IterateCells(OutCells, ColumnA, ColumnB);
+			int32 Column1;
+			GetColumnsAtTime(Params.TimeA, &Column1, nullptr);
+			int32 Column2;
+			GetColumnsAtTime(Params.TimeB, nullptr, &Column2);
+			IterateCells(OutCells, Column1, Column2);
 		}
 		break;
 
 		case IAllocationsProvider::EQueryRule::aABf: // long living allocs
 		{
-			const int32 ColumnA = GetColumnAtTime(Params.TimeA);
-			const int32 ColumnB = GetColumnAtTime(Params.TimeB);
-			IterateCells(OutCells, ColumnA, ColumnB);
+			int32 Column1;
+			GetColumnsAtTime(Params.TimeB, &Column1, nullptr);
+			IterateCells(OutCells, Column1, CurrentColumn);
 		}
 		break;
-	
+
 		case IAllocationsProvider::EQueryRule::AaBCf: // memory leaks
 		{
-			const int32 ColumnB = GetColumnAtTime(Params.TimeB);
-			const int32 ColumnC = GetColumnAtTime(Params.TimeC);
-			IterateCells(OutCells, ColumnB, ColumnC);
+			int32 Column1;
+			GetColumnsAtTime(Params.TimeC, &Column1, nullptr);
+			IterateCells(OutCells, Column1, CurrentColumn);
 		}
 		break;
 
 		case IAllocationsProvider::EQueryRule::AaBfC: // limited lifetime
 		{
-			const int32 ColumnB = GetColumnAtTime(Params.TimeB);
-			IterateCells(OutCells, ColumnB);
+			int32 Column1;
+			GetColumnsAtTime(Params.TimeB, &Column1, nullptr);
+			int32 Column2;
+			GetColumnsAtTime(Params.TimeC, nullptr, &Column2);
+			IterateCells(OutCells, Column1, Column2);
 		}
 		break;
-	
+
 		case IAllocationsProvider::EQueryRule::aABfC: // decline of long living allocs
 		{
-			const int32 ColumnA = GetColumnAtTime(Params.TimeA);
-			const int32 ColumnB = GetColumnAtTime(Params.TimeB);
-			IterateCells(OutCells, ColumnA, ColumnB);
+			int32 Column1;
+			GetColumnsAtTime(Params.TimeB, &Column1, nullptr);
+			int32 Column2;
+			GetColumnsAtTime(Params.TimeC, nullptr, &Column2);
+			IterateCells(OutCells, Column1, Column2);
 		}
 		break;
-	
+
 		case IAllocationsProvider::EQueryRule::AaBCfD: // specific lifetime
 		{
-			const int32 ColumnB = GetColumnAtTime(Params.TimeB);
-			const int32 ColumnC = GetColumnAtTime(Params.TimeC);
-			IterateCells(OutCells, ColumnB, ColumnC);
+			int32 Column1;
+			GetColumnsAtTime(Params.TimeC, &Column1, nullptr);
+			int32 Column2;
+			GetColumnsAtTime(Params.TimeD, nullptr, &Column2);
+			IterateCells(OutCells, Column1, Column2);
 		}
 		break;
 	}
@@ -484,12 +550,15 @@ void FSbTree::IterateCells(TArray<const FSbTreeCell*>& OutCells, int32 Column) c
 		return;
 	}
 
+	const uint32 LocalColumn = static_cast<uint32>(Column);
+
 	const uint32 MaxDepth = FSbTreeUtils::GetMaxDepth(CurrentColumn);
 
+	const uint32 NumCells = static_cast<uint32>(Cells.Num());
 	for (uint32 Depth = 0; Depth <= MaxDepth; ++Depth)
 	{
-		const uint32 CellIndex = FSbTreeUtils::GetCellAtDepth(Column, Depth);
-		if (CellIndex < (uint32)Cells.Num())
+		const uint32 CellIndex = FSbTreeUtils::GetCellAtDepth(LocalColumn, Depth);
+		if (CellIndex < NumCells)
 		{
 			const FSbTreeCell* CellPtr = Cells[CellIndex];
 			if (CellPtr)
@@ -497,70 +566,90 @@ void FSbTree::IterateCells(TArray<const FSbTreeCell*>& OutCells, int32 Column) c
 				OutCells.Add(CellPtr);
 			}
 		}
+	}
 
 #if USE_OFFSETTED_CELLS
+	const uint32 NumOffsettedCells = static_cast<uint32>(OffsettedCells.Num());
+	for (uint32 Depth = 1; Depth <= MaxDepth; ++Depth) // offsetted cells doesn't exist on Depth 0
+	{
 		const uint32 HalfCellWidth = (1 << Depth) >> 1;
-		const uint32 OffsettedCellIndex = FSbTreeUtils::GetCellAtDepth(Column & ~HalfCellWidth, Depth);
-		if (OffsettedCellIndex < (uint32)OffsettedCells.Num())
+		if (HalfCellWidth > LocalColumn)
 		{
-			const FSbTreeCell* CellPtr = OffsettedCells[OffsettedCellIndex];
+			break;
+		}
+		const uint32 CellIndex = FSbTreeUtils::GetCellAtDepth(LocalColumn - HalfCellWidth, Depth);
+		if (CellIndex < NumOffsettedCells)
+		{
+			const FSbTreeCell* CellPtr = OffsettedCells[CellIndex];
 			if (CellPtr)
 			{
 				OutCells.Add(CellPtr);
 			}
 		}
-#endif // USE_OFFSETTED_CELLS
 	}
+#endif // USE_OFFSETTED_CELLS
 }
 
 ////////////////////////////////////////////////////////////////////////////////////////////////////
 
 void FSbTree::IterateCells(TArray<const FSbTreeCell*> & OutCells, int32 StartColumn, int32 EndColumn) const
 {
+	if (StartColumn == EndColumn)
+	{
+		IterateCells(OutCells, StartColumn);
+		return;
+	}
+
 	if (EndColumn < 0 || StartColumn > (int32)CurrentColumn || StartColumn > EndColumn)
 	{
 		// not an error; just early out
 		return;
 	}
 
-	const uint32 LocalStartColumn = (StartColumn < 0) ? 0 : (uint32)StartColumn;
-	const uint32 LocalEndColumn = ((uint32)EndColumn > CurrentColumn) ? CurrentColumn : (uint32)EndColumn;
+	const uint32 LocalStartColumn = (StartColumn < 0) ? 0 : static_cast<uint32>(StartColumn);
+	const uint32 LocalEndColumn = (static_cast<uint32>(EndColumn) > CurrentColumn) ? CurrentColumn : static_cast<uint32>(EndColumn);
 	check(LocalStartColumn <= LocalEndColumn);
 
 	const uint32 MaxDepth = FSbTreeUtils::GetMaxDepth(CurrentColumn);
 
+	const uint32 NumCells = static_cast<uint32>(Cells.Num());
 	for (uint32 Depth = 0; Depth <= MaxDepth; ++Depth)
 	{
-		const uint32 CellOffset = 1 << (Depth + 1);
-
 		const uint32 FirstCellIndex = FSbTreeUtils::GetCellAtDepth(LocalStartColumn, Depth);
-		const uint32 LastCellIndex = FMath::Min(FSbTreeUtils::GetCellAtDepth(LocalEndColumn, Depth) + 1, (uint32)Cells.Num());
-		for (uint32 CellIndex = FirstCellIndex; CellIndex < LastCellIndex; CellIndex += CellOffset)
+		const uint32 LastCellIndex = FMath::Min(FSbTreeUtils::GetCellAtDepth(LocalEndColumn, Depth) + 1, NumCells);
+		const uint32 CellIncrement = 1 << (Depth + 1);
+		for (uint32 CellIndex = FirstCellIndex; CellIndex < LastCellIndex; CellIndex += CellIncrement)
 		{
-			FSbTreeCell* CellPtr = Cells[CellIndex];
+			const FSbTreeCell* CellPtr = Cells[CellIndex];
 			if (CellPtr)
 			{
 				OutCells.Add(CellPtr);
 			}
 		}
+	}
 
 #if USE_OFFSETTED_CELLS
-		if (Depth > 0) // offsetted cells doesn't exist on Depth 0
+	const uint32 NumOffsettedCells = static_cast<uint32>(OffsettedCells.Num());
+	for (uint32 Depth = 1; Depth <= MaxDepth; ++Depth) // offsetted cells doesn't exist on Depth 0
+	{
+		uint32 HalfCellWidth = (1 << Depth) >> 1;
+		if (HalfCellWidth > LocalStartColumn)
+		{
+			break;
+		}
+		uint32 FirstCellIndex = FSbTreeUtils::GetCellAtDepth(LocalStartColumn - HalfCellWidth, Depth);
+		uint32 LastCellIndex = FMath::Min(FSbTreeUtils::GetCellAtDepth(LocalEndColumn - HalfCellWidth, Depth) + 1, NumOffsettedCells);
+		const uint32 CellIncrement = 1 << (Depth + 1);
+		for (uint32 CellIndex = FirstCellIndex; CellIndex < LastCellIndex; CellIndex += CellIncrement)
 		{
-			uint32 HalfCellWidth = (1 << Depth) >> 1;
-			uint32 FirstOffsettedCellIndex = FSbTreeUtils::GetCellAtDepth(LocalStartColumn & ~HalfCellWidth, Depth);
-			uint32 LastOffsettedCellIndex = FMath::Min(FSbTreeUtils::GetCellAtDepth(LocalEndColumn & ~HalfCellWidth, Depth) + 1, (uint32)OffsettedCells.Num());
-			for (uint32 CellIndex = FirstOffsettedCellIndex; CellIndex < LastOffsettedCellIndex; CellIndex += CellOffset)
+			const FSbTreeCell* CellPtr = OffsettedCells[CellIndex];
+			if (CellPtr)
 			{
-				FSbTreeCell* CellPtr = OffsettedCells[CellIndex];
-				if (CellPtr)
-				{
-					OutCells.Add(CellPtr);
-				}
+				OutCells.Add(CellPtr);
 			}
 		}
-#endif // USE_OFFSETTED_CELLS
 	}
+#endif // USE_OFFSETTED_CELLS
 }
 
 ////////////////////////////////////////////////////////////////////////////////////////////////////
@@ -574,14 +663,14 @@ void FSbTree::DebugPrint() const
 	uint32 LocalMaxAllocCountPerCell = 0;
 
 	uint32 NotEmptyCellCount = 0;
-	const uint32 CellCount = (uint32)Cells.Num();
+	const uint32 CellCount = static_cast<uint32>(Cells.Num());
 	for (uint32 CellIndex = 0; CellIndex < CellCount; ++CellIndex)
 	{
 		const FSbTreeCell* CellPtr = Cells[CellIndex];
 		if (CellPtr != nullptr)
 		{
 			++NotEmptyCellCount;
-			uint32 CellAllocCount = CellPtr->GetAllocCount();
+			const uint32 CellAllocCount = CellPtr->GetAllocCount();
 			TotalAllocs += CellAllocCount;
 			if (CellAllocCount > LocalMaxAllocCountPerCell)
 			{
@@ -594,14 +683,14 @@ void FSbTree::DebugPrint() const
 
 #if USE_OFFSETTED_CELLS
 	uint32 NotEmptyOffsettedCellCount = 0;
-	const uint32 OffsettedCellCount = (uint32)OffsettedCells.Num();
+	const uint32 OffsettedCellCount = static_cast<uint32>(OffsettedCells.Num());
 	for (uint32 CellIndex = 0; CellIndex < OffsettedCellCount; ++CellIndex)
 	{
 		const FSbTreeCell* CellPtr = OffsettedCells[CellIndex];
 		if (CellPtr != nullptr)
 		{
 			++NotEmptyOffsettedCellCount;
-			uint32 CellAllocCount = CellPtr->GetAllocCount();
+			const uint32 CellAllocCount = CellPtr->GetAllocCount();
 			TotalAllocs += CellAllocCount;
 			if (CellAllocCount > LocalMaxAllocCountPerCell)
 			{
@@ -689,6 +778,64 @@ void FSbTree::DebugPrint() const
 
 ////////////////////////////////////////////////////////////////////////////////////////////////////
 
+void FSbTree::Validate() const
+{
+	const uint32 MaxDepth = FSbTreeUtils::GetMaxDepth(CurrentColumn);
+
+	const uint32 NumCells = static_cast<uint32>(Cells.Num());
+	for (uint32 Depth = 0; Depth <= MaxDepth; ++Depth)
+	{
+		uint32 EventIndex = 0;
+		double Time = 0.0;
+		const uint32 CellIncrement = 1 << (Depth + 1);
+		uint32 CellIndex = FSbTreeUtils::GetCellAtDepth(0, Depth);
+		while (CellIndex < NumCells)
+		{
+			const FSbTreeCell* CellPtr = Cells[CellIndex];
+			if (CellPtr)
+			{
+				check(CellPtr->GetMaxEndEventIndex() > CellPtr->GetMinStartEventIndex());
+				check(CellPtr->GetMinStartEventIndex() >= EventIndex);
+				EventIndex = CellPtr->GetMaxEndEventIndex() + 1;
+
+				check(CellPtr->GetMaxEndTime() >= CellPtr->GetMinStartTime());
+				check(CellPtr->GetMinStartTime() >= Time);
+				Time = CellPtr->GetMinStartTime();
+			}
+			CellIndex += CellIncrement;
+		}
+	}
+
+#if USE_OFFSETTED_CELLS
+	const uint32 NumOffsettedCells = static_cast<uint32>(OffsettedCells.Num());
+	for (uint32 Depth = 0; Depth <= MaxDepth; ++Depth)
+	{
+		uint32 EventIndex = 0;
+		double Time = 0.0;
+		const uint32 CellIncrement = 1 << (Depth + 1);
+		uint32 CellIndex = FSbTreeUtils::GetCellAtDepth(0, Depth);
+		while (CellIndex < NumOffsettedCells)
+		{
+			const FSbTreeCell* CellPtr = OffsettedCells[CellIndex];
+			if (CellPtr)
+			{
+				check(CellPtr->GetMaxEndEventIndex() > CellPtr->GetMinStartEventIndex());
+				check(CellPtr->GetMinStartEventIndex() >= EventIndex);
+				EventIndex = CellPtr->GetMaxEndEventIndex() + 1;
+
+				check(CellPtr->GetMaxEndTime() >= CellPtr->GetMinStartTime());
+				check(CellPtr->GetMinStartTime() >= Time);
+				Time = CellPtr->GetMinStartTime();
+			}
+			CellIndex += CellIncrement;
+		}
+	}
+	check(NumCells == NumOffsettedCells);
+#endif // USE_OFFSETTED_CELLS
+}
+
+////////////////////////////////////////////////////////////////////////////////////////////////////
+
 #undef USE_OFFSETTED_CELLS
 
 } // namespace TraceServices
diff --git a/Engine/Source/Developer/TraceServices/Private/Model/SbTree.h b/Engine/Source/Developer/TraceServices/Private/Model/SbTree.h
--- a/Engine/Source/Developer/TraceServices/Private/Model/SbTree.h
+++ b/Engine/Source/Developer/TraceServices/Private/Model/SbTree.h
@@ -31,6 +31,9 @@ public:
 	uint32 GetAllocCount() const { return Allocs.Num(); }
 	void AddAlloc(const FAllocationItem* Alloc);
 
+	uint32 GetMinStartEventIndex() const { return MinStartEventIndex; }
+	uint32 GetMaxEndEventIndex() const { return MaxEndEventIndex; }
+
 	double GetMinStartTime() const { return MinStartTime; }
 	double GetMaxEndTime() const { return MaxEndTime; }
 
@@ -63,22 +66,39 @@ public:
 	uint32 GetColumnWidth() const { return 1 << ColumnShift; }
 	uint32 GetCurrentColumn() const { return CurrentColumn; }
 
-	int32 GetColumnAtTime(double Time) const;
-
 	void Query(TArray<const FSbTreeCell*>& OutCells, const IAllocationsProvider::FQueryParams& Params) const;
 
 	void IterateCells(TArray<const FSbTreeCell*>& OutCells, int32 Column) const;
 	void IterateCells(TArray<const FSbTreeCell*>& OutCells, int32 StartColumn, int32 EndColumn) const;
 
 	void DebugPrint() const;
+	void Validate() const;
+
+private:
+	int32 GetColumnsAtTime(double Time, int32* StartColumnPtr, int32* EndColumnPtr) const;
 
 private:
 	ILinearAllocator& Allocator;
+
+	// Normal cells.
 	TArray<FSbTreeCell*> Cells;
+
+	// Offsetted cells. Note: Depth 0 doesn't use offseted cells. Those will always be empty.
 	TArray<FSbTreeCell*> OffsettedCells;
+
+	// Array with time values of the first event in each column.
+	// Because multiple columns can have same time, CST[i] can also be assumed the time of the last event in previous column (i-1).
+	// Time range for events in column i is inclusive interval [CST[i], CST[i+1]], with CST[N] = infinity.
 	TArray<double> ColumnStartTimes;
-	uint32 ColumnShift; // ColumnWidth = 1 << ColumnShift
+
+	// It defines the width of a column (as number of events).
+	// ColumnWidth = 1 << ColumnShift
+	uint32 ColumnShift;
+
+	// Index of the current (last) column.
 	uint32 CurrentColumn;
+
+	double LastAllocEndTime = 0.0;
 };
 
 ////////////////////////////////////////////////////////////////////////////////////////////////////
