Normalised the namespace to match the rest of the module. I had nested everything in an additional namespace in error.

#rnx

diff --git a/Engine/Source/Developer/TraceServices/Private/Analyzers/Memory/Config.h b/Engine/Source/Developer/TraceServices/Private/Analyzers/Memory/Config.h
--- a/Engine/Source/Developer/TraceServices/Private/Analyzers/Memory/Config.h
+++ b/Engine/Source/Developer/TraceServices/Private/Analyzers/Memory/Config.h
@@ -4,7 +4,6 @@
 
 #include "CoreTypes.h"
 
-namespace Trace {
 namespace TraceServices {
 
 namespace FTrackerConfig
@@ -21,6 +20,5 @@ namespace FTrackerConfig
 };
 
 } // namespace TraceServices
-} // namespace Trace
 
 /* vim: set noet : */
diff --git a/Engine/Source/Developer/TraceServices/Private/Analyzers/Memory/Lane.h b/Engine/Source/Developer/TraceServices/Private/Analyzers/Memory/Lane.h
--- a/Engine/Source/Developer/TraceServices/Private/Analyzers/Memory/Lane.h
+++ b/Engine/Source/Developer/TraceServices/Private/Analyzers/Memory/Lane.h
@@ -5,7 +5,6 @@
 #include "Containers/ArrayView.h"
 #include "LaneItemSet.h"
 
-namespace Trace {
 namespace TraceServices {
 
 ////////////////////////////////////////////////////////////////////////////////
@@ -127,6 +126,5 @@ private:
 };
 
 } // namespace TraceServices
-} // namespace Trace
 
 /* vim: set noet : */
diff --git a/Engine/Source/Developer/TraceServices/Private/Analyzers/Memory/LaneItemSet.cpp b/Engine/Source/Developer/TraceServices/Private/Analyzers/Memory/LaneItemSet.cpp
--- a/Engine/Source/Developer/TraceServices/Private/Analyzers/Memory/LaneItemSet.cpp
+++ b/Engine/Source/Developer/TraceServices/Private/Analyzers/Memory/LaneItemSet.cpp
@@ -3,7 +3,6 @@
 #include "LaneItemSet.h"
 #include "Support.h"
 
-namespace Trace {
 namespace TraceServices {
 
 ////////////////////////////////////////////////////////////////////////////////
@@ -240,6 +239,5 @@ uint32 FLaneItemSet::GetBucketIndex(uint64 Hash) const
 }
 
 } // namespace TraceServices
-} // namespace Trace
 
 /* vim: set noet : */
diff --git a/Engine/Source/Developer/TraceServices/Private/Analyzers/Memory/LaneItemSet.h b/Engine/Source/Developer/TraceServices/Private/Analyzers/Memory/LaneItemSet.h
--- a/Engine/Source/Developer/TraceServices/Private/Analyzers/Memory/LaneItemSet.h
+++ b/Engine/Source/Developer/TraceServices/Private/Analyzers/Memory/LaneItemSet.h
@@ -4,7 +4,6 @@
 
 #include "CoreTypes.h"
 
-namespace Trace {
 namespace TraceServices {
 
 ////////////////////////////////////////////////////////////////////////////////
@@ -75,6 +74,5 @@ inline uint32 FLaneItemSet::GetMetadataId(uint32 Index) const
 }
 
 } // namespace TraceServices
-} // namespace Trace
 
 /* vim: set noet : */
diff --git a/Engine/Source/Developer/TraceServices/Private/Analyzers/Memory/MetadataDb.cpp b/Engine/Source/Developer/TraceServices/Private/Analyzers/Memory/MetadataDb.cpp
--- a/Engine/Source/Developer/TraceServices/Private/Analyzers/Memory/MetadataDb.cpp
+++ b/Engine/Source/Developer/TraceServices/Private/Analyzers/Memory/MetadataDb.cpp
@@ -3,7 +3,6 @@
 #include "MetadataDb.h"
 #include "Support.h"
 
-namespace Trace {
 namespace TraceServices {
 
 ////////////////////////////////////////////////////////////////////////////////
@@ -94,6 +93,5 @@ uint64 FMetadataDb::Remix(uint64 Id)
 }
 
 } // namespace TraceServices
-} // namespace Trace
 
 /* vim: set noet : */
diff --git a/Engine/Source/Developer/TraceServices/Private/Analyzers/Memory/MetadataDb.h b/Engine/Source/Developer/TraceServices/Private/Analyzers/Memory/MetadataDb.h
--- a/Engine/Source/Developer/TraceServices/Private/Analyzers/Memory/MetadataDb.h
+++ b/Engine/Source/Developer/TraceServices/Private/Analyzers/Memory/MetadataDb.h
@@ -5,7 +5,6 @@
 #include "Containers/Array.h"
 #include "Containers/Map.h"
 
-namespace Trace {
 namespace TraceServices {
 
 ////////////////////////////////////////////////////////////////////////////////
@@ -79,6 +78,5 @@ inline bool FMetadataDb::FEntry::IsRealloc() const
 }
 
 } // namespace TraceServices
-} // namespace Trace
 
 /* vim: set noet : */
diff --git a/Engine/Source/Developer/TraceServices/Private/Analyzers/Memory/SbifBuilder.h b/Engine/Source/Developer/TraceServices/Private/Analyzers/Memory/SbifBuilder.h
--- a/Engine/Source/Developer/TraceServices/Private/Analyzers/Memory/SbifBuilder.h
+++ b/Engine/Source/Developer/TraceServices/Private/Analyzers/Memory/SbifBuilder.h
@@ -4,7 +4,6 @@
 
 #include "SbifIdentities.inl"
 
-namespace Trace {
 namespace TraceServices {
 
 class FMetadataDb;
@@ -135,6 +134,5 @@ public:
 };
 
 } // namespace TraceServices
-} // namespace Trace
 
 /* vim: set noet : */
diff --git a/Engine/Source/Developer/TraceServices/Private/Analyzers/Memory/Support.h b/Engine/Source/Developer/TraceServices/Private/Analyzers/Memory/Support.h
--- a/Engine/Source/Developer/TraceServices/Private/Analyzers/Memory/Support.h
+++ b/Engine/Source/Developer/TraceServices/Private/Analyzers/Memory/Support.h
@@ -8,7 +8,6 @@
 #	include <intrin.h>
 #endif
 
-namespace Trace {
 namespace TraceServices {
 
 ////////////////////////////////////////////////////////////////////////////////
@@ -105,6 +104,5 @@ inline void FTrackerBuffer::Free(void* Ptr)
 }
 
 } // namespace TraceServices
-} // namespace Trace
 
 /* vim: set noet : */
diff --git a/Engine/Source/Developer/TraceServices/Private/Analyzers/Memory/Tracker.cpp b/Engine/Source/Developer/TraceServices/Private/Analyzers/Memory/Tracker.cpp
--- a/Engine/Source/Developer/TraceServices/Private/Analyzers/Memory/Tracker.cpp
+++ b/Engine/Source/Developer/TraceServices/Private/Analyzers/Memory/Tracker.cpp
@@ -8,7 +8,6 @@
 #include "SbifBuilder.h"
 #include "Support.h"
 
-namespace Trace {
 namespace TraceServices {
 
 ////////////////////////////////////////////////////////////////////////////////
@@ -547,6 +546,5 @@ FLaneInput* FTracker::GetLaneInput(uint64 Address)
 }
 
 } // namespace TraceServices
-} // namespace Trace
 
 /* vim: set noet : */
diff --git a/Engine/Source/Developer/TraceServices/Private/Analyzers/Memory/Tracker.h b/Engine/Source/Developer/TraceServices/Private/Analyzers/Memory/Tracker.h
--- a/Engine/Source/Developer/TraceServices/Private/Analyzers/Memory/Tracker.h
+++ b/Engine/Source/Developer/TraceServices/Private/Analyzers/Memory/Tracker.h
@@ -5,7 +5,6 @@
 #include "Config.h"
 #include "MetadataDb.h"
 
-namespace Trace {
 namespace TraceServices {
 
 class FLane;
@@ -64,6 +63,5 @@ private:
 };
 
 } // namespace TraceServices
-} // namespace Trace
 
 /* vim: set noet : */
diff --git a/Engine/Source/Developer/TraceServices/Private/Analyzers/Memory/TrackerJobs.cpp b/Engine/Source/Developer/TraceServices/Private/Analyzers/Memory/TrackerJobs.cpp
--- a/Engine/Source/Developer/TraceServices/Private/Analyzers/Memory/TrackerJobs.cpp
+++ b/Engine/Source/Developer/TraceServices/Private/Analyzers/Memory/TrackerJobs.cpp
@@ -6,7 +6,6 @@
 #include "SbifIdentities.inl"
 #include "Support.h"
 
-namespace Trace {
 namespace TraceServices {
 
 ////////////////////////////////////////////////////////////////////////////////
@@ -310,6 +309,5 @@ void LaneLeaksJob(FLeakJobData* Data)
 }
 
 } // namespace TraceServices
-} // namespace Trace
 
 /* vim: set noet : */
diff --git a/Engine/Source/Developer/TraceServices/Private/Analyzers/Memory/TrackerJobs.h b/Engine/Source/Developer/TraceServices/Private/Analyzers/Memory/TrackerJobs.h
--- a/Engine/Source/Developer/TraceServices/Private/Analyzers/Memory/TrackerJobs.h
+++ b/Engine/Source/Developer/TraceServices/Private/Analyzers/Memory/TrackerJobs.h
@@ -4,7 +4,6 @@
 
 #include "Lane.h"
 
-namespace Trace {
 namespace TraceServices {
 
 ////////////////////////////////////////////////////////////////////////////////
@@ -80,6 +79,5 @@ void LaneRetireeJob(FRetireeJobData* Data);
 void LaneLeaksJob(FLeakJobData* Data);
 
 } // namespace TraceServices
-} // namespace Trace
 
 /* vim: set noet : */
diff --git a/Engine/Source/Developer/TraceServices/Private/Analyzers/Memory/TrackerTime.h b/Engine/Source/Developer/TraceServices/Private/Analyzers/Memory/TrackerTime.h
--- a/Engine/Source/Developer/TraceServices/Private/Analyzers/Memory/TrackerTime.h
+++ b/Engine/Source/Developer/TraceServices/Private/Analyzers/Memory/TrackerTime.h
@@ -2,7 +2,6 @@
 
 #pragma once
 
-namespace Trace {
 namespace TraceServices {
 
 ////////////////////////////////////////////////////////////////////////////////
@@ -21,6 +20,5 @@ private:
 };
 
 } // namespace TraceServices
-} // namespace Trace
 
 /* vim: set noet : */
