Missing file that slipped through the net.

diff --git a/Engine/Source/Developer/TraceServices/Private/Analyzers/Memory/SbifIdentities.inl b/Engine/Source/Developer/TraceServices/Private/Analyzers/Memory/SbifIdentities.inl
--- a/Engine/Source/Developer/TraceServices/Private/Analyzers/Memory/SbifIdentities.inl
+++ b/Engine/Source/Developer/TraceServices/Private/Analyzers/Memory/SbifIdentities.inl
@@ -4,7 +4,6 @@
 
 #include "Support.h"
 
-namespace Trace {
 namespace TraceServices {
 
 ////////////////////////////////////////////////////////////////////////////////
@@ -51,6 +50,5 @@ inline void Sbif_GetColumnDepthFromCell(uint32 CellIndex, uint32* Column, uint32
 }
 
 } // namespace TraceServices
-} // namespace Trace
 
 /* vim: set noet : */
