Bias metadata IDs. A metadata ID of zero is used to detect when metadata is unset so we can not use the internal array indices directly.

#rnx

diff --git a/Engine/Source/Developer/TraceServices/Private/Analyzers/Memory/MetadataDb.cpp b/Engine/Source/Developer/TraceServices/Private/Analyzers/Memory/MetadataDb.cpp
--- a/Engine/Source/Developer/TraceServices/Private/Analyzers/Memory/MetadataDb.cpp
+++ b/Engine/Source/Developer/TraceServices/Private/Analyzers/Memory/MetadataDb.cpp
@@ -5,6 +5,9 @@
 
 namespace TraceServices {
 
+// Used to unsure opaque Id return to outside world won't be zero //////////////
+enum { MetadataDbBias = 1 };
+
 ////////////////////////////////////////////////////////////////////////////////
 bool FMetadataDb::FEntryInternal::operator == (const FEntryInternal& Rhs) const
 {
@@ -67,12 +70,12 @@ uint32 FMetadataDb::AddInternal(uint64 Id, const FEntryInternal& Entry)
 	{
 		Index = Entries.Num();
 		Entries.Add(Entry);
-		return Index;
+		return Index + MetadataDbBias;
 	}
 
 	if (LIKELY(Entries[Index] == Entry))
 	{
-		return Index;
+		return Index + MetadataDbBias;
 	}
 
 	++Collisions;
@@ -83,6 +86,7 @@ uint32 FMetadataDb::AddInternal(uint64 Id, const FEntryInternal& Entry)
 ////////////////////////////////////////////////////////////////////////////////
 const FMetadataDb::FEntry* FMetadataDb::Get(uint32 Id) const
 {
+	Id -= MetadataDbBias;
 	return (const FEntry*)((Id < uint32(Entries.Num())) ? &(Entries[Id]) : nullptr);
 }
 
