Removed Vim mode lines

#rnx

diff --git a/Engine/Source/Developer/TraceServices/Private/Analyzers/Memory/Config.h b/Engine/Source/Developer/TraceServices/Private/Analyzers/Memory/Config.h
--- a/Engine/Source/Developer/TraceServices/Private/Analyzers/Memory/Config.h
+++ b/Engine/Source/Developer/TraceServices/Private/Analyzers/Memory/Config.h
@@ -20,5 +20,3 @@ namespace FTrackerConfig
 };
 
 } // namespace TraceServices
-
-/* vim: set noet : */
diff --git a/Engine/Source/Developer/TraceServices/Private/Analyzers/Memory/Lane.h b/Engine/Source/Developer/TraceServices/Private/Analyzers/Memory/Lane.h
--- a/Engine/Source/Developer/TraceServices/Private/Analyzers/Memory/Lane.h
+++ b/Engine/Source/Developer/TraceServices/Private/Analyzers/Memory/Lane.h
@@ -126,5 +126,3 @@ private:
 };
 
 } // namespace TraceServices
-
-/* vim: set noet : */
diff --git a/Engine/Source/Developer/TraceServices/Private/Analyzers/Memory/LaneItemSet.cpp b/Engine/Source/Developer/TraceServices/Private/Analyzers/Memory/LaneItemSet.cpp
--- a/Engine/Source/Developer/TraceServices/Private/Analyzers/Memory/LaneItemSet.cpp
+++ b/Engine/Source/Developer/TraceServices/Private/Analyzers/Memory/LaneItemSet.cpp
@@ -239,5 +239,3 @@ uint32 FLaneItemSet::GetBucketIndex(uint64 Hash) const
 }
 
 } // namespace TraceServices
-
-/* vim: set noet : */
diff --git a/Engine/Source/Developer/TraceServices/Private/Analyzers/Memory/LaneItemSet.h b/Engine/Source/Developer/TraceServices/Private/Analyzers/Memory/LaneItemSet.h
--- a/Engine/Source/Developer/TraceServices/Private/Analyzers/Memory/LaneItemSet.h
+++ b/Engine/Source/Developer/TraceServices/Private/Analyzers/Memory/LaneItemSet.h
@@ -74,5 +74,3 @@ inline uint32 FLaneItemSet::GetMetadataId(uint32 Index) const
 }
 
 } // namespace TraceServices
-
-/* vim: set noet : */
diff --git a/Engine/Source/Developer/TraceServices/Private/Analyzers/Memory/MetadataDb.cpp b/Engine/Source/Developer/TraceServices/Private/Analyzers/Memory/MetadataDb.cpp
--- a/Engine/Source/Developer/TraceServices/Private/Analyzers/Memory/MetadataDb.cpp
+++ b/Engine/Source/Developer/TraceServices/Private/Analyzers/Memory/MetadataDb.cpp
@@ -97,5 +97,3 @@ uint64 FMetadataDb::Remix(uint64 Id)
 }
 
 } // namespace TraceServices
-
-/* vim: set noet : */
diff --git a/Engine/Source/Developer/TraceServices/Private/Analyzers/Memory/MetadataDb.h b/Engine/Source/Developer/TraceServices/Private/Analyzers/Memory/MetadataDb.h
--- a/Engine/Source/Developer/TraceServices/Private/Analyzers/Memory/MetadataDb.h
+++ b/Engine/Source/Developer/TraceServices/Private/Analyzers/Memory/MetadataDb.h
@@ -78,5 +78,3 @@ inline bool FMetadataDb::FEntry::IsRealloc() const
 }
 
 } // namespace TraceServices
-
-/* vim: set noet : */
diff --git a/Engine/Source/Developer/TraceServices/Private/Analyzers/Memory/SbifBuilder.h b/Engine/Source/Developer/TraceServices/Private/Analyzers/Memory/SbifBuilder.h
--- a/Engine/Source/Developer/TraceServices/Private/Analyzers/Memory/SbifBuilder.h
+++ b/Engine/Source/Developer/TraceServices/Private/Analyzers/Memory/SbifBuilder.h
@@ -134,5 +134,3 @@ public:
 };
 
 } // namespace TraceServices
-
-/* vim: set noet : */
diff --git a/Engine/Source/Developer/TraceServices/Private/Analyzers/Memory/Support.h b/Engine/Source/Developer/TraceServices/Private/Analyzers/Memory/Support.h
--- a/Engine/Source/Developer/TraceServices/Private/Analyzers/Memory/Support.h
+++ b/Engine/Source/Developer/TraceServices/Private/Analyzers/Memory/Support.h
@@ -104,5 +104,3 @@ inline void FTrackerBuffer::Free(void* Ptr)
 }
 
 } // namespace TraceServices
-
-/* vim: set noet : */
diff --git a/Engine/Source/Developer/TraceServices/Private/Analyzers/Memory/Tracker.cpp b/Engine/Source/Developer/TraceServices/Private/Analyzers/Memory/Tracker.cpp
--- a/Engine/Source/Developer/TraceServices/Private/Analyzers/Memory/Tracker.cpp
+++ b/Engine/Source/Developer/TraceServices/Private/Analyzers/Memory/Tracker.cpp
@@ -546,5 +546,3 @@ FLaneInput* FTracker::GetLaneInput(uint64 Address)
 }
 
 } // namespace TraceServices
-
-/* vim: set noet : */
diff --git a/Engine/Source/Developer/TraceServices/Private/Analyzers/Memory/Tracker.h b/Engine/Source/Developer/TraceServices/Private/Analyzers/Memory/Tracker.h
--- a/Engine/Source/Developer/TraceServices/Private/Analyzers/Memory/Tracker.h
+++ b/Engine/Source/Developer/TraceServices/Private/Analyzers/Memory/Tracker.h
@@ -63,5 +63,3 @@ private:
 };
 
 } // namespace TraceServices
-
-/* vim: set noet : */
diff --git a/Engine/Source/Developer/TraceServices/Private/Analyzers/Memory/TrackerJobs.cpp b/Engine/Source/Developer/TraceServices/Private/Analyzers/Memory/TrackerJobs.cpp
--- a/Engine/Source/Developer/TraceServices/Private/Analyzers/Memory/TrackerJobs.cpp
+++ b/Engine/Source/Developer/TraceServices/Private/Analyzers/Memory/TrackerJobs.cpp
@@ -309,5 +309,3 @@ void LaneLeaksJob(FLeakJobData* Data)
 }
 
 } // namespace TraceServices
-
-/* vim: set noet : */
diff --git a/Engine/Source/Developer/TraceServices/Private/Analyzers/Memory/TrackerJobs.h b/Engine/Source/Developer/TraceServices/Private/Analyzers/Memory/TrackerJobs.h
--- a/Engine/Source/Developer/TraceServices/Private/Analyzers/Memory/TrackerJobs.h
+++ b/Engine/Source/Developer/TraceServices/Private/Analyzers/Memory/TrackerJobs.h
@@ -79,5 +79,3 @@ void LaneRetireeJob(FRetireeJobData* Data);
 void LaneLeaksJob(FLeakJobData* Data);
 
 } // namespace TraceServices
-
-/* vim: set noet : */
diff --git a/Engine/Source/Developer/TraceServices/Private/Analyzers/Memory/TrackerTime.h b/Engine/Source/Developer/TraceServices/Private/Analyzers/Memory/TrackerTime.h
--- a/Engine/Source/Developer/TraceServices/Private/Analyzers/Memory/TrackerTime.h
+++ b/Engine/Source/Developer/TraceServices/Private/Analyzers/Memory/TrackerTime.h
@@ -20,5 +20,3 @@ private:
 };
 
 } // namespace TraceServices
-
-/* vim: set noet : */
