Deleted old profiling scopes
#rnx

diff --git a/Engine/Source/Developer/TraceServices/Private/Analyzers/Memory/Support.h b/Engine/Source/Developer/TraceServices/Private/Analyzers/Memory/Support.h
--- a/Engine/Source/Developer/TraceServices/Private/Analyzers/Memory/Support.h
+++ b/Engine/Source/Developer/TraceServices/Private/Analyzers/Memory/Support.h
@@ -21,8 +21,6 @@ inline uint32 UnsafeCountLeadingZeros(uint32 Value)
 #endif
 }
 
-#define PROF_SCOPE(...)
-
 ////////////////////////////////////////////////////////////////////////////////
 class FTrackerBuffer
 {
diff --git a/Engine/Source/Developer/TraceServices/Private/Analyzers/Memory/Tracker.cpp b/Engine/Source/Developer/TraceServices/Private/Analyzers/Memory/Tracker.cpp
--- a/Engine/Source/Developer/TraceServices/Private/Analyzers/Memory/Tracker.cpp
+++ b/Engine/Source/Developer/TraceServices/Private/Analyzers/Memory/Tracker.cpp
@@ -89,10 +89,6 @@ uint32 FTracker::GetCurrentSerial() const
 ////////////////////////////////////////////////////////////////////////////////
 void FTracker::AddAlloc(uint64 Address, const FMetadata& Metadata)
 {
-#if defined(PROF_ON)
-	++AllocDelta;
-#endif
-
 	uint32 MetadataId = MetadataDb.Add(
 		Metadata.Owner,
 		Metadata.Size,
@@ -108,10 +104,6 @@ void FTracker::AddAlloc(uint64 Address, const FMetadata& Metadata)
 ////////////////////////////////////////////////////////////////////////////////
 void FTracker::AddFree(uint64 Address)
 {
-#if defined(PROF_ON)
-	++FreeDelta;
-#endif
-
 	FLaneInput* Input = GetLaneInput(Address);
 	bool bLaneFull = Input->AddFree(Address, Serial);
 	Update(bLaneFull);
@@ -229,8 +221,6 @@ FLaneJobData* FTracker::Sync()
 		return nullptr;
 	}
 
-	PROF_SCOPE("Sync");
-
 	auto* Data = FTrackerScheduler::Wait<FLaneJobData>(SyncWait);
 	SyncWait = 0;
 	return Data;
@@ -347,8 +337,6 @@ void FTracker::RetireAllocs(
 ////////////////////////////////////////////////////////////////////////////////
 void FTracker::ProcessRetirees(const FRetireeJobData* Data)
 {
-	PROF_SCOPE("ProcessRetirees");
-
 	const FRetiree* __restrict Cursor = Data->Retirees->Items;
 	const FRetiree* __restrict End = Cursor + Data->Retirees->Num;
 
@@ -435,22 +423,6 @@ void FTracker::ProcessRetirees(const FRetireeJobData* Data)
 ////////////////////////////////////////////////////////////////////////////////
 void FTracker::Dispatch(bool bDoRehash)
 {
-#if defined(PROF_ON)
-	Running += AllocDelta - FreeDelta;
-	PROF_COUNTER("stats", {
-		{ "allocs", AllocDelta },
-		{ "frees", FreeDelta },
-		{ "alive", Running },
-	});
-
-	PROF_COUNTER("metadata", {
-		{ "#", int64(MetadataDb.GetNum()) },
-		{ "collisions", int64(MetadataDb.GetCollisionNum()) },
-	});
-
-	AllocDelta = FreeDelta = 0;
-#endif
-
 	FTrackerScheduler::FJobHandle LaneJobs[FTrackerConfig::NumLanes] = {};
 	if (bDoRehash)
 	{
diff --git a/Engine/Source/Developer/TraceServices/Private/Analyzers/Memory/Tracker.h b/Engine/Source/Developer/TraceServices/Private/Analyzers/Memory/Tracker.h
--- a/Engine/Source/Developer/TraceServices/Private/Analyzers/Memory/Tracker.h
+++ b/Engine/Source/Developer/TraceServices/Private/Analyzers/Memory/Tracker.h
@@ -54,12 +54,6 @@ private:
 	uint32					ColumnShift;
 	uint32					ColumnMask;
 	uint32					NumColumns;
-
-#if defined(PROF_ON)
-	int32					AllocDelta = 0;
-	int32					FreeDelta = 0;
-	int32					Running = 0;
-#endif
 };
 
 } // namespace TraceServices
diff --git a/Engine/Source/Developer/TraceServices/Private/Analyzers/Memory/TrackerJobs.cpp b/Engine/Source/Developer/TraceServices/Private/Analyzers/Memory/TrackerJobs.cpp
--- a/Engine/Source/Developer/TraceServices/Private/Analyzers/Memory/TrackerJobs.cpp
+++ b/Engine/Source/Developer/TraceServices/Private/Analyzers/Memory/TrackerJobs.cpp
@@ -15,10 +15,7 @@ void LaneRehashJob(void* Data)
 	UPTRINT Ptr = UPTRINT(Data) ^ (bGrow ? ~0ull : 0ull);
 	FLane* Lane = (FLane*)Ptr;
 
-	{
-		PROF_SCOPE(bGrow ? "Grow" : "Rehash");
-		Lane->GetActiveSet().Rehash(bGrow ? FTrackerConfig::ActiveSetPageSize : 0);
-	}
+	Lane->GetActiveSet().Rehash(bGrow ? FTrackerConfig::ActiveSetPageSize : 0);
 }
 
 ////////////////////////////////////////////////////////////////////////////////
@@ -28,14 +25,8 @@ void LaneInputJob(FLaneJobData* Data)
 	FLaneItemView Frees = Data->Input->GetFrees();
 
 	// Get everything in address order.
-	{
-		PROF_SCOPE("AllocSort", {{"#", Allocs.Num}});
-		Algo::Sort(Allocs);
-	}
-	{
-		PROF_SCOPE("FreeSort", {{"#", Frees.Num}});
-		Algo::Sort(Frees);
-	}
+	Algo::Sort(Allocs);
+	Algo::Sort(Frees);
 
 	auto* Retirees = FTrackerBuffer::CallocTemp<FRetirees>(Frees.Num());
 
@@ -53,8 +44,6 @@ void LaneInputJob(FLaneJobData* Data)
 
 	if (Allocs.Num() && Frees.Num())
 	{
-		PROF_SCOPE("Short allocs");
-
 		while (true)
 		{
 			const FLaneItem& AllocItem = *(Iters[0].Cursor);
@@ -117,8 +106,6 @@ void LaneInputJob(FLaneJobData* Data)
 	{
 		uint32 FreesNum = uint32(UPTRINT(FreeWrite - Frees.GetData()));
 
-		PROF_SCOPE("RemainingFrees");
-
 		FLaneItemSet& ActiveSet = Data->Lane->GetActiveSet();
 
 		for (uint32 i = 0, n = FreesNum; i < n; ++i)
@@ -163,29 +150,20 @@ void LaneUpdateJob(FLaneJobData* Data)
 		return;
 	}
 
+	auto Predicate = [] (const FLaneItem& Lhs, const FLaneItem& Rhs)
 	{
-		PROF_SCOPE("UpdateSort");
-
-		auto Predicate = [] (const FLaneItem& Lhs, const FLaneItem& Rhs)
+		if (Lhs.IsSameAddress(Rhs))
 		{
-			if (Lhs.IsSameAddress(Rhs))
-			{
-				return Rhs.HasMetadata();
-			}
+			return Rhs.HasMetadata();
+		}
 
-			return Lhs < Rhs;
-		};
-		Algo::Sort(SetUpdates, Predicate);
-	}
+		return Lhs < Rhs;
+	};
+	Algo::Sort(SetUpdates, Predicate);
 
 	FLaneItemSet& ActiveSet = Data->Lane->GetActiveSet();
 	const FLaneItem* __restrict Ptr = SetUpdates.GetData();
 
-	PROF_SCOPE("SetUpdate", {
-		{"update", SetUpdates.Num()},
-		{"set", ActiveSet.GetNum()},
-	});
-
 	auto ApplyToSet = [&ActiveSet, Data] (const FLaneItem& Item)
 	{
 		if (Item.HasMetadata())
@@ -304,7 +282,6 @@ void LaneLeaksJob(FLeakJobData* Data)
 		SetItem = ActiveSet->NextItem(SetItem);
 	}
 
-	PROF_SCOPE("Retirees");
 	LaneRetireeJob(Data);
 }
 
