Reserve some space for metadata ahead of time.

#rnx

diff --git a/Engine/Source/Developer/TraceServices/Private/Analyzers/Memory/MetadataDb.cpp b/Engine/Source/Developer/TraceServices/Private/Analyzers/Memory/MetadataDb.cpp
--- a/Engine/Source/Developer/TraceServices/Private/Analyzers/Memory/MetadataDb.cpp
+++ b/Engine/Source/Developer/TraceServices/Private/Analyzers/Memory/MetadataDb.cpp
@@ -17,6 +17,8 @@ bool FMetadataDb::FEntryInternal::operator == (const FEntryInternal& Rhs) const
 ////////////////////////////////////////////////////////////////////////////////
 FMetadataDb::FMetadataDb()
 {
+	Map.Reserve(2 << 20);
+	Entries.Reserve(2 << 20);
 }
 
 ////////////////////////////////////////////////////////////////////////////////
