Null-terminate the array of retirees. This will allow us to iterate across the list without having to maintain a count. Useful when doing a k-way merge later.

#rnx

diff --git a/Engine/Source/Developer/TraceServices/Private/Analyzers/Memory/TrackerJobs.cpp b/Engine/Source/Developer/TraceServices/Private/Analyzers/Memory/TrackerJobs.cpp
--- a/Engine/Source/Developer/TraceServices/Private/Analyzers/Memory/TrackerJobs.cpp
+++ b/Engine/Source/Developer/TraceServices/Private/Analyzers/Memory/TrackerJobs.cpp
@@ -28,7 +28,8 @@ void LaneInputJob(FLaneJobData* Data)
 	Algo::Sort(Allocs);
 	Algo::Sort(Frees);
 
-	auto* Retirees = FTrackerBuffer::CallocTemp<FRetirees>(Frees.Num());
+	// +1 so we can null terminate
+	auto* Retirees = FTrackerBuffer::CallocTemp<FRetirees>(Frees.Num() + 1);
 
 	struct FLaneItemIter
 	{
@@ -247,6 +248,8 @@ void LaneRetireeJob(FRetireeJobData* Data)
 
 	TArrayView<FRetiree> Range(Retirees->Items, Retirees->Num);
 	Algo::Sort(Range, Predicate);
+
+	Retirees->Items[Retirees->Num] = {};
 }
 
 ////////////////////////////////////////////////////////////////////////////////
