Resolved some shadowing between leak and retiree job data.

#rnx

diff --git a/Engine/Source/Developer/TraceServices/Private/Analyzers/Memory/Tracker.cpp b/Engine/Source/Developer/TraceServices/Private/Analyzers/Memory/Tracker.cpp
--- a/Engine/Source/Developer/TraceServices/Private/Analyzers/Memory/Tracker.cpp
+++ b/Engine/Source/Developer/TraceServices/Private/Analyzers/Memory/Tracker.cpp
@@ -175,7 +175,7 @@ void FTracker::Finalize()
 			ProcessRetirees(JobData);
 		}
 
-		FLeakJobData* NextData = JobData->Next;
+		auto* NextData = (FLeakJobData*)(JobData->Next);
 		FTrackerBuffer::FreeTemp(JobData->Retirees);
 		FTrackerBuffer::FreeTemp(JobData);
 		JobData = NextData;
@@ -245,7 +245,7 @@ void FTracker::FinalizeWork(FLaneJobData* Data)
 			ProcessRetirees(Data);
 		}
 
-		FLaneJobData* NextData = Data->Next;
+		auto* NextData = (FLaneJobData*)(Data->Next);
 		FTrackerBuffer::FreeTemp(Data->Input);
 		FTrackerBuffer::FreeTemp(Data->Retirees);
 		FTrackerBuffer::FreeTemp(Data);
diff --git a/Engine/Source/Developer/TraceServices/Private/Analyzers/Memory/TrackerJobs.h b/Engine/Source/Developer/TraceServices/Private/Analyzers/Memory/TrackerJobs.h
--- a/Engine/Source/Developer/TraceServices/Private/Analyzers/Memory/TrackerJobs.h
+++ b/Engine/Source/Developer/TraceServices/Private/Analyzers/Memory/TrackerJobs.h
@@ -45,8 +45,15 @@ struct FRetirees
 	FRetiree		Items[];
 };
 
+////////////////////////////////////////////////////////////////////////////////
+struct FBaseJobData
+{
+	FBaseJobData*		Next;
+};
+
 ////////////////////////////////////////////////////////////////////////////////
 struct FRetireeJobData
+	: public FBaseJobData
 {
 	FRetirees*			Retirees;
 	uint32				SerialBias;
@@ -57,7 +64,6 @@ struct FRetireeJobData
 struct FLaneJobData
 	: public FRetireeJobData
 {
-	FLaneJobData*		Next;
 	FLane*				Lane;
 	FLaneInput*			Input;
 	FLaneItemView		SetUpdates;
@@ -67,7 +73,6 @@ struct FLaneJobData
 struct FLeakJobData
 	: public FRetireeJobData
 {
-	FLeakJobData*		Next;
 	const FLaneItemSet* ActiveSet;
 };
 
