Include a tag when adding new allocations

#rnx

diff --git a/Engine/Source/Developer/TraceServices/Private/Analyzers/Memory/MetadataDb.cpp b/Engine/Source/Developer/TraceServices/Private/Analyzers/Memory/MetadataDb.cpp
--- a/Engine/Source/Developer/TraceServices/Private/Analyzers/Memory/MetadataDb.cpp
+++ b/Engine/Source/Developer/TraceServices/Private/Analyzers/Memory/MetadataDb.cpp
@@ -39,7 +39,7 @@ uint32 FMetadataDb::GetCollisionNum() const
 }
 
 ////////////////////////////////////////////////////////////////////////////////
-uint32 FMetadataDb::Add(uint64 Owner, uint64 Size, uint32 Alignment, bool bRealloc)
+uint32 FMetadataDb::Add(uint64 Owner, uint64 Size, uint32 Alignment, uint32 Tag, bool bRealloc)
 {
 	static const uint64 Mix = 0x00000100000001b3;
 
@@ -47,6 +47,7 @@ uint32 FMetadataDb::Add(uint64 Owner, uint64 Size, uint32 Alignment, bool bReall
 	Id ^= Owner;		Id *= Mix;
 	Id ^= Size;			Id *= Mix;
 	Id ^= Alignment;	Id *= Mix;
+	Id ^= Tag;			Id *= Mix;
 	Id ^= !!bRealloc;	Id *= Mix;
 
 	FEntryInternal Entry = {};
diff --git a/Engine/Source/Developer/TraceServices/Private/Analyzers/Memory/MetadataDb.h b/Engine/Source/Developer/TraceServices/Private/Analyzers/Memory/MetadataDb.h
--- a/Engine/Source/Developer/TraceServices/Private/Analyzers/Memory/MetadataDb.h
+++ b/Engine/Source/Developer/TraceServices/Private/Analyzers/Memory/MetadataDb.h
@@ -23,7 +23,7 @@ public:
 						~FMetadataDb();
 	uint32				GetNum() const;
 	uint32				GetCollisionNum() const;
-	uint32				Add(uint64 Owner, uint64 Size, uint32 Alignment, bool bRealloc);
+	uint32				Add(uint64 Owner, uint64 Size, uint32 Alignment, uint32 Tag, bool bRealloc);
 	const FEntry*		Get(uint32 Id) const;
 
 private:
diff --git a/Engine/Source/Developer/TraceServices/Private/Analyzers/Memory/Tracker.cpp b/Engine/Source/Developer/TraceServices/Private/Analyzers/Memory/Tracker.cpp
--- a/Engine/Source/Developer/TraceServices/Private/Analyzers/Memory/Tracker.cpp
+++ b/Engine/Source/Developer/TraceServices/Private/Analyzers/Memory/Tracker.cpp
@@ -93,6 +93,7 @@ void FTracker::AddAlloc(uint64 Address, const FMetadata& Metadata)
 		Metadata.Owner,
 		Metadata.Size,
 		Metadata.Alignment,
+		Metadata.Tag,
 		Metadata.bIsRealloc
 	);
 
diff --git a/Engine/Source/Developer/TraceServices/Private/Analyzers/Memory/Tracker.h b/Engine/Source/Developer/TraceServices/Private/Analyzers/Memory/Tracker.h
--- a/Engine/Source/Developer/TraceServices/Private/Analyzers/Memory/Tracker.h
+++ b/Engine/Source/Developer/TraceServices/Private/Analyzers/Memory/Tracker.h
@@ -22,8 +22,10 @@ public:
 	{
 		uint64				Owner;
 		uint64				Size; 
-		uint32				Alignment;
+		uint32				Tag;
+		uint16				Alignment;
 		bool				bIsRealloc;
+		uint8				_Padding0;
 	};
 
 							FTracker(ISbifBuilder* InSbifBuilder);
