Sort retireed allocations by their end serial

#rnx

diff --git a/Engine/Source/Developer/TraceServices/Private/Analyzers/Memory/Tracker.cpp b/Engine/Source/Developer/TraceServices/Private/Analyzers/Memory/Tracker.cpp
--- a/Engine/Source/Developer/TraceServices/Private/Analyzers/Memory/Tracker.cpp
+++ b/Engine/Source/Developer/TraceServices/Private/Analyzers/Memory/Tracker.cpp
@@ -169,19 +169,16 @@ void FTracker::Finalize()
 	FTrackerScheduler::SubmitJob(RootJob);
 	auto* JobData = FTrackerScheduler::Wait<FLeakJobData>(LeakWait);
 
+	ProcessRetirees(LastJobData);
+
 	do
 	{
-		if (JobData->Retirees != nullptr)
-		{
-			ProcessRetirees(JobData);
-		}
-
-		auto* NextData = (FLeakJobData*)(JobData->Next);
-		FTrackerBuffer::FreeTemp(JobData->Retirees);
-		FTrackerBuffer::FreeTemp(JobData);
-		JobData = NextData;
+		auto* NextData = (FLeakJobData*)(LastJobData->Next);
+		FTrackerBuffer::FreeTemp(LastJobData->Retirees);
+		FTrackerBuffer::FreeTemp(LastJobData);
+		LastJobData = NextData;
 	}
-	while (JobData != nullptr);
+	while (LastJobData != nullptr);
 
 	// Finish off the Sbif
 	/*
@@ -239,13 +236,10 @@ void FTracker::FinalizeWork(FLaneJobData* Data)
 		SbifBuilder->AddColumn();
 	}
 
+	ProcessRetirees(Data);
+
 	do
 	{
-		if (Data->Retirees->Num)
-		{
-			ProcessRetirees(Data);
-		}
-
 		auto* NextData = (FLaneJobData*)(Data->Next);
 		FTrackerBuffer::FreeTemp(Data->Input);
 		FTrackerBuffer::FreeTemp(Data->Retirees);
@@ -293,132 +287,60 @@ void FTracker::FinalizeWork(FLaneJobData* Data)
 	*/
 }
 
-////////////////////////////////////////////////////////////////////////////////
-void FTracker::RetireAllocs(
-	const FRetireeJobData* Data,
-	uint32 Column,
-	uint32 Depth,
-	const FRetiree* __restrict Left,
-	const FRetiree* __restrict Right)
-{
-	static_assert(sizeof(FRetiree) == sizeof(FSbifRetiree), "These are converted up in place so must be the same size");
-
-	uint32 BaseColumn = Sbif_GetBaseColumn(Column, Depth);
-	uint32 BaseSerial = BaseColumn << ColumnShift;
-
-	uint32 Num = uint32(UPTRINT(Right - Left));
-	for (uint32 i = 0; i < Num; ++i)
-	{
-		const FRetiree* __restrict In = Left + i;
-
-		uint64 Address = In->GetAddress();
-		uint32 MetadataId = In->GetMetadataId();
-		uint32 StartSerial = In->GetStartSerial();
-		uint32 EndSerial = In->GetEndSerial(Data->SerialBias);
-
-		auto* __restrict Out = (FSbifRetiree* __restrict)In;
-		Out->_SpaceForAddress = 0xfedbab1e;
-		Out->StartSerial = StartSerial - BaseSerial;
-		Out->EndSerial = EndSerial - BaseSerial;
-		Out->MetadataId = MetadataId;
-
-		check(int32(Out->StartSerial) >= 0);
-		check(int32(Out->EndSerial) >= 0);
-	}
-
-	const auto* __restrict Out = (FSbifRetiree*)Left;
-
-	FSbifContext Context;
-	Context.Column = Column;
-	Context.Depth = Depth;
-	Context.BaseSerial = BaseSerial;
-	SbifBuilder->AddRetirees(&Context, Out, Num);
-}
-
 ////////////////////////////////////////////////////////////////////////////////
 void FTracker::ProcessRetirees(const FRetireeJobData* Data)
 {
-	const FRetiree* __restrict Cursor = Data->Retirees->Items;
-	const FRetiree* __restrict End = Cursor + Data->Retirees->Num;
-
-	auto GetColumn = [Data] (const FRetiree* __restrict Retiree)
+	auto Comparison = [] (const FRetiree& Lhs, const FRetiree& Rhs)
 	{
-		uint32 EndSerial = Retiree->GetEndSerial(Data->SerialBias);
-		return EndSerial >> Data->ColumnShift;
+		return Lhs.GetEndSerialBiased() < Rhs.GetEndSerialBiased();
 	};
 
-	auto GetDepth = [Data] (const FRetiree* __restrict Retiree)
-	{
-		return Sbif_GetCommonDepth(
-			Retiree->GetStartSerial() >> Data->ColumnShift,
-			Retiree->GetEndSerial(Data->SerialBias) >> Data->ColumnShift
-		);
-	};
+	uint32 BundleSerialBias = Data->SerialBias;
 
-	auto ByColumn = [&] (uint32 Column, const FRetiree& __restrict Retiree) -> uint32
+	// Build a min-heap of each retiree list
+	uint32 NumRetirees = 0;
+	TArray<const FRetiree*> Heap;
+	do
 	{
-		return Column < GetColumn(&Retiree);
-	};
+		if (Data->Retirees->Num > 1)
+		{
+			NumRetirees += Data->Retirees->Num;
+			Heap.Add(Data->Retirees->Items);
+		}
+		Data = (FRetireeJobData*)(Data->Next);
+	}
+	while (Data != nullptr);
+	Heap.Heapify(Comparison);
 
-	auto ByDepth = [&] (uint32 Depth, const FRetiree& __restrict Retiree) -> uint32
-	{
-		return Depth < GetDepth(&Retiree);
-	};
+	static const int32 BundleSize = 2048;
+	TArray<FRetiree> Bundle;
+	Bundle.Reserve(BundleSize);
 
-	do
+	for (; NumRetirees; --NumRetirees)
 	{
-		uint32 Depth = GetDepth(Cursor);
-		TArrayView<const FRetiree> Range0(Cursor, int32(UPTRINT(End - Cursor)));
-		const FRetiree* __restrict Next = Cursor + Algo::UpperBound(Range0, Depth, ByDepth);
+		const FRetiree* Ret = Heap.HeapTop();
 
-		uint32 Column = GetColumn(Cursor);
-		if (Column == GetColumn(Next - 1))
+		// Retiree lists are null terminated
+		const FRetiree* Next = Ret + 1;
+		if (Next->GetMetadataId())
 		{
-			/*
-			check(Sbif_GetCellAtDepth(Column, Depth) < (NumColumns * 2) - 1);
-			for (; Cursor < Next; ++Cursor)
-			{
-				uint32 a = Cursor->GetStartSerial();
-				uint32 b = Cursor->GetEndSerial(Data->SerialBias);
-				uint32 bs = Cursor->GetBiasedSerial();
-				check(GetColumn(Cursor) == Column);
-				check(GetDepth(Cursor) == Depth);
-			}
-			*/
-			RetireAllocs(Data, Column, Depth, Cursor, Next);
+			Heap.Add(Next);
 		}
-		else
+		Heap.HeapPopDiscard(Comparison);
+
+		Bundle.Add(*Ret);
+		if (Bundle.Num() >= BundleSize)
 		{
-			do
-			{
-				TArrayView<const FRetiree> Range1(Cursor, int32(UPTRINT(Next - Cursor)));
-				const FRetiree* __restrict SubNext = Cursor + Algo::UpperBound(Range1, Column, ByColumn);
-
-				/*
-				check(Sbif_GetCellAtDepth(Column, Depth) < (NumColumns * 2) - 1);
-				for (; Cursor < SubNext; ++Cursor)
-				{
-					uint32 a = Cursor->GetStartSerial();
-					uint32 b = Cursor->GetEndSerial(Data->SerialBias);
-					uint32 bs = Cursor->GetBiasedSerial();
-					check(GetColumn(Cursor) == Column);
-					check(GetDepth(Cursor) == Depth);
-				}
-				*/
-				RetireAllocs(Data, Column, Depth, Cursor, SubNext);
-
-				Cursor = SubNext;
-				if (Cursor >= Next)
-				{
-					break;
-				}
-			}
-			while (Column = GetColumn(Cursor), true);
+			/* send bundle */
+			Bundle.SetNum(0);
 		}
+	}
 
-		Cursor = Next;
+	if (Bundle.Num() > 0)
+	{
+		/* send bundle */
+		Bundle.SetNum(0);
 	}
-	while (Cursor < End);
 }
 
 ////////////////////////////////////////////////////////////////////////////////
diff --git a/Engine/Source/Developer/TraceServices/Private/Analyzers/Memory/Tracker.h b/Engine/Source/Developer/TraceServices/Private/Analyzers/Memory/Tracker.h
--- a/Engine/Source/Developer/TraceServices/Private/Analyzers/Memory/Tracker.h
+++ b/Engine/Source/Developer/TraceServices/Private/Analyzers/Memory/Tracker.h
@@ -40,7 +40,6 @@ private:
 	FLaneInput*				GetLaneInput(uint64 Address);
 	void					Update(bool bDispatch);
 	void					ProcessRetirees(const FRetireeJobData* JobData);
-	void					RetireAllocs(const FRetireeJobData* Data, uint32 Column, uint32 Depth, const FRetiree* __restrict Left, const FRetiree* __restrict Right);
 	FLaneJobData*			Sync();
 	void					Dispatch(bool bDoRehash=true);
 	void					Provision();
diff --git a/Engine/Source/Developer/TraceServices/Private/Analyzers/Memory/TrackerJobs.cpp b/Engine/Source/Developer/TraceServices/Private/Analyzers/Memory/TrackerJobs.cpp
--- a/Engine/Source/Developer/TraceServices/Private/Analyzers/Memory/TrackerJobs.cpp
+++ b/Engine/Source/Developer/TraceServices/Private/Analyzers/Memory/TrackerJobs.cpp
@@ -226,23 +226,8 @@ void LaneRetireeJob(FRetireeJobData* Data)
 {
 	FRetirees* Retirees = Data->Retirees;
 
-	auto GetDepth = [Data] (const FRetiree& Retiree)
+	auto Predicate = [] (const FRetiree& Lhs, const FRetiree& Rhs)
 	{
-		return Sbif_GetCommonDepth(
-			Retiree.GetStartSerial() >> Data->ColumnShift,
-			Retiree.GetEndSerial(Data->SerialBias) >> Data->ColumnShift
-		);
-	};
-
-	auto Predicate = [&GetDepth] (const FRetiree& Lhs, const FRetiree& Rhs)
-	{
-		uint32 LhsDepth = GetDepth(Lhs);
-		uint32 RhsDepth = GetDepth(Rhs);
-		if (LhsDepth != RhsDepth)
-		{
-			return LhsDepth < RhsDepth;
-		}
-
 		return Lhs.GetEndSerialBiased() < Rhs.GetEndSerialBiased();
 	};
 
