! Hooked up allocation tracker analysis to TaskGraph

#rnx

diff --git a/Engine/Source/Developer/TraceServices/Private/Analyzers/Memory/Tracker.cpp b/Engine/Source/Developer/TraceServices/Private/Analyzers/Memory/Tracker.cpp
--- a/Engine/Source/Developer/TraceServices/Private/Analyzers/Memory/Tracker.cpp
+++ b/Engine/Source/Developer/TraceServices/Private/Analyzers/Memory/Tracker.cpp
@@ -13,21 +13,63 @@ namespace TraceServices {
 ////////////////////////////////////////////////////////////////////////////////
 namespace FTrackerScheduler
 {
-	typedef UPTRINT			FJobHandle;
-	typedef UPTRINT			FWaitHandle;
-	void					SubmitJob(FJobHandle Job) {}
-	FJobHandle				CreateJob(const ANSICHAR* Name) { return 0; }
-	void					DependsOn(FJobHandle This, FJobHandle That) {}
-	void					StartAfter(FJobHandle This, FJobHandle RunsFirst) {}
-	FWaitHandle				MakeWaitable(FJobHandle Job, void* Data=nullptr) { return UPTRINT(Data); }
-	template <class T> T*	Wait(FWaitHandle Waitable) { return (T*)Waitable; }
-
-	template <class T, class U>
-	FJobHandle CreateJob(const ANSICHAR* Name, void (*EntryFunc)(T*), U* Data=nullptr)
+
+////////////////////////////////////////////////////////////////////////////////
+template <typename T, typename U>
+struct FTrackerGraphTask
+{
+	typedef void (EntryFuncType)(T*);
+
+	FTrackerGraphTask(EntryFuncType* InEntryFunc, U* InData)
+	: EntryFunc(InEntryFunc)
+	, Data(InData)
+	{
+	}
+
+	void DoTask(ENamedThreads::Type CurrentThread, const FGraphEventRef& MyCompletionGraphEvent)
 	{
 		EntryFunc(Data);
-		return 0;
 	}
+
+	EntryFuncType*	EntryFunc;
+	U*				Data;
+
+	ENamedThreads::Type GetDesiredThread()				{ return ENamedThreads::AnyThread; }
+	TStatId GetStatId() const							{ return TStatId(); }
+	static ESubsequentsMode::Type GetSubsequentsMode()	{ return ESubsequentsMode::TrackSubsequents; }
+};
+
+////////////////////////////////////////////////////////////////////////////////
+static FGraphEventRef CreateJob(const ANSICHAR* Name, const FGraphEventArray& Prereqs)
+{
+	return TGraphTask<FNullGraphTask>::CreateTask(&Prereqs)
+		.ConstructAndDispatchWhenReady(TStatId(), ENamedThreads::AnyThread);
+}
+
+////////////////////////////////////////////////////////////////////////////////
+template <typename T, typename U>
+static FGraphEventRef CreateJob(
+	const ANSICHAR*	Name,
+	void			(*EntryFunc)(T*),
+	U*				Data=nullptr,
+	FGraphEventRef	Prereq=0)
+{
+	FGraphEventArray Prereqs;
+	if (Prereq)
+	{
+		Prereqs.Add(Prereq);
+	}
+
+	return TGraphTask<FTrackerGraphTask<T,U>>::CreateTask(&Prereqs)
+		.ConstructAndDispatchWhenReady(EntryFunc, Data);
+}
+
+////////////////////////////////////////////////////////////////////////////////
+static void Wait(FGraphEventRef Job)
+{
+	return FTaskGraphInterface::Get().WaitUntilTaskCompletes(Job);
+}
+
 };
 
 
@@ -56,7 +98,7 @@ void FTracker::Begin()
 {
 	SerialBias = 0;
 	Serial = 0;
-	SyncWait = 0;
+	SyncWait = {};
 	NumColumns = 1;
 
 	uint32 EventsPerColumn = SbifBuilder->GetEventsPerColumn();
@@ -143,10 +185,10 @@ void FTracker::Finalize()
 	// column overflow by one serial
 	SerialBias -= 1;
 
-	// Collect leaks
-	FTrackerScheduler::FJobHandle SyncJob = FTrackerScheduler::CreateJob("Sync");
-	FTrackerScheduler::FJobHandle RootJob = FTrackerScheduler::CreateJob("Root");
+	FGraphEventArray TailJobs;
+	TailJobs.Reserve(FTrackerConfig::NumLanes);
 
+	// Dispatch jobs to collect leaks on each lane
 	FLeakJobData* LastJobData = nullptr;
 	for (uint32 i = 0; i < FTrackerConfig::NumLanes; ++i)
 	{
@@ -161,13 +203,11 @@ void FTracker::Finalize()
 		LastJobData = JobData;
 
 		auto LeakJob = FTrackerScheduler::CreateJob("LaneLeaks", LaneLeaksJob, JobData);
-		FTrackerScheduler::StartAfter(LeakJob, RootJob);
-		FTrackerScheduler::DependsOn(SyncJob, LeakJob);
+		TailJobs.Add(LeakJob);
 	}
 
-	FTrackerScheduler::FWaitHandle LeakWait = FTrackerScheduler::MakeWaitable(SyncJob, LastJobData);
-	FTrackerScheduler::SubmitJob(RootJob);
-	auto* JobData = FTrackerScheduler::Wait<FLeakJobData>(LeakWait);
+	auto SyncJob = FTrackerScheduler::CreateJob("Sync", TailJobs);
+	FTrackerScheduler::Wait(SyncJob);
 
 	ProcessRetirees(LastJobData);
 
@@ -214,13 +254,14 @@ void FTracker::Provision()
 ////////////////////////////////////////////////////////////////////////////////
 FLaneJobData* FTracker::Sync()
 {
-	if (!SyncWait)
+	if (!SyncWait.JobRef.IsValid())
 	{
 		return nullptr;
 	}
 
-	auto* Data = FTrackerScheduler::Wait<FLaneJobData>(SyncWait);
-	SyncWait = 0;
+	FTrackerScheduler::Wait(SyncWait.JobRef);
+	auto* Data = (FLaneJobData*)(SyncWait.WaitParam);
+	SyncWait = {};
 	return Data;
 }
 
@@ -346,7 +387,9 @@ void FTracker::ProcessRetirees(const FRetireeJobData* Data)
 ////////////////////////////////////////////////////////////////////////////////
 void FTracker::Dispatch(bool bDoRehash)
 {
-	FTrackerScheduler::FJobHandle LaneJobs[FTrackerConfig::NumLanes] = {};
+	FLaneJobData* DoneJobData = Sync();
+
+	FGraphEventRef LaneJobs[FTrackerConfig::NumLanes] = {};
 	if (bDoRehash)
 	{
 		for (uint32 i = 0; i < FTrackerConfig::NumLanes; ++i)
@@ -374,7 +417,8 @@ void FTracker::Dispatch(bool bDoRehash)
 		}
 	}
 
-	FTrackerScheduler::FJobHandle SyncJob = FTrackerScheduler::CreateJob("Sync");
+	FGraphEventArray TailJobs;
+	TailJobs.Reserve(FTrackerConfig::NumLanes * 2);
 
 	// Lane input jobs
 	FLaneJobData* LastJobData = nullptr;
@@ -394,38 +438,17 @@ void FTracker::Dispatch(bool bDoRehash)
 		JobData->ColumnShift = ColumnShift;
 		LastJobData = JobData;
 
-		auto InputJob = FTrackerScheduler::CreateJob("LaneInput", LaneInputJob, JobData);
-		auto UpdateJob = FTrackerScheduler::CreateJob("LaneUpdate", LaneUpdateJob, JobData);
-		auto RetireeJob = FTrackerScheduler::CreateJob("LaneRetiree", LaneRetireeJob, JobData);
-
-		FTrackerScheduler::StartAfter(UpdateJob, InputJob);
-		FTrackerScheduler::StartAfter(RetireeJob, InputJob);
-
-		FTrackerScheduler::DependsOn(SyncJob, UpdateJob);
-		FTrackerScheduler::DependsOn(SyncJob, RetireeJob);
+		auto InputJob = FTrackerScheduler::CreateJob("LaneInput", LaneInputJob, JobData, LaneJobs[i]);
+		auto UpdateJob = FTrackerScheduler::CreateJob("LaneUpdate", LaneUpdateJob, JobData, InputJob);
+		auto RetireeJob = FTrackerScheduler::CreateJob("LaneRetiree", LaneRetireeJob, JobData, InputJob);
 
-		if (LaneJobs[i])
-		{
-			FTrackerScheduler::StartAfter(InputJob, LaneJobs[i]);
-		}
-		else
-		{
-			LaneJobs[i] = InputJob;
-		}
+		TailJobs.Add(UpdateJob);
+		TailJobs.Add(RetireeJob);
 	}
 
-	FTrackerScheduler::FJobHandle RootJob = FTrackerScheduler::CreateJob("Root");
-	for (FTrackerScheduler::FJobHandle LaneJob : LaneJobs)
-	{
-		if (LaneJob)
-		{
-			FTrackerScheduler::StartAfter(LaneJob, RootJob);
-		}
-	}
+	auto SyncJob = FTrackerScheduler::CreateJob("Sync", TailJobs);
 
-	FLaneJobData* DoneJobData = Sync();
-	SyncWait = FTrackerScheduler::MakeWaitable(SyncJob, LastJobData);
-	FTrackerScheduler::SubmitJob(RootJob);
+	SyncWait = { SyncJob, LastJobData };
 	if (DoneJobData != nullptr)
 	{
 		FinalizeWork(DoneJobData);
diff --git a/Engine/Source/Developer/TraceServices/Private/Analyzers/Memory/Tracker.h b/Engine/Source/Developer/TraceServices/Private/Analyzers/Memory/Tracker.h
--- a/Engine/Source/Developer/TraceServices/Private/Analyzers/Memory/Tracker.h
+++ b/Engine/Source/Developer/TraceServices/Private/Analyzers/Memory/Tracker.h
@@ -2,6 +2,7 @@
 
 #pragma once
 
+#include "Async/TaskGraphInterfaces.h"
 #include "Config.h"
 #include "MetadataDb.h"
 
@@ -37,6 +38,12 @@ public:
 	uint32					GetCurrentSerial() const;
 
 private:
+	struct FSyncWait
+	{
+		FGraphEventRef		JobRef;
+		void*				WaitParam;
+	};
+
 	FLaneInput*				GetLaneInput(uint64 Address);
 	void					Update(bool bDispatch);
 	void					ProcessRetirees(const FRetireeJobData* JobData);
@@ -49,7 +56,7 @@ private:
 	FLane*					Lanes[FTrackerConfig::NumLanes];
 	FLaneInput*				LaneInputs[FTrackerConfig::NumLanes];
 	ISbifBuilder*			SbifBuilder;
-	UPTRINT					SyncWait;
+	FSyncWait				SyncWait;
 	uint32					Serial;
 	uint32					SerialBias;
 	uint32					ColumnShift;
