End serials occupy the most significant bit when packed with addresses. Thus we can sort by the whole word instead of spending time extracting a portion of it

#rnx

diff --git a/Engine/Source/Developer/TraceServices/Private/Analyzers/Memory/Tracker.cpp b/Engine/Source/Developer/TraceServices/Private/Analyzers/Memory/Tracker.cpp
--- a/Engine/Source/Developer/TraceServices/Private/Analyzers/Memory/Tracker.cpp
+++ b/Engine/Source/Developer/TraceServices/Private/Analyzers/Memory/Tracker.cpp
@@ -333,7 +333,7 @@ void FTracker::ProcessRetirees(const FRetireeJobData* Data)
 {
 	auto Comparison = [] (const FRetiree& Lhs, const FRetiree& Rhs)
 	{
-		return Lhs.GetEndSerialBiased() < Rhs.GetEndSerialBiased();
+		return Lhs.GetSortKey() < Rhs.GetSortKey();
 	};
 
 	uint32 BundleSerialBias = Data->SerialBias;
@@ -363,7 +363,7 @@ void FTracker::ProcessRetirees(const FRetireeJobData* Data)
 
 		// Retiree lists are null terminated
 		const FRetiree* Next = Ret + 1;
-		if (Next->GetMetadataId())
+		if (Next->GetSortKey())
 		{
 			Heap.Add(Next);
 		}
diff --git a/Engine/Source/Developer/TraceServices/Private/Analyzers/Memory/TrackerJobs.cpp b/Engine/Source/Developer/TraceServices/Private/Analyzers/Memory/TrackerJobs.cpp
--- a/Engine/Source/Developer/TraceServices/Private/Analyzers/Memory/TrackerJobs.cpp
+++ b/Engine/Source/Developer/TraceServices/Private/Analyzers/Memory/TrackerJobs.cpp
@@ -228,7 +228,7 @@ void LaneRetireeJob(FRetireeJobData* Data)
 
 	auto Predicate = [] (const FRetiree& Lhs, const FRetiree& Rhs)
 	{
-		return Lhs.GetEndSerialBiased() < Rhs.GetEndSerialBiased();
+		return Lhs.GetSortKey() < Rhs.GetSortKey();
 	};
 
 	TArrayView<FRetiree> Range(Retirees->Items, Retirees->Num);
diff --git a/Engine/Source/Developer/TraceServices/Private/Analyzers/Memory/TrackerJobs.h b/Engine/Source/Developer/TraceServices/Private/Analyzers/Memory/TrackerJobs.h
--- a/Engine/Source/Developer/TraceServices/Private/Analyzers/Memory/TrackerJobs.h
+++ b/Engine/Source/Developer/TraceServices/Private/Analyzers/Memory/TrackerJobs.h
@@ -11,6 +11,7 @@ class alignas(16) FRetiree
 {
 public:
 	void			Set(uint32 Start, uint32 EndBiased, uint64 InAddress, uint32 InMetadataId);
+	uint64			GetSortKey() const				{ return SortKey; }
 	uint64			GetAddress() const				{ return Address; }
 	uint64			GetStartSerial() const			{ return StartSerial; }
 	uint32			GetEndSerial(uint32 Bias) const	{ return EndSerialBiased + Bias; }
@@ -18,8 +19,15 @@ public:
 	uint32			GetMetadataId() const			{ return MetadataId; }
 
 private:
-	uint64			Address			: 44;
-	uint64			EndSerialBiased	: 20;
+	union
+	{
+		struct
+		{
+			uint64	Address			: 44;
+			uint64	EndSerialBiased	: 20;
+		};
+		uint64		SortKey; // end-serial is most significant.
+	};
 	uint64			StartSerial		: 36;
 	uint64			MetadataId		: 28;
 };
