Hauled out any building of my swept-binary-interval-forest. The tracker doesn't concern itself with such details, instead concentrates on delivering retireed (and leaked) allocations in end-serial order.

#rnx

diff --git a/Engine/Source/Developer/TraceServices/Private/Analyzers/Memory/Tracker.cpp b/Engine/Source/Developer/TraceServices/Private/Analyzers/Memory/Tracker.cpp
--- a/Engine/Source/Developer/TraceServices/Private/Analyzers/Memory/Tracker.cpp
+++ b/Engine/Source/Developer/TraceServices/Private/Analyzers/Memory/Tracker.cpp
@@ -99,16 +99,6 @@ void FTracker::Begin()
 	SerialBias = 0;
 	Serial = 0;
 	SyncWait = {};
-	NumColumns = 1;
-
-	uint32 EventsPerColumn = SbifBuilder->GetEventsPerColumn();
-	if (EventsPerColumn == 0)
-	{
-		EventsPerColumn = 128 << 10;
-	}
-	ColumnShift = 31 - UnsafeCountLeadingZeros(EventsPerColumn);
-	ColumnMask = (1 << ColumnShift) - 1;
-	check(ColumnShift <= FTrackerConfig::MaxSerialBits);
 
 	SbifBuilder->Begin(&MetadataDb);
 
@@ -199,7 +189,6 @@ void FTracker::Finalize()
 		JobData->Next = LastJobData;
 		JobData->ActiveSet = &ActiveSet;
 		JobData->SerialBias = SerialBias;
-		JobData->ColumnShift = ColumnShift;
 		LastJobData = JobData;
 
 		auto LeakJob = FTrackerScheduler::CreateJob("LaneLeaks", LaneLeaksJob, JobData);
@@ -219,24 +208,6 @@ void FTracker::Finalize()
 		LastJobData = NextData;
 	}
 	while (LastJobData != nullptr);
-
-	// Finish off the Sbif
-	/*
-	uint32 MaxCell = (NumColumns << 1) - 1;
-	for (uint32 i = 0, n = Sbif_GetMaxDepth(NumColumns); i <= n; ++i)
-	{
-		uint32 Delta = 2 << i;
-		uint32 CellIndex = Sbif_GetCellAtDepth(NumColumns - 1, i);
-		for (auto j : {0, 1})
-		{
-			if (CellIndex < MaxCell)
-			{
-				SbifBuilder->CloseCell(CellIndex);
-			}
-			CellIndex -= Delta;
-		}
-	}
-	*/
 }
 
 ////////////////////////////////////////////////////////////////////////////////
@@ -268,14 +239,6 @@ FLaneJobData* FTracker::Sync()
 ////////////////////////////////////////////////////////////////////////////////
 void FTracker::FinalizeWork(FLaneJobData* Data)
 {
-	// SerialBias is the start of the next batch. The batch we're are finalising
-	// (e.g. Data) is up to but not including SerialBias so that's how many
-	// columns we're going to need.
-	while (((SerialBias + ColumnMask) >> ColumnShift) > NumColumns)
-	{
-		NumColumns += 1;
-		SbifBuilder->AddColumn();
-	}
 
 	ProcessRetirees(Data);
 
@@ -288,44 +251,6 @@ void FTracker::FinalizeWork(FLaneJobData* Data)
 		Data = NextData;
 	}
 	while (Data != nullptr);
-	
-	/*
-	for (; ClosedColumns < NumColumns - 2; ++ClosedColumns)
-	{
-		// Columns to close lag two behind the active column. OEIS' A091090 tells us
-		// how many cells are due to close. Cell indices step by 3i, weirdly enough
-
-		uint32 CellIndex = ClosedColumns << 1;
-
-		uint32 C1 = ClosedColumns + 1;
-		uint32 C2 = ClosedColumns + 2;
-		uint32 A = 31 - UnsafeCountLeadingZeros(C1 ^ C2);
-		uint32 B = (C1 & C2) != 0;
-
-		for (uint32 i = 0, n = A + B, d = 3; i < n; ++i, d <<= 1)
-		{
-			SbifBuilder->CloseCell(CellIndex);
-			CellIndex -= d;
-		}
-		
-		/ *
-		uint32 CellIndex = ClosedColumns << 1;
-		uint32 Num = ClosedColumns + 2;
-		Num = (Num & -Num); // extracts lowest set bit
-		for (uint32 d = 3; ; d <<= 1) // As if by magic a multiple of three turns up!!
-		{
-			SbifBuilder->CloseCell(CellIndex);
-			CellIndex -= d;
-
-			Num >>= 1;
-			if (Num == 0 || int32(CellIndex) < 0)
-			{
-				break;
-			}
-		}
-		* /
-	}
-	*/
 }
 
 ////////////////////////////////////////////////////////////////////////////////
@@ -435,7 +360,6 @@ void FTracker::Dispatch(bool bDoRehash)
 		JobData->Lane = Lanes[i];
 		JobData->Input = LaneInput;
 		JobData->SerialBias = SerialBias;
-		JobData->ColumnShift = ColumnShift;
 		LastJobData = JobData;
 
 		auto InputJob = FTrackerScheduler::CreateJob("LaneInput", LaneInputJob, JobData, LaneJobs[i]);
diff --git a/Engine/Source/Developer/TraceServices/Private/Analyzers/Memory/Tracker.h b/Engine/Source/Developer/TraceServices/Private/Analyzers/Memory/Tracker.h
--- a/Engine/Source/Developer/TraceServices/Private/Analyzers/Memory/Tracker.h
+++ b/Engine/Source/Developer/TraceServices/Private/Analyzers/Memory/Tracker.h
@@ -59,9 +59,6 @@ private:
 	FSyncWait				SyncWait;
 	uint32					Serial;
 	uint32					SerialBias;
-	uint32					ColumnShift;
-	uint32					ColumnMask;
-	uint32					NumColumns;
 };
 
 } // namespace TraceServices
diff --git a/Engine/Source/Developer/TraceServices/Private/Analyzers/Memory/TrackerJobs.cpp b/Engine/Source/Developer/TraceServices/Private/Analyzers/Memory/TrackerJobs.cpp
--- a/Engine/Source/Developer/TraceServices/Private/Analyzers/Memory/TrackerJobs.cpp
+++ b/Engine/Source/Developer/TraceServices/Private/Analyzers/Memory/TrackerJobs.cpp
@@ -3,7 +3,6 @@
 #include "TrackerJobs.h"
 #include "Algo/Sort.h"
 #include "Config.h"
-#include "SbifIdentities.inl"
 #include "Support.h"
 
 namespace TraceServices {
diff --git a/Engine/Source/Developer/TraceServices/Private/Analyzers/Memory/TrackerJobs.h b/Engine/Source/Developer/TraceServices/Private/Analyzers/Memory/TrackerJobs.h
--- a/Engine/Source/Developer/TraceServices/Private/Analyzers/Memory/TrackerJobs.h
+++ b/Engine/Source/Developer/TraceServices/Private/Analyzers/Memory/TrackerJobs.h
@@ -65,7 +65,6 @@ struct FRetireeJobData
 {
 	FRetirees*			Retirees;
 	uint32				SerialBias;
-	uint32				ColumnShift;
 };
 
 ////////////////////////////////////////////////////////////////////////////////
