.Off-by-one error. The count doesn't include the null terminating retiree.

#rnx

diff --git a/Engine/Source/Developer/TraceServices/Private/Analyzers/Memory/Tracker.cpp b/Engine/Source/Developer/TraceServices/Private/Analyzers/Memory/Tracker.cpp
--- a/Engine/Source/Developer/TraceServices/Private/Analyzers/Memory/Tracker.cpp
+++ b/Engine/Source/Developer/TraceServices/Private/Analyzers/Memory/Tracker.cpp
@@ -268,7 +268,7 @@ void FTracker::ProcessRetirees(const FRetireeJobData* Data)
 	TArray<const FRetiree*> Heap;
 	do
 	{
-		if (Data->Retirees->Num > 1)
+		if (Data->Retirees->Num)
 		{
 			NumRetirees += Data->Retirees->Num;
 			Heap.Add(Data->Retirees->Items);
