Removed any remenants of interval forest building, replacing the builder interface with a sink for retirees.

#rnx

diff --git a/Engine/Source/Developer/TraceServices/Private/Analyzers/Memory/RetireeSink.h b/Engine/Source/Developer/TraceServices/Private/Analyzers/Memory/RetireeSink.h
new file mode 100644
--- /dev/null
+++ b/Engine/Source/Developer/TraceServices/Private/Analyzers/Memory/RetireeSink.h
@@ -0,0 +1,21 @@
+// Copyright Epic Games, Inc. All Rights Reserved.
+
+#pragma once
+
+#include "Containers/ArrayView.h"
+#include "TrackerJobs.h"
+
+namespace TraceServices {
+
+class FMetadataDb;
+
+////////////////////////////////////////////////////////////////////////////////
+class IRetireeSink
+{
+public:
+	virtual void	Begin(const FMetadataDb* MetadataDb) = 0;
+	virtual void	End() = 0;
+	virtual void	AddRetirees(uint32 SerialBias, TArrayView<FRetiree> Retirees) = 0;
+};
+
+} // namespace TraceServices
diff --git a/Engine/Source/Developer/TraceServices/Private/Analyzers/Memory/SbifBuilder.h b/Engine/Source/Developer/TraceServices/Private/Analyzers/Memory/SbifBuilder.h
deleted file mode 100644
diff --git a/Engine/Source/Developer/TraceServices/Private/Analyzers/Memory/SbifIdentities.inl b/Engine/Source/Developer/TraceServices/Private/Analyzers/Memory/SbifIdentities.inl
deleted file mode 100644
diff --git a/Engine/Source/Developer/TraceServices/Private/Analyzers/Memory/Tracker.cpp b/Engine/Source/Developer/TraceServices/Private/Analyzers/Memory/Tracker.cpp
--- a/Engine/Source/Developer/TraceServices/Private/Analyzers/Memory/Tracker.cpp
+++ b/Engine/Source/Developer/TraceServices/Private/Analyzers/Memory/Tracker.cpp
@@ -5,7 +5,7 @@
 #include "Algo/BinarySearch.h"
 #include "Lane.h"
 #include "TrackerJobs.h"
-#include "SbifBuilder.h"
+#include "RetireeSink.h"
 #include "Support.h"
 
 namespace TraceServices {
@@ -75,8 +75,8 @@ static void Wait(FGraphEventRef Job)
 
 
 ////////////////////////////////////////////////////////////////////////////////
-FTracker::FTracker(ISbifBuilder* InSbifBuilder)
-: SbifBuilder(InSbifBuilder)
+FTracker::FTracker(IRetireeSink* InRetireeSink)
+: RetireeSink(InRetireeSink)
 {
 	for (FLane*& Lane : Lanes)
 	{
@@ -100,7 +100,7 @@ void FTracker::Begin()
 	Serial = 0;
 	SyncWait = {};
 
-	SbifBuilder->Begin(&MetadataDb);
+	RetireeSink->Begin(&MetadataDb);
 
 	Provision();
 }
@@ -109,7 +109,7 @@ void FTracker::Begin()
 void FTracker::End()
 {
 	Finalize();
-	SbifBuilder->End();
+	RetireeSink->End();
 }
 
 ////////////////////////////////////////////////////////////////////////////////
@@ -297,14 +297,14 @@ void FTracker::ProcessRetirees(const FRetireeJobData* Data)
 		Bundle.Add(*Ret);
 		if (Bundle.Num() >= BundleSize)
 		{
-			/* send bundle */
+			RetireeSink->AddRetirees(BundleSerialBias, Bundle);
 			Bundle.SetNum(0);
 		}
 	}
 
 	if (Bundle.Num() > 0)
 	{
-		/* send bundle */
+		RetireeSink->AddRetirees(BundleSerialBias, Bundle);
 		Bundle.SetNum(0);
 	}
 }
@@ -330,8 +330,8 @@ void FTracker::Dispatch(bool bDoRehash)
 
 			bool bGrow = (Capacity - Num) <= AllocNum; // Not enough space to fit potential allocs
 
-			bGrow |= !Capacity || ((100 * Num) / Capacity) > 88;				// Load % is getting too high
-			bool bRehash = Num && ((100 * Tombs) / (Tombs + Num)) > 50;// Too tombstoney
+			bGrow |= !Capacity || ((100 * Num) / Capacity) > 88;		// Load % is getting too high
+			bool bRehash = Num && ((100 * Tombs) / (Tombs + Num)) > 50;	// Too tombstoney
 
 			if (bGrow | bRehash)
 			{
diff --git a/Engine/Source/Developer/TraceServices/Private/Analyzers/Memory/Tracker.h b/Engine/Source/Developer/TraceServices/Private/Analyzers/Memory/Tracker.h
--- a/Engine/Source/Developer/TraceServices/Private/Analyzers/Memory/Tracker.h
+++ b/Engine/Source/Developer/TraceServices/Private/Analyzers/Memory/Tracker.h
@@ -13,7 +13,7 @@ class FLaneInput;
 struct FLaneJobData;
 class FRetiree;
 struct FRetireeJobData;
-class ISbifBuilder;
+class IRetireeSink;
 
 ////////////////////////////////////////////////////////////////////////////////
 class FTracker
@@ -29,7 +29,7 @@ public:
 		uint8				_Padding0;
 	};
 
-							FTracker(ISbifBuilder* InSbifBuilder);
+							FTracker(IRetireeSink* InRetireeSink);
 							~FTracker();
 	void					Begin();
 	void					End();
@@ -55,7 +55,7 @@ private:
 	FMetadataDb				MetadataDb;
 	FLane*					Lanes[FTrackerConfig::NumLanes];
 	FLaneInput*				LaneInputs[FTrackerConfig::NumLanes];
-	ISbifBuilder*			SbifBuilder;
+	IRetireeSink*			RetireeSink;
 	FSyncWait				SyncWait;
 	uint32					Serial;
 	uint32					SerialBias;
