Removed Begin/End() API. It is not entirely clear what the begins and ends signify or what purpose they're supposed to support.

#rnx

diff --git a/Engine/Source/Developer/TraceServices/Private/Analyzers/Memory/RetireeSink.h b/Engine/Source/Developer/TraceServices/Private/Analyzers/Memory/RetireeSink.h
--- a/Engine/Source/Developer/TraceServices/Private/Analyzers/Memory/RetireeSink.h
+++ b/Engine/Source/Developer/TraceServices/Private/Analyzers/Memory/RetireeSink.h
@@ -13,9 +13,14 @@ class FMetadataDb;
 class IRetireeSink
 {
 public:
-	virtual void	Begin(const FMetadataDb* MetadataDb) = 0;
-	virtual void	End() = 0;
-	virtual void	AddRetirees(uint32 SerialBias, TArrayView<FRetiree> Retirees) = 0;
+	struct FRetirements
+	{
+		const class FMetadataDb*	MetadataDb;
+		TArrayView<FRetiree>		Retirees;
+		uint32						SerialBias;
+	};
+
+	virtual void RetireAllocs(const FRetirements& Retirement) = 0;
 };
 
 } // namespace TraceServices
diff --git a/Engine/Source/Developer/TraceServices/Private/Analyzers/Memory/Tracker.cpp b/Engine/Source/Developer/TraceServices/Private/Analyzers/Memory/Tracker.cpp
--- a/Engine/Source/Developer/TraceServices/Private/Analyzers/Memory/Tracker.cpp
+++ b/Engine/Source/Developer/TraceServices/Private/Analyzers/Memory/Tracker.cpp
@@ -100,8 +100,6 @@ void FTracker::Begin()
 	Serial = 0;
 	SyncWait = {};
 
-	RetireeSink->Begin(&MetadataDb);
-
 	Provision();
 }
 
@@ -109,7 +107,6 @@ void FTracker::Begin()
 void FTracker::End()
 {
 	Finalize();
-	RetireeSink->End();
 }
 
 ////////////////////////////////////////////////////////////////////////////////
@@ -261,7 +258,9 @@ void FTracker::ProcessRetirees(const FRetireeJobData* Data)
 		return Lhs.GetSortKey() < Rhs.GetSortKey();
 	};
 
-	uint32 BundleSerialBias = Data->SerialBias;
+	IRetireeSink::FRetirements SinkParam;
+	SinkParam.MetadataDb = &MetadataDb;
+	SinkParam.SerialBias = Data->SerialBias;
 
 	// Build a min-heap of each retiree list
 	uint32 NumRetirees = 0;
@@ -297,14 +296,16 @@ void FTracker::ProcessRetirees(const FRetireeJobData* Data)
 		Bundle.Add(*Ret);
 		if (Bundle.Num() >= BundleSize)
 		{
-			RetireeSink->AddRetirees(BundleSerialBias, Bundle);
+			SinkParam.Retirees = Bundle;
+			RetireeSink->RetireAllocs(SinkParam);
 			Bundle.SetNum(0);
 		}
 	}
 
 	if (Bundle.Num() > 0)
 	{
-		RetireeSink->AddRetirees(BundleSerialBias, Bundle);
+		SinkParam.Retirees = Bundle;
+		RetireeSink->RetireAllocs(SinkParam);
 		Bundle.SetNum(0);
 	}
 }
