Separated the deduplication of allocations' properties and the processing of allocation and free events.

#rnx

diff --git a/Engine/Source/Developer/TraceServices/Private/Analyzers/Memory/RetireeSink.h b/Engine/Source/Developer/TraceServices/Private/Analyzers/Memory/RetireeSink.h
--- a/Engine/Source/Developer/TraceServices/Private/Analyzers/Memory/RetireeSink.h
+++ b/Engine/Source/Developer/TraceServices/Private/Analyzers/Memory/RetireeSink.h
@@ -7,15 +7,12 @@
 
 namespace TraceServices {
 
-class FMetadataDb;
-
 ////////////////////////////////////////////////////////////////////////////////
 class IRetireeSink
 {
 public:
 	struct FRetirements
 	{
-		const class FMetadataDb*	MetadataDb;
 		TArrayView<FRetiree>		Retirees;
 		uint32						SerialBias;
 	};
diff --git a/Engine/Source/Developer/TraceServices/Private/Analyzers/Memory/Tracker.cpp b/Engine/Source/Developer/TraceServices/Private/Analyzers/Memory/Tracker.cpp
--- a/Engine/Source/Developer/TraceServices/Private/Analyzers/Memory/Tracker.cpp
+++ b/Engine/Source/Developer/TraceServices/Private/Analyzers/Memory/Tracker.cpp
@@ -116,16 +116,8 @@ uint32 FTracker::GetCurrentSerial() const
 }
 
 ////////////////////////////////////////////////////////////////////////////////
-void FTracker::AddAlloc(uint64 Address, const FMetadata& Metadata)
+void FTracker::AddAlloc(uint64 Address, uint32 MetadataId)
 {
-	uint32 MetadataId = MetadataDb.Add(
-		Metadata.Owner,
-		Metadata.Size,
-		Metadata.Alignment,
-		Metadata.Tag,
-		Metadata.bIsRealloc
-	);
-
 	FLaneInput* Input = GetLaneInput(Address);
 	bool bLaneFull = Input->AddAlloc(Address, Serial, MetadataId);
 	Update(bLaneFull);
@@ -259,7 +251,6 @@ void FTracker::ProcessRetirees(const FRetireeJobData* Data)
 	};
 
 	IRetireeSink::FRetirements SinkParam;
-	SinkParam.MetadataDb = &MetadataDb;
 	SinkParam.SerialBias = Data->SerialBias;
 
 	// Build a min-heap of each retiree list
diff --git a/Engine/Source/Developer/TraceServices/Private/Analyzers/Memory/Tracker.h b/Engine/Source/Developer/TraceServices/Private/Analyzers/Memory/Tracker.h
--- a/Engine/Source/Developer/TraceServices/Private/Analyzers/Memory/Tracker.h
+++ b/Engine/Source/Developer/TraceServices/Private/Analyzers/Memory/Tracker.h
@@ -4,7 +4,6 @@
 
 #include "Async/TaskGraphInterfaces.h"
 #include "Config.h"
-#include "MetadataDb.h"
 
 namespace TraceServices {
 
@@ -19,21 +18,11 @@ class IRetireeSink;
 class FTracker
 {
 public:
-	struct FMetadata
-	{
-		uint64				Owner;
-		uint64				Size; 
-		uint32				Tag;
-		uint16				Alignment;
-		bool				bIsRealloc;
-		uint8				_Padding0;
-	};
-
 							FTracker(IRetireeSink* InRetireeSink);
 							~FTracker();
 	void					Begin();
 	void					End();
-	void					AddAlloc(uint64 Address, const FMetadata& TeamTada);
+	void					AddAlloc(uint64 Address, uint32 MeatdataId);
 	void					AddFree(uint64 Address);
 	uint32					GetCurrentSerial() const;
 
@@ -52,7 +41,6 @@ private:
 	void					Provision();
 	void					FinalizeWork(FLaneJobData* Data);
 	void					Finalize();
-	FMetadataDb				MetadataDb;
 	FLane*					Lanes[FTrackerConfig::NumLanes];
 	FLaneInput*				LaneInputs[FTrackerConfig::NumLanes];
 	IRetireeSink*			RetireeSink;
