Fixed segfault if a lane did not retire any allocations.

#rnx

diff --git a/Engine/Source/Developer/TraceServices/Private/Analyzers/Memory/Tracker.cpp b/Engine/Source/Developer/TraceServices/Private/Analyzers/Memory/Tracker.cpp
--- a/Engine/Source/Developer/TraceServices/Private/Analyzers/Memory/Tracker.cpp
+++ b/Engine/Source/Developer/TraceServices/Private/Analyzers/Memory/Tracker.cpp
@@ -258,7 +258,7 @@ void FTracker::ProcessRetirees(const FRetireeJobData* Data)
 	TArray<const FRetiree*> Heap;
 	do
 	{
-		if (Data->Retirees->Num)
+		if (Data->Retirees != nullptr && Data->Retirees->Num)
 		{
 			NumRetirees += Data->Retirees->Num;
 			Heap.Add(Data->Retirees->Items);
