Callstack resolving for Insights initial commit

diff --git a/Engine/Source/Developer/TraceInsights/Private/Insights/MemoryProfiler/ViewModels/MemAllocTable.cpp b/Engine/Source/Developer/TraceInsights/Private/Insights/MemoryProfiler/ViewModels/MemAllocTable.cpp
--- a/Engine/Source/Developer/TraceInsights/Private/Insights/MemoryProfiler/ViewModels/MemAllocTable.cpp
+++ b/Engine/Source/Developer/TraceInsights/Private/Insights/MemoryProfiler/ViewModels/MemAllocTable.cpp
@@ -363,26 +363,58 @@ void FMemAllocTable::AddDefaultColumns()
 		Column.SetHorizontalAlignment(HAlign_Left);
 		Column.SetInitialWidth(120.0f);
 
-		Column.SetDataType(ETableCellDataType::Int64);
+		Column.SetDataType(ETableCellDataType::CString);
 
-		class FMemAllocSizeValueGetter : public FTableCellValueGetter
+		class FBacktraceValueGetter : public FTableCellValueGetter
 		{
 		public:
 			virtual const TOptional<FTableCellValue> GetValue(const FTableColumn& Column, const FBaseTreeNode& Node) const override
 			{
+				static const TCHAR* DisplayStrings[] = {
+					TEXT("Pending..."),
+					TEXT("Not found"),
+					TEXT("")
+				};
+				const TCHAR* Value = DisplayStrings[0];
 				const FMemAllocNode& MemAllocNode = static_cast<const FMemAllocNode&>(Node);
 				const FMemoryAlloc* Alloc = MemAllocNode.GetMemAlloc();
-				const uint64 AllocBacktraceId = Alloc ? Alloc->GetBacktraceId() : 0;
-				return TOptional<FTableCellValue>(FTableCellValue(static_cast<int64>(AllocBacktraceId)));
+				if (Alloc)
+				{
+					const TraceServices::FCallstack* Callstack = Alloc->GetCallstack();
+
+					if (Callstack)
+					{
+						const TraceServices::FStackFrame* Frame = Callstack->Frame(2);
+						const TraceServices::QueryResult Result = Frame->Symbol->Result.load(std::memory_order_acquire);
+						switch (Result)
+						{
+						case TraceServices::QueryResult::QR_NotLoaded:
+							Value = DisplayStrings[0];
+							break;
+						case TraceServices::QueryResult::QR_NotFound:
+							Value = DisplayStrings[1];
+							break;
+						case TraceServices::QueryResult::QR_OK:
+							Value = Frame->Symbol->Name;
+							break;
+						}
+					}
+				}
+				else
+				{
+					// Parent nodes (i.e. not actual allocation)
+					Value = DisplayStrings[2];
+				}
+				return TOptional<FTableCellValue>(FTableCellValue(Value));
 			}
 		};
-		TSharedRef<ITableCellValueGetter> Getter = MakeShared<FMemAllocSizeValueGetter>();
+		TSharedRef<ITableCellValueGetter> Getter = MakeShared<FBacktraceValueGetter>();
 		Column.SetValueGetter(Getter);
 
-		TSharedRef<ITableCellValueFormatter> Formatter = MakeShared<FInt64ValueFormatterAsHex64>();
+		TSharedRef<ITableCellValueFormatter> Formatter = MakeShared<FCStringValueFormatterAsText>();
 		Column.SetValueFormatter(Formatter);
 
-		TSharedRef<ITableCellValueSorter> Sorter = MakeShared<FSorterByInt64Value>(ColumnRef);
+		TSharedRef<ITableCellValueSorter> Sorter = MakeShared<FSorterByCStringValue>(ColumnRef);
 		Column.SetValueSorter(Sorter);
 
 		AddColumn(ColumnRef);
diff --git a/Engine/Source/Developer/TraceInsights/Private/Insights/MemoryProfiler/ViewModels/MemoryAlloc.cpp b/Engine/Source/Developer/TraceInsights/Private/Insights/MemoryProfiler/ViewModels/MemoryAlloc.cpp
--- a/Engine/Source/Developer/TraceInsights/Private/Insights/MemoryProfiler/ViewModels/MemoryAlloc.cpp
+++ b/Engine/Source/Developer/TraceInsights/Private/Insights/MemoryProfiler/ViewModels/MemoryAlloc.cpp
@@ -7,13 +7,13 @@ namespace Insights
 
 ////////////////////////////////////////////////////////////////////////////////////////////////////
 
-FMemoryAlloc::FMemoryAlloc(double InStartTime, double InEndTime, uint64 InAddress, uint64 InSize, FMemoryTagId InMemTag, uint64 InBacktraceId)
+FMemoryAlloc::FMemoryAlloc(double InStartTime, double InEndTime, uint64 InAddress, uint64 InSize, FMemoryTagId InMemTag, const TraceServices::FCallstack* InCallstack)
 : StartTime(InStartTime)
 , EndTime(InEndTime)
 , Address(InAddress)
 , Size(InSize)
 , MemTag(InMemTag)
-, BacktraceId(InBacktraceId)
+, Callstack(InCallstack)
 {
 }
 
diff --git a/Engine/Source/Developer/TraceInsights/Private/Insights/MemoryProfiler/ViewModels/MemoryAlloc.h b/Engine/Source/Developer/TraceInsights/Private/Insights/MemoryProfiler/ViewModels/MemoryAlloc.h
--- a/Engine/Source/Developer/TraceInsights/Private/Insights/MemoryProfiler/ViewModels/MemoryAlloc.h
+++ b/Engine/Source/Developer/TraceInsights/Private/Insights/MemoryProfiler/ViewModels/MemoryAlloc.h
@@ -3,7 +3,7 @@
 #pragma once
 
 #include "CoreMinimal.h"
-
+#include "TraceServices/Model/Callstack.h"
 #include "Insights/MemoryProfiler/ViewModels/MemoryTag.h"
 
 namespace Insights
@@ -16,7 +16,7 @@ class FMemoryAlloc
 	friend class SMemAllocTableTreeView;
 
 public:
-	FMemoryAlloc(double InStartTime, double InEndTime, uint64 InAddress, uint64 InSize, FMemoryTagId InMemTag, uint64 InBacktraceId);
+	FMemoryAlloc(double InStartTime, double InEndTime, uint64 InAddress, uint64 InSize, FMemoryTagId InMemTag, const TraceServices::FCallstack* InCallstack);
 	~FMemoryAlloc();
 
 	double GetStartTime() const { return StartTime; }
@@ -25,7 +25,7 @@ public:
 	uint64 GetAddress() const { return Address; }
 	uint64 GetSize() const { return Size; }
 	FMemoryTagId GetMemTag() const { return MemTag; }
-	uint64 GetBacktraceId() const { return BacktraceId; }
+	const TraceServices::FCallstack* GetCallstack() const { return Callstack; }
 
 private:
 	double StartTime;
@@ -33,7 +33,7 @@ private:
 	uint64 Address;
 	uint64 Size;
 	FMemoryTagId MemTag;
-	uint64 BacktraceId;
+	const TraceServices::FCallstack* Callstack;
 };
 
 ////////////////////////////////////////////////////////////////////////////////////////////////////
diff --git a/Engine/Source/Developer/TraceInsights/Private/Insights/MemoryProfiler/Widgets/SMemAllocTableTreeView.cpp b/Engine/Source/Developer/TraceInsights/Private/Insights/MemoryProfiler/Widgets/SMemAllocTableTreeView.cpp
--- a/Engine/Source/Developer/TraceInsights/Private/Insights/MemoryProfiler/Widgets/SMemAllocTableTreeView.cpp
+++ b/Engine/Source/Developer/TraceInsights/Private/Insights/MemoryProfiler/Widgets/SMemAllocTableTreeView.cpp
@@ -4,6 +4,7 @@
 
 #include "TraceServices/AnalysisService.h"
 #include "TraceServices/Model/AllocationsProvider.h"
+#include "TraceServices/Model/Callstack.h"
 
 // Insights
 #include "Insights/MemoryProfiler/MemoryProfilerManager.h"
@@ -319,7 +320,7 @@ void SMemAllocTableTreeView::UpdateQuery(TraceServices::IAllocationsProvider::EQ
 					Alloc.Address = Allocation->GetAddress();
 					Alloc.Size = Allocation->GetSize();
 					Alloc.MemTag = static_cast<FMemoryTagId>(Allocation->GetTag());
-					Alloc.BacktraceId = Allocation->GetBacktraceId();
+					Alloc.Callstack = Allocation->GetCallstack();
 				}
 #endif
 
diff --git a/Engine/Source/Developer/TraceServices/Private/Analyzers/ModuleAnalysis.cpp b/Engine/Source/Developer/TraceServices/Private/Analyzers/ModuleAnalysis.cpp
new file mode 100644
--- /dev/null
+++ b/Engine/Source/Developer/TraceServices/Private/Analyzers/ModuleAnalysis.cpp
@@ -0,0 +1,84 @@
+// Copyright Epic Games, Inc. All Rights Reserved.
+
+#include "ModuleAnalysis.h"
+#include "TraceServices/Model/AnalysisSession.h"
+#include "Model/ModuleProvider.h"
+
+namespace TraceServices {
+
+////////////////////////////////////////////////////////////////////////////////
+enum Routes
+{
+	RouteId_ModuleInit,
+	RouteId_ModuleLoad,
+	RouteId_ModuleUnload,
+};
+
+////////////////////////////////////////////////////////////////////////////////
+FModuleAnalyzer::FModuleAnalyzer(IAnalysisSession& InSession)
+	: Session(InSession)
+	, Provider(nullptr)
+{
+}
+
+////////////////////////////////////////////////////////////////////////////////
+void FModuleAnalyzer::OnAnalysisBegin(const FOnAnalysisContext& Context)
+{
+	FInterfaceBuilder& Builder = Context.InterfaceBuilder;
+	Builder.RouteEvent(RouteId_ModuleInit, "Diagnostics", "ModuleInit");
+	Builder.RouteEvent(RouteId_ModuleLoad, "Diagnostics", "ModuleLoad");
+	Builder.RouteEvent(RouteId_ModuleUnload, "Diagnostics", "ModuleUnload");
+}
+
+
+////////////////////////////////////////////////////////////////////////////////
+bool FModuleAnalyzer::OnEvent(uint16 RouteId, EStyle Style, const FOnEventContext& Context)
+{
+	switch(RouteId)
+	{
+		case RouteId_ModuleInit:
+			{
+				ModuleBaseShift = Context.EventData.GetValue("ModuleBaseShift", 16);
+				FAnsiStringView SymbolFormat;
+				if (Context.EventData.GetString("SymbolFormat", SymbolFormat))
+				{
+					check(Provider == nullptr); // Should only get one init message
+					Provider = CreateModuleProvider(Session, SymbolFormat);
+					Session.AddProvider(TEXT("ModuleProvider"), Provider);
+				}
+			}
+			break;
+
+		case RouteId_ModuleLoad:
+			{
+				const uint64 Base = GetRealBaseAddress(Context.EventData.GetValue<uint32>("Base"));
+				const uint32 Size = Context.EventData.GetValue<uint32>("Size");
+
+				//todo: Use string store
+				FStringView ModuleName;
+				if (Context.EventData.GetString("Name", ModuleName))
+				{
+					Provider->OnModuleLoad(ModuleName, Base, Size);
+				}
+			}
+			break;
+
+		case RouteId_ModuleUnload:
+			{
+				const uint32 Base = Context.EventData.GetValue<uint32>("Base");
+				Provider->OnModuleUnload(Base);
+			}
+			break;
+	}
+	return true;
+}
+
+////////////////////////////////////////////////////////////////////////////////
+uint64 FModuleAnalyzer::GetRealBaseAddress(uint32 EventBase)
+{
+	// Unshift base address according to capture alignment (e.g. Windows 4k aligned)
+	return uint64(EventBase) << ModuleBaseShift;
+}
+
+////////////////////////////////////////////////////////////////////////////////
+} // namespace TraceServices
diff --git a/Engine/Source/Developer/TraceServices/Private/Analyzers/ModuleAnalysis.h b/Engine/Source/Developer/TraceServices/Private/Analyzers/ModuleAnalysis.h
new file mode 100644
--- /dev/null
+++ b/Engine/Source/Developer/TraceServices/Private/Analyzers/ModuleAnalysis.h
@@ -0,0 +1,30 @@
+// Copyright Epic Games, Inc. All Rights Reserved.
+#pragma once
+
+#include "Trace/Analyzer.h"
+
+namespace TraceServices {
+
+class IAnalysisSession;
+class IModuleAnalysisProvider;
+
+////////////////////////////////////////////////////////////////////////////////
+class FModuleAnalyzer
+	: public UE::Trace::IAnalyzer
+{
+public:
+					FModuleAnalyzer(IAnalysisSession& Session);
+	virtual void	OnAnalysisBegin(const FOnAnalysisContext& Context) override;
+	virtual bool 	OnEvent(uint16 RouteId, EStyle Style, const FOnEventContext& Context) override;
+
+private:
+	uint64			GetRealBaseAddress(uint32 EventBase);
+
+	IAnalysisSession& Session;
+	IModuleAnalysisProvider* Provider;
+	uint8 ModuleBaseShift;
+};
+
+////////////////////////////////////////////////////////////////////////////////
+
+} // namespace TraceServices
diff --git a/Engine/Source/Developer/TraceServices/Private/Model/AllocationsProvider.cpp b/Engine/Source/Developer/TraceServices/Private/Model/AllocationsProvider.cpp
--- a/Engine/Source/Developer/TraceServices/Private/Model/AllocationsProvider.cpp
+++ b/Engine/Source/Developer/TraceServices/Private/Model/AllocationsProvider.cpp
@@ -6,6 +6,7 @@
 #include "Containers/ArrayView.h"
 #include "SbTree.h"
 #include "TraceServices/Containers/Allocators.h"
+#include "TraceServices/Model/Callstack.h"
 
 #include <limits>
 
@@ -225,11 +226,10 @@ uint32 IAllocationsProvider::FAllocation::GetAlignment() const
 }
 
 ////////////////////////////////////////////////////////////////////////////////////////////////////
-
-uint64 IAllocationsProvider::FAllocation::GetBacktraceId() const
+const FCallstack* IAllocationsProvider::FAllocation::GetCallstack() const
 {
 	const auto* Inner = (const FAllocationItem*)this;
-	return Inner->Owner;
+	return Inner->Callstack;
 }
 
 ////////////////////////////////////////////////////////////////////////////////////////////////////
@@ -290,19 +290,19 @@ IAllocationsProvider::FQueryResult IAllocationsProvider::FQueryStatus::NextResul
 // FAllocationsProvider
 ////////////////////////////////////////////////////////////////////////////////////////////////////
 
-FAllocationsProvider::FAllocationsProvider(ILinearAllocator& InAllocator)
-	: Allocator(InAllocator)
+FAllocationsProvider::FAllocationsProvider(IAnalysisSession& InSession)
+	: Session(InSession)
 	, SbTree(nullptr)
-	, Timeline(InAllocator, 1024)
-	, MinTotalAllocatedMemoryTimeline(InAllocator, 1024)
-	, MaxTotalAllocatedMemoryTimeline(InAllocator, 1024)
-	, MinLiveAllocationsTimeline(InAllocator, 1024)
-	, MaxLiveAllocationsTimeline(InAllocator, 1024)
-	, AllocEventsTimeline(InAllocator, 1024)
-	, FreeEventsTimeline(InAllocator, 1024)
+	, Timeline(Session.GetLinearAllocator(), 1024)
+	, MinTotalAllocatedMemoryTimeline(Session.GetLinearAllocator(), 1024)
+	, MaxTotalAllocatedMemoryTimeline(Session.GetLinearAllocator(), 1024)
+	, MinLiveAllocationsTimeline(Session.GetLinearAllocator(), 1024)
+	, MaxLiveAllocationsTimeline(Session.GetLinearAllocator(), 1024)
+	, AllocEventsTimeline(Session.GetLinearAllocator(), 1024)
+	, FreeEventsTimeline(Session.GetLinearAllocator(), 1024)
 {
 	const uint32 ColumnShift = 17; // 1<<17 = 128K
-	SbTree = new FSbTree(InAllocator, ColumnShift);
+	SbTree = new FSbTree(Session.GetLinearAllocator(), ColumnShift);
 }
 
 ////////////////////////////////////////////////////////////////////////////////////////////////////
@@ -417,9 +417,9 @@ void FAllocationsProvider::EditAlloc(double Time, uint64 Owner, uint64 Address,
 			Owner,
 			Address,
 			FAllocationItem::PackSizeAndAlignment(Size, InAlignmentAndSizeLower & AlignmentMask),
+			nullptr,
 			Tag,
 			0, // Reserverd1
-			0, // Reserverd2
 		};
 		LiveAllocs.Add(Address, Allocation);
 
@@ -904,7 +904,8 @@ void FAllocationsProvider::DebugPrint() const
 
 IAllocationsProvider::FQueryHandle FAllocationsProvider::StartQuery(const IAllocationsProvider::FQueryParams& Params) const
 {
-	auto* Inner = new FAllocationsQuery(*this, Params);
+	const FCallstacksProvider* CallstacksProvider = Session.ReadProvider<FCallstacksProvider>(FName("CallstacksProvider"));
+	auto* Inner = new FAllocationsQuery(*this, *CallstacksProvider, Params);
 	return IAllocationsProvider::FQueryHandle(Inner);
 }
 
diff --git a/Engine/Source/Developer/TraceServices/Private/Model/AllocationsProvider.h b/Engine/Source/Developer/TraceServices/Private/Model/AllocationsProvider.h
--- a/Engine/Source/Developer/TraceServices/Private/Model/AllocationsProvider.h
+++ b/Engine/Source/Developer/TraceServices/Private/Model/AllocationsProvider.h
@@ -5,11 +5,13 @@
 #include "Common/PagedArray.h"
 #include "Containers/Map.h"
 #include "HAL/CriticalSection.h"
+#include "TraceServices/Model/Callstack.h"
 #include "TraceServices/Model/AllocationsProvider.h"
 
 namespace TraceServices
 {
 
+class IAnalysisSession;
 class ILinearAllocator;
 class FSbTree;
 
@@ -96,9 +98,9 @@ struct FAllocationItem
 	uint64 Owner;
 	uint64 Address;
 	uint64 SizeAndAlignment; // (Alignment << AlignmentShift) | Size
+	mutable const FCallstack* Callstack;
 	uint32 Tag;
 	uint32 Reserved1;
-	uint64 Reserved2;
 };
 
 static_assert(sizeof(FAllocationItem) == 64, "struct FAllocationItem needs packing");
@@ -124,7 +126,7 @@ private:
 	static constexpr double DefaultTimelineSampleGranularity = 0.0001; // 0.1ms
 
 public:
-	FAllocationsProvider(ILinearAllocator& InAllocator);
+	FAllocationsProvider(IAnalysisSession& InSession);
 	virtual ~FAllocationsProvider();
 
 	static FName GetName();
@@ -192,7 +194,7 @@ private:
 	void AdvanceTimelines(double Time);
 
 private:
-	ILinearAllocator& Allocator;
+	IAnalysisSession& Session;
 
 	mutable FAllocationsProviderLock Lock;
 
diff --git a/Engine/Source/Developer/TraceServices/Private/Model/AllocationsQuery.cpp b/Engine/Source/Developer/TraceServices/Private/Model/AllocationsQuery.cpp
--- a/Engine/Source/Developer/TraceServices/Private/Model/AllocationsQuery.cpp
+++ b/Engine/Source/Developer/TraceServices/Private/Model/AllocationsQuery.cpp
@@ -43,8 +43,9 @@ private:
 // FAllocationsQuery
 ////////////////////////////////////////////////////////////////////////////////////////////////////
 
-FAllocationsQuery::FAllocationsQuery(const FAllocationsProvider& InProvider, const IAllocationsProvider::FQueryParams& InParams)
-	: Provider(InProvider)
+FAllocationsQuery::FAllocationsQuery(const FAllocationsProvider& InAllocationProvider, const FCallstacksProvider& InCallstacksProvider, const IAllocationsProvider::FQueryParams& InParams)
+	: AllocationsProvider(InAllocationProvider)
+	, CallstacksProvider(InCallstacksProvider)
 	, Params(InParams)
 	, IsWorking(true)
 	, IsCanceling(false)
@@ -104,12 +105,12 @@ void FAllocationsQuery::Run()
 	uint32 CellCount = 0;
 	uint32 TotalAllocationCount = 0;
 
-	Provider.BeginRead();
+	AllocationsProvider.BeginRead();
 
-	const FSbTree* SbTree = Provider.GetSbTree();
+	const FSbTree* SbTree = AllocationsProvider.GetSbTree();
 	if (SbTree)
 	{
-		UE_LOG(LogTraceServices, Log, TEXT("[MemAlloc] Processing %d live allocs..."), Provider.GetNumLiveAllocs());
+		UE_LOG(LogTraceServices, Log, TEXT("[MemAlloc] Processing %d live allocs..."), AllocationsProvider.GetNumLiveAllocs());
 
 		FAllocationsImpl* LiveAllocsResult = new FAllocationsImpl();
 		QueryLiveAllocs(LiveAllocsResult->Items);
@@ -155,6 +156,15 @@ void FAllocationsQuery::Run()
 			const uint32 NumAllocs = static_cast<uint32>(CellResult->Items.Num());
 			if (NumAllocs != 0)
 			{
+				for (auto Item : CellResult->Items)
+				{
+					if (!Item->Callstack)
+					{
+						Item->Callstack = CallstacksProvider.GetCallstack(Item->Owner);
+						// check(Item->Callstack);
+					}
+				}
+
 				UE_LOG(LogTraceServices, Log, TEXT("[MemAlloc] Enqueue %u allocs..."), NumAllocs);
 				TotalAllocationCount += NumAllocs;
 				Results.Enqueue(CellResult);
@@ -168,7 +178,7 @@ void FAllocationsQuery::Run()
 		UE_LOG(LogTraceServices, Log, TEXT("[MemAlloc] Done"));
 	}
 
-	Provider.EndRead();
+	AllocationsProvider.EndRead();
 
 	IsWorking = false;
 
@@ -185,7 +195,7 @@ void FAllocationsQuery::Run()
 
 void FAllocationsQuery::QueryLiveAllocs(TArray<const FAllocationItem*>& OutAllocs) const
 {
-	const TMap<uint64, FAllocationItem>& LiveAllocs = Provider.GetLiveAllocs();
+	const TMap<uint64, FAllocationItem>& LiveAllocs = AllocationsProvider.GetLiveAllocs();
 
 	// TODO: Live allocs have the end time set to std::numeric_limits<double>::infinity(),
 	//       so the conditions below could be simplified to remove the end time checks.
diff --git a/Engine/Source/Developer/TraceServices/Private/Model/AllocationsQuery.h b/Engine/Source/Developer/TraceServices/Private/Model/AllocationsQuery.h
--- a/Engine/Source/Developer/TraceServices/Private/Model/AllocationsQuery.h
+++ b/Engine/Source/Developer/TraceServices/Private/Model/AllocationsQuery.h
@@ -3,6 +3,7 @@
 #pragma once
 
 #include "AllocationsProvider.h"
+#include "CallstacksProvider.h"
 #include "Async/TaskGraphInterfaces.h"
 #include "Containers/Queue.h"
 #include "HAL/PlatformAtomics.h"
@@ -25,7 +26,7 @@ class FAllocationsQuery
 	friend class FAllocationsQueryAsyncTask;
 
 public:
-	FAllocationsQuery(const FAllocationsProvider& InProvider, const IAllocationsProvider::FQueryParams& InParams);
+	FAllocationsQuery(const FAllocationsProvider& InAllocationsProvider, const FCallstacksProvider& InCallstacksProvider, const IAllocationsProvider::FQueryParams& InParams);
 
 	void Cancel();
 	IAllocationsProvider::FQueryStatus Poll();
@@ -35,7 +36,8 @@ private:
 	void QueryLiveAllocs(TArray<const FAllocationItem*>& OutAllocs) const;
 
 private:
-	const FAllocationsProvider& Provider;
+	const FAllocationsProvider& AllocationsProvider;
+	const FCallstacksProvider& CallstacksProvider;
 
 	IAllocationsProvider::FQueryParams Params;
 
diff --git a/Engine/Source/Developer/TraceServices/Private/Model/CallstacksProvider.cpp b/Engine/Source/Developer/TraceServices/Private/Model/CallstacksProvider.cpp
--- a/Engine/Source/Developer/TraceServices/Private/Model/CallstacksProvider.cpp
+++ b/Engine/Source/Developer/TraceServices/Private/Model/CallstacksProvider.cpp
@@ -2,6 +2,7 @@
 
 #include "CallstacksProvider.h"
 #include "Misc/ScopeRWLock.h"
+#include "ModuleProvider.h"
 #include "TraceServices/Model/AnalysisSession.h"
 #include "Algo/Unique.h"
 #include "Containers/ArrayView.h"
@@ -21,8 +22,11 @@ static struct FCallstackProviderStats
 #endif
 
 /////////////////////////////////////////////////////////////////////
-FCallstacksProvider::FCallstacksProvider(IAnalysisSession& Session)
-	: Frames(Session.GetLinearAllocator(), FramesPerPage) 
+FCallstacksProvider::FCallstacksProvider(IAnalysisSession& InSession)
+	: Session(InSession)
+	, ModuleProvider(nullptr)
+	, Callstacks(InSession.GetLinearAllocator(), CallstacksPerPage)
+	, Frames(InSession.GetLinearAllocator(), FramesPerPage)
 {
 }
 
@@ -40,6 +44,13 @@ void FCallstacksProvider::AddCallstack(uint64 InCallstackId, const uint64* InFra
 	GCallstackStats.FrameCountHistogram[InFrameCount]++;
 #endif
 
+	// The module provider is created on the fly so we want to cache it 
+	// once it's available.
+	if (!ModuleProvider)
+	{
+		ModuleProvider = Session.EditProvider<IModuleProvider>(FName("ModuleProvider"));
+	}
+
 	// Make sure all the frames fit on one page by appending dummy entries.
 	const uint64 PageHeadroom = Frames.GetPageSize() - (Frames.Num() % Frames.GetPageSize());
 	if (PageHeadroom < InFrameCount)
@@ -55,33 +66,30 @@ void FCallstacksProvider::AddCallstack(uint64 InCallstackId, const uint64* InFra
 	{
 		FStackFrame& F = Frames.PushBack();
 		F.Addr = InFrames[FrameIdx];
-		F.Symbol = nullptr;
+		// This will return immediately. The result will be empty if the symbol 
+		// has not been encountered before, and resolution has been queued up.
+		F.Symbol = ModuleProvider->GetSymbol(InFrames[FrameIdx]);
 	}
 
-	uint64 FrameIdxAndLen = (uint64(InFrameCount) << EntryLenShift) | (~EntryLenMask & FirstFrame);
+	const uint64 CallstackId = InCallstackId & ((1ull << 47) - 1);
 
 	{
 		FRWScopeLock WriteLock(EntriesLock, SLT_Write);
-		CallstackEntries.Add(InCallstackId, FCallstackEntry{ FrameIdxAndLen });
+		FCallstack* Callstack = &Callstacks.EmplaceBack(&Frames[FirstFrame], InFrameCount);
+		CallstackEntries.Add(CallstackId, Callstack);
 	}
 }
 
 /////////////////////////////////////////////////////////////////////
-FCallstack FCallstacksProvider::GetCallstack(uint64 CallstackId) 
+const FCallstack* FCallstacksProvider::GetCallstack(uint64 CallstackId) const
 {
 	FRWScopeLock ReadLock(EntriesLock, SLT_ReadOnly);
-	FCallstackEntry* Entry = CallstackEntries.Find(CallstackId);
-	if (Entry)
-	{
-		const uint64 FirstFrameIdx = ~EntryLenMask & Entry->CallstackLenIndex;
-		const uint8 FrameCount = Entry->CallstackLenIndex >> EntryLenShift;
-		return FCallstack(&Frames[FirstFrameIdx], FrameCount);
-	}
-	return FCallstack(nullptr, 0);	
+	const FCallstack* const* FindResult = CallstackEntries.Find(CallstackId);
+	return FindResult ? *FindResult : nullptr;
 }
 
 /////////////////////////////////////////////////////////////////////
-void FCallstacksProvider::GetCallstacks(const TArrayView<uint64>& CallstackIds, FCallstack* OutCallstacks)
+void FCallstacksProvider::GetCallstacks(const TArrayView<uint64>& CallstackIds, FCallstack const** OutCallstacks) const
 {
 	uint64 OutIdx(0);
 	check(OutCallstacks != nullptr);
@@ -89,17 +97,8 @@ void FCallstacksProvider::GetCallstacks(const TArrayView<uint64>& CallstackIds,
 	FRWScopeLock ReadLock(EntriesLock, SLT_ReadOnly);
 	for (uint64 CallstackId : CallstackIds)
 	{
-		FCallstackEntry* Entry = CallstackEntries.Find(CallstackId);
-		if (Entry)
-		{
-			const uint64 FirstFrameIdx = ~EntryLenMask & Entry->CallstackLenIndex;
-			const uint8 FrameCount = Entry->CallstackLenIndex >> EntryLenShift;
-			new(&OutCallstacks[OutIdx]) FCallstack(&Frames[FirstFrameIdx], FrameCount);
-		}
-		else
-		{
-			new(&OutCallstacks[OutIdx]) FCallstack(nullptr, 0);	
-		}
+		const FCallstack* const* FindResult = CallstackEntries.Find(CallstackId);
+		OutCallstacks[OutIdx] = FindResult ? *FindResult : nullptr;
 		OutIdx++;
 	}
 }
diff --git a/Engine/Source/Developer/TraceServices/Private/Model/CallstacksProvider.h b/Engine/Source/Developer/TraceServices/Private/Model/CallstacksProvider.h
--- a/Engine/Source/Developer/TraceServices/Private/Model/CallstacksProvider.h
+++ b/Engine/Source/Developer/TraceServices/Private/Model/CallstacksProvider.h
@@ -11,39 +11,32 @@ namespace TraceServices
 {
 
 class IAnalysisSession;
+class IModuleProvider;
 
 /////////////////////////////////////////////////////////////////////
 class FCallstacksProvider : public ICallstacksProvider
 {
 public:
-					FCallstacksProvider(IAnalysisSession& Session);
-	virtual			~FCallstacksProvider() {}
+						FCallstacksProvider(IAnalysisSession& Session);
+	virtual				~FCallstacksProvider() {}
 
-	FCallstack		GetCallstack(uint64 CallstackId) override;
-	void			GetCallstacks(const TArrayView<uint64>& CallstackIds, FCallstack* OutCallstacks) override;
-	void			AddCallstack(uint64 CallstackId, const uint64* Frames, uint8 FrameCount);
-	FName			GetName() const;
+	const FCallstack*	GetCallstack(uint64 CallstackId) const override;
+	void				GetCallstacks(const TArrayView<uint64>& CallstackIds, FCallstack const** OutCallstacks) const override;
+	void				AddCallstack(uint64 CallstackId, const uint64* Frames, uint8 FrameCount);
+	FName				GetName() const;
 
 private:
-	struct FCallstackEntry
-	{
-		uint64 CallstackLenIndex;
-	};
-
 	enum
 	{
-		FramesPerPage	= 65536, // 16 bytes/entry -> 1 Mb per page
+		FramesPerPage		= 65536, // 16 bytes/entry -> 1 Mb per page
+		CallstacksPerPage	= 65536 * 2 // 8 bytes/callstack -> 1Mb per page
 	};
 
-	enum : uint64
-	{
-		EntryLenMask	= 0xff00000000000000,
-		EntryLenShift	= 56,
-	};
-
-
-	FRWLock							EntriesLock;
-	TMap<uint64, FCallstackEntry>	CallstackEntries;
+	mutable FRWLock					EntriesLock;
+	IAnalysisSession&				Session;
+	mutable IModuleProvider*		ModuleProvider;
+	TMap<uint64,const FCallstack*>	CallstackEntries;
+	TPagedArray<FCallstack>			Callstacks;
 	TPagedArray<FStackFrame>		Frames;
 };
 
diff --git a/Engine/Source/Developer/TraceServices/Private/Model/ModuleProvider.cpp b/Engine/Source/Developer/TraceServices/Private/Model/ModuleProvider.cpp
new file mode 100644
--- /dev/null
+++ b/Engine/Source/Developer/TraceServices/Private/Model/ModuleProvider.cpp
@@ -0,0 +1,160 @@
+// Copyright Epic Games, Inc. All Rights Reserved.
+
+#include "ModuleProvider.h"
+#include "TraceServices/Model/AnalysisSession.h"
+#include "TraceServices/Model/Diagnostics.h"
+#include "HAL/Runnable.h"
+#include "HAL/RunnableThread.h"
+#include <atomic>
+
+// Symbol files implementations
+#include "PbdSymbols.inl"
+
+namespace TraceServices {
+
+/////////////////////////////////////////////////////////////////////
+struct FModuleEntry
+{
+	const TCHAR*		Name;
+	const TCHAR*		Path;
+	uint64				Base;
+	uint32				Size;
+	bool				bLoaded;
+	std::atomic<bool>	bReady;
+
+	FModuleEntry(const TCHAR* InName, const TCHAR* InPath, uint64 InBase, uint32 InSize, bool InLoaded)
+		: Name(InName)
+		, Path(InPath)
+		, Base(InBase)
+		, Size(InSize)
+		, bLoaded(InLoaded)
+		, bReady(false)
+	{
+
+	}
+};
+
+/////////////////////////////////////////////////////////////////////
+template<typename SymbolResolver>
+class TModuleProvider : public IModuleAnalysisProvider
+{
+public:
+								TModuleProvider(IAnalysisSession& Session);
+	virtual 					~TModuleProvider() {}
+
+	//Query interface
+	const FResolvedSymbol*		GetSymbol(uint64 Address) override;
+	void						GetStats(FStats* OutStats) override;
+	
+	
+	//Analysis interface
+	void						OnModuleLoad(const FStringView& Module, uint64 Base, uint32 Size) override;
+	void 						OnModuleUnload(uint64 Base) override;
+
+private:
+
+	FRWLock						ModulesLock;
+	TPagedArray<FModuleEntry>	Modules;
+
+	FRWLock						SymbolsLock;
+	TPagedArray<FResolvedSymbol> Symbols;
+	TMap<uint64, const FResolvedSymbol*> SymbolCache;
+
+	IAnalysisSession&			Session;
+	FString						Platform;
+	SymbolResolver 				Resolver;
+};
+
+/////////////////////////////////////////////////////////////////////
+template<typename SymbolProvider>
+TModuleProvider<SymbolProvider>::TModuleProvider(IAnalysisSession& Session)
+	: Modules(Session.GetLinearAllocator(), 128)
+	, Symbols(Session.GetLinearAllocator(), 1024*1024)
+	, Session(Session)
+	, Resolver(Session)
+{
+}
+
+/////////////////////////////////////////////////////////////////////
+template<typename SymbolProvider>
+const FResolvedSymbol* TModuleProvider<SymbolProvider>::GetSymbol(uint64 Address)
+{
+	{
+		// Attempt to read from the cached symbols.
+		FReadScopeLock _(SymbolsLock);
+		if (const FResolvedSymbol** Entry = SymbolCache.Find(Address))
+		{
+			return *Entry;
+		}
+	}
+
+	FResolvedSymbol* ResolvedSymbol = nullptr;
+	{
+		// Add a pending entry to our cache
+		FWriteScopeLock _(SymbolsLock);
+		if (SymbolCache.Contains(Address))
+		{
+			return SymbolCache[Address];
+		}
+		ResolvedSymbol = &Symbols.EmplaceBack();
+		ResolvedSymbol->Result = QR_NotLoaded;
+		ResolvedSymbol->Name = nullptr;
+		ResolvedSymbol->FileAndLine = nullptr;
+		SymbolCache.Add(Address, ResolvedSymbol);
+	}
+
+	// If not in cache yet, queue it up in the resolver
+	check(ResolvedSymbol);
+	Resolver.QueueSymbolResolve(Address, ResolvedSymbol);
+
+	return ResolvedSymbol;
+}
+
+/////////////////////////////////////////////////////////////////////
+template<typename SymbolResolver>
+void TModuleProvider<SymbolResolver>::GetStats(FStats* OutStats)
+{
+	Resolver.GetStats(OutStats);
+}
+
+/////////////////////////////////////////////////////////////////////
+template<typename SymbolProvider>
+void TModuleProvider<SymbolProvider>::OnModuleLoad(const FStringView& Module, uint64 Base, uint32 Size)
+{
+	if (Module.Len() == 0)
+	{
+		return;
+	}
+
+	const TCHAR* Path = Session.StoreString(Module);
+	const TCHAR* Name = Session.StoreString(Module);
+
+	FModuleEntry* NewEntry = nullptr;
+	{
+		FWriteScopeLock _(ModulesLock);
+		NewEntry = &Modules.EmplaceBack(Name, Path, Base, Size, false);
+	}
+
+	Resolver.QueueModuleLoad(Module, Base, Size);
+}
+
+/////////////////////////////////////////////////////////////////////
+template<typename SymbolProvider>
+void TModuleProvider<SymbolProvider>::OnModuleUnload(uint64 Base)
+{
+	//todo: Find entry, set bLoaded to false
+}
+
+/////////////////////////////////////////////////////////////////////
+IModuleAnalysisProvider* CreateModuleProvider(IAnalysisSession& InSession, const FAnsiStringView& InSymbolFormat)
+{
+#if PLATFORM_WINDOWS
+	if (InSymbolFormat.Equals("pdb"))
+	{
+		return new TModuleProvider<FPdbSymbols>(InSession);
+	}
+#endif
+	return nullptr;
+}
+
+} // namespace TraceServices
diff --git a/Engine/Source/Developer/TraceServices/Private/Model/ModuleProvider.h b/Engine/Source/Developer/TraceServices/Private/Model/ModuleProvider.h
new file mode 100644
--- /dev/null
+++ b/Engine/Source/Developer/TraceServices/Private/Model/ModuleProvider.h
@@ -0,0 +1,25 @@
+// Copyright Epic Games, Inc. All Rights Reserved.
+#pragma once
+
+#include "TraceServices/Model/Modules.h"
+#include "Common/PagedArray.h"
+#include "Misc/ScopeRWLock.h"
+
+namespace TraceServices 
+{
+
+class IAnalysisSession;
+
+/////////////////////////////////////////////////////////////////////
+class IModuleAnalysisProvider : public IModuleProvider
+{
+public:
+	virtual void	OnModuleLoad(const FStringView& Module, uint64 Base, uint32 Size) = 0;
+	virtual void	OnModuleUnload(uint64 Base) = 0;
+};
+
+/** Create a module provider with the given symbol format */
+IModuleAnalysisProvider* CreateModuleProvider(IAnalysisSession& Session, const FAnsiStringView& SymbolFormat);
+
+
+} // namespace TraceServices
diff --git a/Engine/Source/Developer/TraceServices/Private/Model/PbdSymbols.inl b/Engine/Source/Developer/TraceServices/Private/Model/PbdSymbols.inl
new file mode 100644
--- /dev/null
+++ b/Engine/Source/Developer/TraceServices/Private/Model/PbdSymbols.inl
@@ -0,0 +1,285 @@
+// Copyright Epic Games, Inc. All Rights Reserved.
+
+#if PLATFORM_WINDOWS
+
+#include "Algo/Sort.h"
+#include "Containers/StringView.h"
+#include "Containers/Queue.h"
+#include "HAL/CriticalSection.h"
+#include "HAL/PlatformProcess.h"
+#include "Logging/LogMacros.h"
+#include "Misc/PathViews.h"
+#include "Misc/ScopeRWLock.h"
+#include "Misc/StringBuilder.h"
+#include "TraceServices/Model/AnalysisSession.h"
+#include <atomic>
+
+#include "Windows/AllowWindowsPlatformTypes.h"
+#include <dbghelp.h>
+
+/////////////////////////////////////////////////////////////////////
+DEFINE_LOG_CATEGORY_STATIC(LogPdbSymbols, Log, All);
+
+/////////////////////////////////////////////////////////////////////
+namespace TraceServices {
+
+/////////////////////////////////////////////////////////////////////
+
+static const TCHAR* UNKNOWN_MODULE_TEXT = TEXT("Unknown");
+
+/////////////////////////////////////////////////////////////////////
+class FPdbSymbols : public FRunnable
+{
+public:
+	FPdbSymbols(IAnalysisSession& InSession)
+		: bInitialized(false)
+		, Session(InSession)
+	{
+		static uint64 BaseHandle = 0x493;
+		Handle = (HANDLE)(++BaseHandle);
+
+		// Load DbgHelp interface
+		ULONG SymOpts = 0;
+		SymOpts |= SYMOPT_LOAD_LINES;
+		SymOpts |= SYMOPT_OMAP_FIND_NEAREST;
+		SymOpts |= SYMOPT_DEFERRED_LOADS;
+		SymOpts |= SYMOPT_EXACT_SYMBOLS;
+		SymOpts |= SYMOPT_IGNORE_NT_SYMPATH;
+		SymOpts |= SYMOPT_UNDNAME;
+
+		SymSetOptions(SymOpts);
+		bInitialized = SymInitialize(Handle, NULL, FALSE); 
+
+		if (bInitialized)
+		{
+			// Start the worker thread
+			bRunWorkerThread = true;
+			Thread = FRunnableThread::Create(this, TEXT("PbdSymbolWorker"), 0, TPri_Normal);
+		}
+	}
+
+	~FPdbSymbols()
+	{
+		bRunWorkerThread = false;
+		Thread->WaitForCompletion();
+
+		SymCleanup(Handle);
+	}
+
+	void QueueModuleLoad(const FStringView& ModulePath, uint64 Base, uint32 Size)
+	{
+		FScopeLock _(&ModulesCs);
+
+		FStringView ModuleName = FPathViews::GetCleanFilename(ModulePath);
+
+		// Add module and sort list according to base address
+		int32 Index = LoadedModules.Add(ModuleEntry{ Base, Size, Session.StoreString(ModuleName), Session.StoreString(ModulePath), false });
+
+		// Queue up module to have symbols loaded
+		LoadSymbolsQueue.Enqueue(QueuedModule{ Base, LoadedModules[Index].Path });
+
+		// Sort list according to base address
+		Algo::Sort(LoadedModules, [](const ModuleEntry& Lhs, const ModuleEntry& Rhs) { return Lhs.Base < Rhs.Base; });
+
+		ModulesDiscovered++;
+	}
+
+	void QueueSymbolResolve(uint64 Address, FResolvedSymbol* Symbol)
+	{
+		SymbolsDiscovered++;
+		ResolveQueue.Enqueue(QueuedAddress{Address, Symbol});
+	}
+
+	void GetStats(IModuleProvider::FStats* OutStats)
+	{
+		OutStats->ModulesDiscovered = ModulesDiscovered.load();
+		OutStats->ModulesFailed = ModulesFailed.load();
+		OutStats->ModulesLoaded = ModulesLoaded.load();
+		OutStats->SymbolsDiscovered = SymbolsDiscovered.load();
+		OutStats->SymbolsFailed = SymbolsFailed.load();
+		OutStats->SymbolsResolved = SymbolsResolved.load();
+	}
+
+private:
+
+	struct ModuleEntry
+	{
+		uint64 Base;
+		uint32 Size;
+		const TCHAR* Name;
+		const TCHAR* Path;
+		bool bSymbolsLoaded;
+	};
+
+	struct QueuedAddress
+	{
+		uint64 Address;
+		FResolvedSymbol* Target;
+	};
+
+	struct QueuedModule
+	{
+		uint64 Base;
+		const TCHAR* ImagePath;
+	};
+
+	enum : uint32 {
+		MaxNameLen = 512,
+	};
+
+
+	uint32 Run() override 
+	{
+		while (bRunWorkerThread)
+		{
+			// Prioritize queued module loads
+			while (!LoadSymbolsQueue.IsEmpty() && bRunWorkerThread)
+			{
+				QueuedModule Item;
+				if (LoadSymbolsQueue.Dequeue(Item))
+				{
+					LoadModuleSymbols(Item.Base, Item.ImagePath);
+				}
+			}
+
+			// Resolve one symbol at a time to give way for modules
+			while (!ResolveQueue.IsEmpty() && LoadSymbolsQueue.IsEmpty() && bRunWorkerThread)
+			{
+				QueuedAddress Item;
+				if (ResolveQueue.Dequeue(Item))
+				{
+					ResolveSymbol(Item.Address, Item.Target);
+				}
+			}
+			
+			// ...and breathe...
+			FPlatformProcess::Sleep(0.2f);
+		}
+		
+		return 0;
+	}
+
+	void Stop() override 
+	{
+		bRunWorkerThread = false; 
+	}
+
+	void UpdateResolvedSymbol(FResolvedSymbol* Symbol, QueryResult Result, const TCHAR* Name, const TCHAR* FileAndLine)
+	{
+		Symbol->Name = Name;
+		Symbol->FileAndLine = FileAndLine;
+		Symbol->Result.store(Result, std::memory_order_release);
+	}
+
+	void ResolveSymbol(uint64 Address, FResolvedSymbol* Target)
+	{
+		check(Target);
+
+		uint8 InfoBuffer[sizeof(SYMBOL_INFO) + (MaxNameLen * sizeof(char) + 1)];
+		SYMBOL_INFO* Info = (SYMBOL_INFO*)InfoBuffer;
+		Info->SizeOfStruct = sizeof(SYMBOL_INFO);
+		Info->MaxNameLen = MaxNameLen;
+
+		// Find and build the symbol name
+		if (!SymFromAddr(Handle, Address, NULL, Info))
+		{
+			DWORD Err = GetLastError();
+			SymbolsFailed++;
+			UpdateResolvedSymbol(Target, QR_NotFound, UNKNOWN_MODULE_TEXT, UNKNOWN_MODULE_TEXT);
+			return;
+		}
+
+		TStringBuilder<256> SymbolName;
+		SymbolName << ANSI_TO_TCHAR(Info->Name);
+		const TCHAR* SymbolNameStr = Session.StoreString(FStringView(SymbolName));
+
+		// Find the source file and line
+		DWORD  dwDisplacement;
+		IMAGEHLP_LINE Line;
+		Line.SizeOfStruct = sizeof(IMAGEHLP_LINE);
+
+		if (!SymGetLineFromAddr(Handle, Address, &dwDisplacement, &Line))
+		{
+			SymbolsFailed++;
+			UpdateResolvedSymbol(Target, QR_OK, SymbolNameStr, UNKNOWN_MODULE_TEXT);
+			return;
+		}
+
+		TStringBuilder<256> FileAndLine;
+		FileAndLine << ANSI_TO_TCHAR(Line.FileName) << TEXT(" (") << uint32(Line.LineNumber) << TEXT(")");
+		const TCHAR* FileAndLineStr = Session.StoreString(FStringView(FileAndLine));
+		
+		SymbolsResolved++;
+		UpdateResolvedSymbol(Target, QR_OK, SymbolNameStr, FileAndLineStr);
+	}
+
+	bool LoadModuleSymbols(uint64 Base, const TCHAR* Path)
+	{
+		// Attempt to load symbols
+		const DWORD64 LoadedBaseAddress = SymLoadModuleEx(Handle, NULL, TCHAR_TO_ANSI(Path), NULL, Base, 0x7fffffff /*Size*/, NULL, 0);
+		const bool bSymbolsLoaded = Base == LoadedBaseAddress;
+
+		if (!bSymbolsLoaded)
+		{
+			UE_LOG(LogPdbSymbols, Warning, TEXT("Unable to load symbols for %s at %p"), Path, Base);
+			ModulesFailed++;
+		}
+		else
+		{
+			UE_LOG(LogPdbSymbols, Display, TEXT("Loaded symbols for %s at %p."), Path, Base);
+			ModulesLoaded++;
+		}
+
+		// Update the module entry with the result
+		FScopeLock _(&ModulesCs);
+		const int32 EntryIdx = Algo::BinarySearchBy(LoadedModules, Base, [](const ModuleEntry& Entry) { return Entry.Base; });
+		check(EntryIdx >= 0 && EntryIdx < LoadedModules.Num());
+		LoadedModules[EntryIdx].bSymbolsLoaded = bSymbolsLoaded;
+		
+		return bSymbolsLoaded;
+	}
+
+
+	const TCHAR* GetModuleNameForAddress(uint64 Address) const
+	{
+		const int32 EntryIdx = Algo::LowerBoundBy(LoadedModules, Address, [](const ModuleEntry& Entry) { return Entry.Base; }) - 1;
+		if (EntryIdx < 0 || EntryIdx >= LoadedModules.Num())
+		{
+			return UNKNOWN_MODULE_TEXT;
+		}
+
+		const ModuleEntry& Entry = LoadedModules[EntryIdx];
+		if (Address > (Entry.Base + Entry.Size))
+		{
+			return UNKNOWN_MODULE_TEXT;
+		}
+
+		return Entry.Name;
+	}
+
+	FCriticalSection ModulesCs;
+	TArray<ModuleEntry> LoadedModules;
+	TQueue<QueuedModule, EQueueMode::Mpsc> LoadSymbolsQueue;
+	TQueue<QueuedAddress, EQueueMode::Mpsc> ResolveQueue;
+
+	std::atomic<uint32> ModulesDiscovered;
+	std::atomic<uint32> ModulesFailed;
+	std::atomic<uint32> ModulesLoaded;
+	std::atomic<uint32> SymbolsDiscovered;
+	std::atomic<uint32> SymbolsFailed;
+	std::atomic<uint32> SymbolsResolved;
+
+	bool bRunWorkerThread;
+	bool bInitialized;
+	HANDLE Handle;
+	IAnalysisSession& Session;
+	FRunnableThread* Thread;
+};
+
+#include "Windows/HideWindowsPlatformTypes.h"
+
+/////////////////////////////////////////////////////////////////////
+
+} // namespace TraceServices
+
+#endif // PLATFORM_WINDOWS
diff --git a/Engine/Source/Developer/TraceServices/Private/Modules/MemoryModule.cpp b/Engine/Source/Developer/TraceServices/Private/Modules/MemoryModule.cpp
--- a/Engine/Source/Developer/TraceServices/Private/Modules/MemoryModule.cpp
+++ b/Engine/Source/Developer/TraceServices/Private/Modules/MemoryModule.cpp
@@ -4,6 +4,7 @@
 #include "Analyzers/AllocationsAnalysis.h"
 #include "Analyzers/CallstacksAnalysis.h"
 #include "Analyzers/MemoryAnalysis.h"
+#include "Analyzers/ModuleAnalysis.h"
 #include "Model/AllocationsProvider.h"
 #include "Model/CallstacksProvider.h"
 #include "TraceServices/Model/AnalysisSession.h"
@@ -21,10 +22,12 @@ void FMemoryModule::GetModuleInfo(FModuleInfo& OutModuleInfo)
 
 void FMemoryModule::OnAnalysisBegin(IAnalysisSession& Session)
 {
-	FAllocationsProvider* AllocationsProvider = new FAllocationsProvider(Session.GetLinearAllocator());
+	FAllocationsProvider* AllocationsProvider = new FAllocationsProvider(Session);
 	Session.AddProvider(AllocationsProvider->GetName(), AllocationsProvider);
 	Session.AddAnalyzer(new FAllocationsAnalyzer(Session, *AllocationsProvider));
 
+	Session.AddAnalyzer(new FModuleAnalyzer(Session));
+	
 	FCallstacksProvider* CallstacksProvider = new FCallstacksProvider(Session);
 	Session.AddProvider(CallstacksProvider->GetName(), CallstacksProvider);
 	Session.AddAnalyzer(new FCallstacksAnalyzer(Session, CallstacksProvider));
diff --git a/Engine/Source/Developer/TraceServices/Public/TraceServices/Model/AllocationsProvider.h b/Engine/Source/Developer/TraceServices/Public/TraceServices/Model/AllocationsProvider.h
--- a/Engine/Source/Developer/TraceServices/Public/TraceServices/Model/AllocationsProvider.h
+++ b/Engine/Source/Developer/TraceServices/Public/TraceServices/Model/AllocationsProvider.h
@@ -88,7 +88,7 @@ public:
 		uint64 GetAddress() const;
 		uint64 GetSize() const;
 		uint32 GetAlignment() const;
-		uint64 GetBacktraceId() const;
+		const struct FCallstack* GetCallstack() const;
 		uint32 GetTag() const;
 	};
 
diff --git a/Engine/Source/Developer/TraceServices/Public/TraceServices/Model/Callstack.h b/Engine/Source/Developer/TraceServices/Public/TraceServices/Model/Callstack.h
--- a/Engine/Source/Developer/TraceServices/Public/TraceServices/Model/Callstack.h
+++ b/Engine/Source/Developer/TraceServices/Public/TraceServices/Model/Callstack.h
@@ -4,42 +4,43 @@
 
 #include "TraceServices/Model/AnalysisSession.h"
 #include "TraceServices/Model/Memory.h"
+#include "TraceServices/Model/Modules.h"
 
 namespace TraceServices
 {
 
 /////////////////////////////////////////////////////////////////////
-struct SymbolEntry
+struct TRACESERVICES_API FStackFrame
 {
-	TCHAR*			Name;
-	TCHAR*			Filename;
-	TCHAR*			Module;
-	uint32			Line;
+	uint64					Addr;
+	const FResolvedSymbol*	Symbol;
 };
 
 /////////////////////////////////////////////////////////////////////
-struct FStackFrame
+struct TRACESERVICES_API FCallstack
 {
-	uint64			Addr;
-	SymbolEntry*	Symbol;
-};
-
-/////////////////////////////////////////////////////////////////////
-struct FCallstack
-{
-					FCallstack(FStackFrame* FirstEntry, uint8 FrameCount);
+	FCallstack(const FStackFrame* FirstEntry, uint8 FrameCount);
 	/** Get the number of stack frames in callstack. */
-	uint32			Num();
+	uint32			Num() const;
 	/** Gets the address at a given stack depth. */
-	uint64			Addr(uint8 Depth);
+	uint64			Addr(uint8 Depth) const;
 	/** Gets the cached symbol name at a given stack depth. */
-	TCHAR*			SymbolName(uint8 Depth);
+	const TCHAR*	Name(uint8 Depth) const;
+	/** Gets the entire frame at given depth */
+	const FStackFrame* Frame(uint8 Depth) const;
 
 private:
-	FStackFrame*	FirstEntry;
-	uint8			FrameCount;
+	enum : uint64
+	{
+		EntryLenShift = 56,
+		EntryLenMask = uint64(0xff) << EntryLenShift,
+	};
+
+	uint64 CallstackLenIndex;
 };
 
+static_assert(sizeof(FCallstack) == 8, "struct FCallstack is too large");
+
 /////////////////////////////////////////////////////////////////////
 class ICallstacksProvider
 	: public IProvider
@@ -52,39 +53,51 @@ public:
 	  * @param CallstackId		Callstack id to query
 	  * @return					Callstack information. If id is not found a callstack with zero stack depth is returned.
 	  */
-	virtual FCallstack	GetCallstack(uint64 CallstackId) = 0;
+	virtual const FCallstack*	GetCallstack(uint64 CallstackId) const = 0;
 
 	/**
 	  * Queries a set of callstack ids.
 	  * @param CallstackIds		List of callstack ids to query
 	  * @param OutCallstacks	Output list of callstacks. Caller is responsible for allocating space according to CallstackIds length.
 	  */
-	virtual void		GetCallstacks(const TArrayView<uint64>& CallstackIds, FCallstack* OutCallstacks) = 0;
+	virtual void				GetCallstacks(const TArrayView<uint64>& CallstackIds, FCallstack const** OutCallstacks) const = 0;
 };
 
 /////////////////////////////////////////////////////////////////////
-inline FCallstack::FCallstack(FStackFrame* InFirstEntry, uint8 InFrameCount)
-	: FirstEntry(InFirstEntry)
-	, FrameCount(InFrameCount)
+inline FCallstack::FCallstack(const FStackFrame* InFirstFrame, uint8 InFrameCount)
+{
+	check((uint64(InFirstFrame) & EntryLenMask) == 0);
+	CallstackLenIndex = (uint64(InFrameCount) << EntryLenShift) | (~EntryLenMask & uint64(InFirstFrame));
+}
+
+/////////////////////////////////////////////////////////////////////
+inline uint32 FCallstack::Num() const
 {
+	return CallstackLenIndex >> EntryLenShift;
 }
 
 /////////////////////////////////////////////////////////////////////
-inline uint32 FCallstack::Num()
+inline uint64 FCallstack::Addr(uint8 Depth) const
 {
-	return FrameCount;
+	const FStackFrame* FirstFrame = (const FStackFrame *) (~EntryLenMask & CallstackLenIndex);
+	const uint32 FrameCount = Num();
+	return (Depth < FrameCount) ? (FirstFrame + Depth)->Addr : 0;
 }
 
 /////////////////////////////////////////////////////////////////////
-inline uint64 FCallstack::Addr(uint8 Depth)
+inline const TCHAR* FCallstack::Name(uint8 Depth) const
 {
-	return Depth < FrameCount ? FirstEntry[Depth].Addr : 0;
+	const FStackFrame* FirstFrame = (const FStackFrame*)(~EntryLenMask & CallstackLenIndex);
+	const uint32 FrameCount = Num();
+	return (Depth < FrameCount) ? (FirstFrame + Depth)->Symbol->Name : 0;
 }
 
 /////////////////////////////////////////////////////////////////////
-inline TCHAR* FCallstack::SymbolName(uint8 Depth)
+inline const FStackFrame* FCallstack::Frame(uint8 Depth) const
 {
-	return Depth < FrameCount ? FirstEntry[Depth].Symbol->Name : 0;
+	const FStackFrame* FirstFrame = (const FStackFrame*)(~EntryLenMask & CallstackLenIndex);
+	const uint32 FrameCount = Num();
+	return (Depth < FrameCount) ? (FirstFrame + Depth) : nullptr;
 }
 
 } // namespace TraceServices
diff --git a/Engine/Source/Developer/TraceServices/Public/TraceServices/Model/Modules.h b/Engine/Source/Developer/TraceServices/Public/TraceServices/Model/Modules.h
new file mode 100644
--- /dev/null
+++ b/Engine/Source/Developer/TraceServices/Public/TraceServices/Model/Modules.h
@@ -0,0 +1,56 @@
+// Copyright Epic Games, Inc. All Rights Reserved.
+
+#pragma once
+#include "TraceServices/Model/AnalysisSession.h"
+#include <atomic>
+
+namespace TraceServices {
+
+////////////////////////////////////////////////////////////////////////////////
+/**
+  * Result of a query. Since symbol resolving can be deferred this signals if a
+  * symbol has been resolved, waiting to be resolved or wasn't found at all.
+  */
+enum QueryResult 
+{
+	QR_OK,
+	QR_NotFound,
+	QR_NotLoaded,
+};
+
+////////////////////////////////////////////////////////////////////////////////
+/**
+  * Represent a resolved symbol. The resolve status and string values may change
+  * over time, but string pointers returned from the methods are guaranteed to live 
+  * during the entire analysis session.
+  */
+struct FResolvedSymbol
+{
+	std::atomic<QueryResult> Result;
+	const TCHAR* Name;
+	const TCHAR* FileAndLine;
+};
+
+////////////////////////////////////////////////////////////////////////////////
+class IModuleProvider : public IProvider
+{
+public:
+	struct FStats
+	{
+		uint32 ModulesDiscovered;
+		uint32 ModulesLoaded;
+		uint32 ModulesFailed;
+		uint32 SymbolsDiscovered;
+		uint32 SymbolsResolved;
+		uint32 SymbolsFailed;
+	};
+	
+	/** Queries the name of the symbol at address */
+	virtual const FResolvedSymbol* GetSymbol(uint64 Address) = 0;
+
+	/** Gets statistics from provider */
+	virtual void GetStats(FStats* OutStats) = 0;
+};
+
+////////////////////////////////////////////////////////////////////////////////
+} // namespace TraceServices
diff --git a/Engine/Source/Runtime/Core/Private/ProfilingDebugging/Windows/ModuleDiagnostics.cpp b/Engine/Source/Runtime/Core/Private/ProfilingDebugging/Windows/ModuleDiagnostics.cpp
--- a/Engine/Source/Runtime/Core/Private/ProfilingDebugging/Windows/ModuleDiagnostics.cpp
+++ b/Engine/Source/Runtime/Core/Private/ProfilingDebugging/Windows/ModuleDiagnostics.cpp
@@ -32,6 +32,11 @@ struct FNtDllFunction
 
 ////////////////////////////////////////////////////////////////////////////////
 UE_TRACE_CHANNEL(ModuleChannel, "Module information needed for symbols resolution", true)
+ 
+UE_TRACE_EVENT_BEGIN(Diagnostics, ModuleInit, NoSync|Important)
+	UE_TRACE_EVENT_FIELD(UE::Trace::AnsiString, SymbolFormat)
+	UE_TRACE_EVENT_FIELD(uint8, ModuleBaseShift)
+UE_TRACE_EVENT_END()
 
 UE_TRACE_EVENT_BEGIN(Diagnostics, ModuleLoad, NoSync|Important)
 	UE_TRACE_EVENT_FIELD(UE::Trace::WideString, Name)
@@ -92,6 +97,10 @@ void FModuleTrace::Initialize()
 {
 	using namespace UE::Trace;
 
+	UE_TRACE_LOG(Diagnostics, ModuleInit, ModuleChannel, sizeof(char) * 3)
+		<< ModuleInit.SymbolFormat("pdb", 3)
+		<< ModuleInit.ModuleBaseShift(uint8(16));
+
 	// Register for DLL load/unload notifications.
 	auto Thunk = [] (ULONG Reason, const void* Data, void* Context)
 	{
