Fix non unity compile errors

diff --git a/Engine/Source/Developer/TraceServices/Private/Analyzers/ModuleAnalysis.cpp b/Engine/Source/Developer/TraceServices/Private/Analyzers/ModuleAnalysis.cpp
--- a/Engine/Source/Developer/TraceServices/Private/Analyzers/ModuleAnalysis.cpp
+++ b/Engine/Source/Developer/TraceServices/Private/Analyzers/ModuleAnalysis.cpp
@@ -3,6 +3,7 @@
 #include "ModuleAnalysis.h"
 #include "TraceServices/Model/AnalysisSession.h"
 #include "Model/ModuleProvider.h"
+#include "UObject/NameTypes.h"
 
 namespace TraceServices {
 
diff --git a/Engine/Source/Developer/TraceServices/Private/Model/ModuleProvider.cpp b/Engine/Source/Developer/TraceServices/Private/Model/ModuleProvider.cpp
--- a/Engine/Source/Developer/TraceServices/Private/Model/ModuleProvider.cpp
+++ b/Engine/Source/Developer/TraceServices/Private/Model/ModuleProvider.cpp
@@ -1,6 +1,7 @@
 // Copyright Epic Games, Inc. All Rights Reserved.
 
 #include "ModuleProvider.h"
+#include "Containers/Map.h"
 #include "TraceServices/Model/AnalysisSession.h"
 #include "TraceServices/Model/Diagnostics.h"
 #include "HAL/Runnable.h"
diff --git a/Engine/Source/Developer/TraceServices/Private/Model/PbdSymbols.inl b/Engine/Source/Developer/TraceServices/Private/Model/PbdSymbols.inl
--- a/Engine/Source/Developer/TraceServices/Private/Model/PbdSymbols.inl
+++ b/Engine/Source/Developer/TraceServices/Private/Model/PbdSymbols.inl
@@ -9,6 +9,7 @@
 #include "HAL/PlatformProcess.h"
 #include "Logging/LogMacros.h"
 #include "Misc/PathViews.h"
+#include "Misc/ScopeLock.h"
 #include "Misc/ScopeRWLock.h"
 #include "Misc/StringBuilder.h"
 #include "TraceServices/Model/AnalysisSession.h"
