[Insights] Remove shift of callstack id. Matching runtime shift has already been removed.

#rnx

diff --git a/Engine/Source/Developer/TraceServices/Private/Model/CallstacksProvider.cpp b/Engine/Source/Developer/TraceServices/Private/Model/CallstacksProvider.cpp
--- a/Engine/Source/Developer/TraceServices/Private/Model/CallstacksProvider.cpp
+++ b/Engine/Source/Developer/TraceServices/Private/Model/CallstacksProvider.cpp
@@ -71,12 +71,10 @@ void FCallstacksProvider::AddCallstack(uint64 InCallstackId, const uint64* InFra
 		F.Symbol = ModuleProvider->GetSymbol(InFrames[FrameIdx]);
 	}
 
-	const uint64 CallstackId = InCallstackId & ((1ull << 47) - 1);
-
 	{
 		FRWScopeLock WriteLock(EntriesLock, SLT_Write);
 		FCallstack* Callstack = &Callstacks.EmplaceBack(&Frames[FirstFrame], InFrameCount);
-		CallstackEntries.Add(CallstackId, Callstack);
+		CallstackEntries.Add(InCallstackId, Callstack);
 	}
 }
 
