Fixing PVS 7.7 Issues:

warning V1064: The '(16 - 1)' operand of integer division is less than the '16' one. The result will always be zero.

diff --git a/Engine/Source/Developer/TraceServices/Private/Common/SlabAllocator.h b/Engine/Source/Developer/TraceServices/Private/Common/SlabAllocator.h
--- a/Engine/Source/Developer/TraceServices/Private/Common/SlabAllocator.h
+++ b/Engine/Source/Developer/TraceServices/Private/Common/SlabAllocator.h
@@ -29,7 +29,7 @@ public:
 
 	virtual void* Allocate(uint64 Size) override
 	{
-		uint64 AllocationSize = Size + (16 - 1) / 16;
+		uint64 AllocationSize = Size + (-int64(Size) & 15);
 		if (!CurrentSlab || CurrentSlabAllocatedSize + AllocationSize > SlabSize)
 		{
 			TotalAllocatedSize += SlabSize;
