[Insights] Fixed the EQueryRule::aAfB (decline) mem query.

diff --git a/Engine/Source/Developer/TraceServices/Private/Model/AllocationsQuery.cpp b/Engine/Source/Developer/TraceServices/Private/Model/AllocationsQuery.cpp
--- a/Engine/Source/Developer/TraceServices/Private/Model/AllocationsQuery.cpp
+++ b/Engine/Source/Developer/TraceServices/Private/Model/AllocationsQuery.cpp
@@ -252,7 +252,7 @@ void FAllocationsQuery::QueryLiveAllocs(TArray<const FAllocationItem*>& OutAlloc
 		for (const auto& KV : LiveAllocs)
 		{
 			const FAllocationItem& Alloc = KV.Value;
-			if (Alloc.StartTime <= TimeA && Alloc.EndTime >= TimeA && Alloc.EndTime >= TimeB)
+			if (Alloc.StartTime <= TimeA && Alloc.EndTime >= TimeA && Alloc.EndTime <= TimeB)
 			{
 				OutAllocs.Add(&Alloc);
 			}
diff --git a/Engine/Source/Developer/TraceServices/Private/Model/SbTree.cpp b/Engine/Source/Developer/TraceServices/Private/Model/SbTree.cpp
--- a/Engine/Source/Developer/TraceServices/Private/Model/SbTree.cpp
+++ b/Engine/Source/Developer/TraceServices/Private/Model/SbTree.cpp
@@ -136,7 +136,7 @@ void FSbTreeCell::Query(TArray<const FAllocationItem*>& OutAllocs, const IAlloca
 		const double TimeB = Params.TimeB;
 		for (const FAllocationItem& Alloc : Allocs)
 		{
-			if (Alloc.StartTime <= TimeA && Alloc.EndTime >= TimeA && Alloc.EndTime >= TimeB)
+			if (Alloc.StartTime <= TimeA && Alloc.EndTime >= TimeA && Alloc.EndTime <= TimeB)
 			{
 				OutAllocs.Add(&Alloc);
 			}
