[Insights] Use module size when loading symbols

Use the traced module size when loading symbols for a module. Previous code overlapped module regions causing dynamically loaded configurations to fail resolving symbols.

#rnx

diff --git a/Engine/Source/Developer/TraceServices/Private/Model/PbdSymbols.inl b/Engine/Source/Developer/TraceServices/Private/Model/PbdSymbols.inl
--- a/Engine/Source/Developer/TraceServices/Private/Model/PbdSymbols.inl
+++ b/Engine/Source/Developer/TraceServices/Private/Model/PbdSymbols.inl
@@ -77,7 +77,7 @@ public:
 		int32 Index = LoadedModules.Add(ModuleEntry{ Base, Size, Session.StoreString(ModuleName), Session.StoreString(ModulePath), false });
 
 		// Queue up module to have symbols loaded
-		LoadSymbolsQueue.Enqueue(QueuedModule{ Base, LoadedModules[Index].Path });
+		LoadSymbolsQueue.Enqueue(QueuedModule{ Base, Size, LoadedModules[Index].Path });
 
 		// Sort list according to base address
 		Algo::Sort(LoadedModules, [](const ModuleEntry& Lhs, const ModuleEntry& Rhs) { return Lhs.Base < Rhs.Base; });
@@ -121,6 +121,7 @@ private:
 	struct QueuedModule
 	{
 		uint64 Base;
+		uint64 Size;
 		const TCHAR* ImagePath;
 	};
 
@@ -139,7 +140,7 @@ private:
 				QueuedModule Item;
 				if (LoadSymbolsQueue.Dequeue(Item))
 				{
-					LoadModuleSymbols(Item.Base, Item.ImagePath);
+					LoadModuleSymbols(Item.Base, Item.Size, Item.ImagePath);
 				}
 			}
 
@@ -214,10 +215,10 @@ private:
 		UpdateResolvedSymbol(Target, QR_OK, SymbolNameStr, FileAndLineStr);
 	}
 
-	bool LoadModuleSymbols(uint64 Base, const TCHAR* Path)
+	bool LoadModuleSymbols(uint64 Base, uint64 Size, const TCHAR* Path)
 	{
 		// Attempt to load symbols
-		const DWORD64 LoadedBaseAddress = SymLoadModuleEx(Handle, NULL, TCHAR_TO_ANSI(Path), NULL, Base, 0x7fffffff /*Size*/, NULL, 0);
+		const DWORD64 LoadedBaseAddress = SymLoadModuleEx(Handle, NULL, TCHAR_TO_ANSI(Path), NULL, Base, Size, NULL, 0);
 		const bool bSymbolsLoaded = Base == LoadedBaseAddress;
 
 		if (!bSymbolsLoaded)
