[Insights] Make string store thread safe

#rnx

diff --git a/Engine/Source/Developer/TraceServices/Private/Common/StringStore.cpp b/Engine/Source/Developer/TraceServices/Private/Common/StringStore.cpp
--- a/Engine/Source/Developer/TraceServices/Private/Common/StringStore.cpp
+++ b/Engine/Source/Developer/TraceServices/Private/Common/StringStore.cpp
@@ -18,6 +18,7 @@ const TCHAR* FStringStore::Store(const TCHAR* String)
 
 const TCHAR* FStringStore::Store(const FStringView& String)
 {
+	FScopeLock _(&Cs);
 	uint32 Hash = GetTypeHash(String);
 	const TCHAR** AlreadyStored = StoredStrings.Find(Hash);
 	if (AlreadyStored && !String.Compare(FStringView(*AlreadyStored)))
diff --git a/Engine/Source/Developer/TraceServices/Private/Common/StringStore.h b/Engine/Source/Developer/TraceServices/Private/Common/StringStore.h
--- a/Engine/Source/Developer/TraceServices/Private/Common/StringStore.h
+++ b/Engine/Source/Developer/TraceServices/Private/Common/StringStore.h
@@ -23,6 +23,7 @@ private:
 	{
 		BlockSize = 4 << 20
 	};
+	FCriticalSection Cs;
 	FSlabAllocator& Allocator;
 	TMap<uint32, const TCHAR*> StoredStrings;
 	TCHAR* BufferPtr = nullptr;
