[Insights] Fix missing callstacks on live allocations

Code that resolved callstacks from owner id was missing from "live allocations" set, causing many N/A callstacks in memory queries.

#rnx

diff --git a/Engine/Source/Developer/TraceServices/Private/Model/AllocationsQuery.cpp b/Engine/Source/Developer/TraceServices/Private/Model/AllocationsQuery.cpp
--- a/Engine/Source/Developer/TraceServices/Private/Model/AllocationsQuery.cpp
+++ b/Engine/Source/Developer/TraceServices/Private/Model/AllocationsQuery.cpp
@@ -120,6 +120,9 @@ void FAllocationsQuery::Run()
 		{
 			UE_LOG(LogTraceServices, Log, TEXT("[MemAlloc] Enqueue %u live allocs..."), NumLiveAllocs);
 			TotalAllocationCount += NumLiveAllocs;
+
+			QueryCallstacks(LiveAllocsResult);
+
 			Results.Enqueue(LiveAllocsResult);
 		}
 		else
@@ -156,14 +159,7 @@ void FAllocationsQuery::Run()
 			const uint32 NumAllocs = static_cast<uint32>(CellResult->Items.Num());
 			if (NumAllocs != 0)
 			{
-				for (auto Item : CellResult->Items)
-				{
-					if (!Item->Callstack)
-					{
-						Item->Callstack = CallstacksProvider.GetCallstack(Item->Owner);
-						// check(Item->Callstack);
-					}
-				}
+				QueryCallstacks(CellResult);
 
 				UE_LOG(LogTraceServices, Log, TEXT("[MemAlloc] Enqueue %u allocs..."), NumAllocs);
 				TotalAllocationCount += NumAllocs;
@@ -374,4 +370,19 @@ void FAllocationsQuery::QueryLiveAllocs(TArray<const FAllocationItem*>& OutAlloc
 
 ////////////////////////////////////////////////////////////////////////////////////////////////////
 
+void FAllocationsQuery::QueryCallstacks(FAllocationsImpl* Result) const
+{
+	for (auto Item : Result->Items)
+	{
+		// Callstacks could have been resolved in previous queries, 
+		// check before querying again
+		if (!Item->Callstack)
+		{
+			Item->Callstack = CallstacksProvider.GetCallstack(Item->Owner);
+		}
+	}
+}
+
+////////////////////////////////////////////////////////////////////////////////////////////////////
+
 } // namespace TraceServices
diff --git a/Engine/Source/Developer/TraceServices/Private/Model/AllocationsQuery.h b/Engine/Source/Developer/TraceServices/Private/Model/AllocationsQuery.h
--- a/Engine/Source/Developer/TraceServices/Private/Model/AllocationsQuery.h
+++ b/Engine/Source/Developer/TraceServices/Private/Model/AllocationsQuery.h
@@ -34,6 +34,7 @@ public:
 private:
 	void Run();
 	void QueryLiveAllocs(TArray<const FAllocationItem*>& OutAllocs) const;
+	void QueryCallstacks(FAllocationsImpl* Result) const;
 
 private:
 	const FAllocationsProvider& AllocationsProvider;
