Fixed backwards compatibility with implicitly registered threads. CPU events could arrive on a thread before the game thread event was received.

diff --git a/Engine/Source/Developer/TraceServices/Private/Model/Threads.cpp b/Engine/Source/Developer/TraceServices/Private/Model/Threads.cpp
--- a/Engine/Source/Developer/TraceServices/Private/Model/Threads.cpp
+++ b/Engine/Source/Developer/TraceServices/Private/Model/Threads.cpp
@@ -27,17 +27,8 @@ FThreadProvider::~FThreadProvider()
 
 void FThreadProvider::AddGameThread(uint32 Id)
 {
-	Session.WriteAccessCheck();
-
-	check(!ThreadMap.Contains(Id));
-	FThreadInfoInternal* ThreadInfo = new FThreadInfoInternal();
-	ThreadInfo->Id = Id;
-	ThreadInfo->PrioritySortOrder = -2;
-	ThreadInfo->Name = Session.StoreString(*FName(NAME_GameThread).GetPlainNameString());
-	ThreadInfo->FallbackSortOrder = SortedThreads.Num();
-	SortedThreads.Add(ThreadInfo);
-	ThreadMap.Add(Id, ThreadInfo);
-	++ModCount;
+	const FString Name = *FName(NAME_GameThread).GetPlainNameString();
+	AddThread(Id, *Name, EThreadPriority(-2));
 }
 
 void FThreadProvider::AddThread(uint32 Id, const TCHAR* Name, EThreadPriority Priority)
