[Insights] Register analyzer to recieve realloc free events.

#rnx

diff --git a/Engine/Source/Developer/TraceServices/Private/Analyzers/AllocationsAnalysis.cpp b/Engine/Source/Developer/TraceServices/Private/Analyzers/AllocationsAnalysis.cpp
--- a/Engine/Source/Developer/TraceServices/Private/Analyzers/AllocationsAnalysis.cpp
+++ b/Engine/Source/Developer/TraceServices/Private/Analyzers/AllocationsAnalysis.cpp
@@ -33,6 +33,7 @@ void FAllocationsAnalyzer::OnAnalysisBegin(const FOnAnalysisContext& Context)
 	Builder.RouteEvent(RouteId_Alloc,        "Memory", "Alloc");
 	Builder.RouteEvent(RouteId_Free,         "Memory", "Free");
 	Builder.RouteEvent(RouteId_ReallocAlloc, "Memory", "ReallocAlloc");
+	Builder.RouteEvent(RouteId_ReallocFree,	 "Memory", "ReallocFree");
 	Builder.RouteEvent(RouteId_Marker,       "Memory", "Marker");
 	Builder.RouteEvent(RouteId_TagSpec,      "Memory", "TagSpec");
 
