[Insights] Added log messages for errors detected by AllocationsProvider.

diff --git a/Engine/Source/Developer/TraceServices/Private/Model/AllocationsProvider.cpp b/Engine/Source/Developer/TraceServices/Private/Model/AllocationsProvider.cpp
--- a/Engine/Source/Developer/TraceServices/Private/Model/AllocationsProvider.cpp
+++ b/Engine/Source/Developer/TraceServices/Private/Model/AllocationsProvider.cpp
@@ -3,6 +3,7 @@
 #include "AllocationsProvider.h"
 
 #include "AllocationsQuery.h"
+#include "Common/Utils.h"
 #include "Containers/ArrayView.h"
 #include "SbTree.h"
 #include "TraceServices/Containers/Allocators.h"
@@ -13,6 +14,8 @@
 namespace TraceServices
 {
 
+constexpr uint32 MaxLogMessagesPerErrorType = 100;
+
 ////////////////////////////////////////////////////////////////////////////////////////////////////
 // FAllocationsProviderLock
 ////////////////////////////////////////////////////////////////////////////////////////////////////
@@ -88,7 +91,10 @@ void FTagTracker::AddTagSpec(uint32 InTag, uint32 InParentTag, const TCHAR* InDi
 	else
 	{
 		++NumErrors;
-		//TODO: error message
+		if (NumErrors <= MaxLogMessagesPerErrorType)
+		{
+			UE_LOG(LogTraceServices, Error, TEXT("[MemAlloc] Tag with id %u (ParentTag=%u, Display=%s) already added!"), InTag, InParentTag, InDisplay);
+		}
 	}
 }
 
@@ -114,7 +120,10 @@ void FTagTracker::PopTag(uint32 InThreadId, uint8 InTracker)
 	else
 	{
 		++NumErrors;
-		//TODO: error message
+		if (NumErrors <= MaxLogMessagesPerErrorType)
+		{
+			UE_LOG(LogTraceServices, Error, TEXT("[MemAlloc] Tag stack on Thread %u (Tracker=%u) is already empty!"), InThreadId, InTracker);
+		}
 	}
 }
 
@@ -161,7 +170,10 @@ void FTagTracker::PopRealloc(uint32 InThreadId, uint8 InTracker)
 	else
 	{
 		++NumErrors;
-		//TODO: error message
+		if (NumErrors <= MaxLogMessagesPerErrorType)
+		{
+			UE_LOG(LogTraceServices, Error, TEXT("[MemAlloc] Realloc stack on Thread %u (Tracker=%u) is already empty!"), InThreadId, InTracker);
+		}
 	}
 }
 
@@ -376,6 +388,17 @@ void FAllocationsProvider::EditRemoveCore(double Time, uint64 Owner, uint64 Base
 
 ////////////////////////////////////////////////////////////////////////////////////////////////////
 
+#define MEMALLOC_DEBUG_WATCH 0
+
+#if MEMALLOC_DEBUG_WATCH
+static uint64 GWatchAddresses[] =
+{
+	0x0, // add here the addresses to watch
+};
+#endif // MEMALLOC_DEBUG_WATCH
+
+////////////////////////////////////////////////////////////////////////////////////////////////////
+
 void FAllocationsProvider::EditAlloc(double Time, uint64 Owner, uint64 Address, uint32 InSize, uint8 InAlignmentAndSizeLower, uint32 ThreadId, uint8 Tracker)
 {
 	Lock.WriteAccessCheck();
@@ -392,21 +415,35 @@ void FAllocationsProvider::EditAlloc(double Time, uint64 Owner, uint64 Address,
 
 	SbTree->SetTimeForEvent(EventIndex, Time);
 
+	const uint8 SizeLowerMask = ((1 << SizeShift) - 1);
+	const uint8 AlignmentMask = ~SizeLowerMask;
+	const uint64 Size = (static_cast<uint64>(InSize) << SizeShift) | static_cast<uint64>(InAlignmentAndSizeLower & SizeLowerMask);
+
+	const uint32 Tag = TagTracker.GetCurrentTag(ThreadId, Tracker);
+
 	FAllocationItem* AllocationPtr = LiveAllocs.Find(Address);
 	if (AllocationPtr)
 	{
 		++AllocErrors;
+		if (AllocErrors <= MaxLogMessagesPerErrorType)
+		{
+			UE_LOG(LogTraceServices, Error, TEXT("[MemAlloc] Invalid ALLOC event (Address=0x%llX, Size=%llu, Tag=%u, Time=%f)!"), Address, Size, Tag, Time);
+		}
 	}
 	else
 	{
-		AdvanceTimelines(Time);
-
-		const uint8 SizeLowerMask = ((1 << SizeShift) - 1);
-		const uint8 AlignmentMask = ~SizeLowerMask;
-
-		const uint64 Size = (static_cast<uint64>(InSize) << SizeShift) | static_cast<uint64>(InAlignmentAndSizeLower & SizeLowerMask);
+#if MEMALLOC_DEBUG_WATCH
+		for (int32 AddrIndex = 0; AddrIndex < UE_ARRAY_COUNT(GWatchAddresses); ++AddrIndex)
+		{
+			if (GWatchAddresses[AddrIndex] == Address)
+			{
+				UE_LOG(LogTraceServices, Warning, TEXT("[MemAlloc] Alloc 0x%llX : Size=%llu, Tag=%u, Time=%f"), Address, Size, Tag, Time);
+				break;
+			}
+		}
+#endif // MEMALLOC_DEBUG_WATCH
 
-		const uint32 Tag = TagTracker.GetCurrentTag(ThreadId, Tracker);
+		AdvanceTimelines(Time);
 
 		FAllocationItem Allocation =
 		{
@@ -466,9 +503,24 @@ void FAllocationsProvider::EditFree(double Time, uint64 Address)
 	if (!AllocationPtr)
 	{
 		++FreeErrors;
+		if (FreeErrors <= MaxLogMessagesPerErrorType)
+		{
+			UE_LOG(LogTraceServices, Error, TEXT("[MemAlloc] Invalid FREE event (Address=0x%llX, Time=%f)!"), Address, Time);
+		}
 	}
 	else
 	{
+#if MEMALLOC_DEBUG_WATCH
+		for (int32 AddrIndex = 0; AddrIndex < UE_ARRAY_COUNT(GWatchAddresses); ++AddrIndex)
+		{
+			if (GWatchAddresses[AddrIndex] == Address)
+			{
+				UE_LOG(LogTraceServices, Warning, TEXT("[MemAlloc] Free 0x%llX : Time=%f"), Address, Time);
+				break;
+			}
+		}
+#endif // MEMALLOC_DEBUG_WATCH
+
 		AdvanceTimelines(Time);
 
 		check(EventIndex > AllocationPtr->StartEventIndex);
