[Insights]
  - Fixed AllocationAnalysis to increase the Session Time.
  - Workaround for non monotonic time markers in AllocationAnalysis.

diff --git a/Engine/Source/Developer/TraceServices/Private/Analyzers/AllocationsAnalysis.cpp b/Engine/Source/Developer/TraceServices/Private/Analyzers/AllocationsAnalysis.cpp
--- a/Engine/Source/Developer/TraceServices/Private/Analyzers/AllocationsAnalysis.cpp
+++ b/Engine/Source/Developer/TraceServices/Private/Analyzers/AllocationsAnalysis.cpp
@@ -12,6 +12,10 @@ namespace TraceServices
 FAllocationsAnalyzer::FAllocationsAnalyzer(IAnalysisSession& InSession, FAllocationsProvider& InAllocationsProvider)
 	: Session(InSession)
 	, AllocationsProvider(InAllocationsProvider)
+	, BaseCycle(0)
+	, MarkerPeriod(0)
+	, LastMarkerCycle(0)
+	, LastMarkerSeconds(0.0)
 {
 }
 
@@ -33,7 +37,7 @@ void FAllocationsAnalyzer::OnAnalysisBegin(const FOnAnalysisContext& Context)
 	Builder.RouteEvent(RouteId_Alloc,        "Memory", "Alloc");
 	Builder.RouteEvent(RouteId_Free,         "Memory", "Free");
 	Builder.RouteEvent(RouteId_ReallocAlloc, "Memory", "ReallocAlloc");
-	Builder.RouteEvent(RouteId_ReallocFree,	 "Memory", "ReallocFree");
+	Builder.RouteEvent(RouteId_ReallocFree,  "Memory", "ReallocFree");
 	Builder.RouteEvent(RouteId_Marker,       "Memory", "Marker");
 	Builder.RouteEvent(RouteId_TagSpec,      "Memory", "TagSpec");
 
@@ -48,6 +52,10 @@ void FAllocationsAnalyzer::OnAnalysisEnd()
 	{
 		FAnalysisSessionEditScope _(Session);
 		SessionTime = Session.GetDurationSeconds();
+		if (LastMarkerSeconds > SessionTime)
+		{
+			Session.UpdateDurationSeconds(LastMarkerSeconds);
+		}
 	}
 	const double Time = FMath::Max(SessionTime, LastMarkerSeconds);
 
@@ -128,8 +136,25 @@ bool FAllocationsAnalyzer::OnEvent(uint16 RouteId, EStyle Style, const FOnEventC
 		}
 		case RouteId_Marker:
 		{
-			const uint64 LastCycle = BaseCycle + EventData.GetValue<uint32>("Cycle");
-			LastMarkerSeconds = Context.EventTime.AsSeconds(LastCycle);
+			const uint64 Cycle = BaseCycle + EventData.GetValue<uint32>("Cycle");
+			if (ensure(Cycle >= LastMarkerCycle))
+			{
+				const double Seconds = Context.EventTime.AsSeconds(Cycle);
+				check(Seconds >= LastMarkerSeconds);
+				if (ensure(Seconds - LastMarkerSeconds < 60.0))
+				{
+					LastMarkerCycle = Cycle;
+					LastMarkerSeconds = Seconds;
+					{
+						FAnalysisSessionEditScope _(Session);
+						double SessionTime = Session.GetDurationSeconds();
+						if (LastMarkerSeconds > SessionTime)
+						{
+							Session.UpdateDurationSeconds(LastMarkerSeconds);
+						}
+					}
+				}
+			}
 			break;
 		}
 		case RouteId_TagSpec:
diff --git a/Engine/Source/Developer/TraceServices/Private/Analyzers/AllocationsAnalysis.h b/Engine/Source/Developer/TraceServices/Private/Analyzers/AllocationsAnalysis.h
--- a/Engine/Source/Developer/TraceServices/Private/Analyzers/AllocationsAnalysis.h
+++ b/Engine/Source/Developer/TraceServices/Private/Analyzers/AllocationsAnalysis.h
@@ -41,6 +41,7 @@ private:
 	FAllocationsProvider& AllocationsProvider;
 	uint64 BaseCycle;
 	uint32 MarkerPeriod;
+	uint64 LastMarkerCycle;
 	double LastMarkerSeconds;
 };
 
