TaskTrace: TaskGraph traces and UnrealInsights support.
added TaskGraph specific traces as a separate channel that needs to be enabled by `-trace=tasks`. They include timing info about all task states and waiting for tasks. UnrealInsights support is limited for now. All info is available, a number of task specific counters are added, but dependency and waiting visualisation is not done yet. Available counters are:
* task latency
* number of scheduled tasks
* number of executing tasks

diff --git a/Engine/Source/Developer/TraceInsights/Private/Insights/ViewModels/ThreadTimingTrack.cpp b/Engine/Source/Developer/TraceInsights/Private/Insights/ViewModels/ThreadTimingTrack.cpp
--- a/Engine/Source/Developer/TraceInsights/Private/Insights/ViewModels/ThreadTimingTrack.cpp
+++ b/Engine/Source/Developer/TraceInsights/Private/Insights/ViewModels/ThreadTimingTrack.cpp
@@ -8,6 +8,7 @@
 #include "HAL/PlatformApplicationMisc.h"
 #include "Serialization/MemoryReader.h"
 #include "Styling/SlateBrush.h"
+#include "TraceServices/Model/TasksProfiler.h"
 
 // Insights
 #include "Insights/Common/PaintUtils.h"
@@ -1004,6 +1005,73 @@ void FThreadTimingTrack::InitTooltip(FTooltipDrawState& InOutTooltip, const ITim
 		{
 			AppendMetadataToTooltip(InOutTooltip, Metadata);
 		}
+
+		// tasks
+
+		const TraceServices::ITasksProvider* TasksProvider = TraceServices::ReadTasksProvider(*Session.Get());
+
+		auto AddTaskInfo = [&InOutTooltip, this](const TraceServices::FTaskInfo& Task)
+		{
+			InOutTooltip.AddTextLine(TEXT("-------- TaskInfo --------"), FLinearColor::Green);
+			InOutTooltip.AddNameValueTextLine(TEXT("TaskId:"), FString::Printf(TEXT("%d"), Task.Id));
+			InOutTooltip.AddNameValueTextLine(TEXT("Task Debug Name:"), FString::Printf(TEXT("%s"), Task.DebugName));
+			InOutTooltip.AddNameValueTextLine(TEXT("Tracked:"), FString::Printf(TEXT("%d"), Task.bTracked));
+			InOutTooltip.AddNameValueTextLine(TEXT("Created:"), FString::Printf(TEXT("%f"), Task.CreatedTimestamp));
+			InOutTooltip.AddNameValueTextLine(TEXT("Launched:"), FString::Printf(TEXT("%f (+%s)"), Task.LaunchedTimestamp, *TimeUtils::FormatTimeAuto(Task.LaunchedTimestamp - Task.CreatedTimestamp)));
+			InOutTooltip.AddNameValueTextLine(TEXT("Started:"), FString::Printf(TEXT("%f (+%s) on thread \"%s\""), Task.StartedTimestamp, *TimeUtils::FormatTimeAuto(Task.StartedTimestamp - Task.LaunchedTimestamp), *SharedState.GetCpuTrack(Task.StartedThreadId)->GetName()));
+			if (Task.FinishedTimestamp != TraceServices::FTaskInfo::InvalidTimestamp)
+			{
+				InOutTooltip.AddNameValueTextLine(TEXT("Finished:"), FString::Printf(TEXT("%f (+%s)"), Task.FinishedTimestamp, *TimeUtils::FormatTimeAuto(Task.FinishedTimestamp - Task.StartedTimestamp)));
+
+				if (Task.CompletedTimestamp != TraceServices::FTaskInfo::InvalidTimestamp)
+				{
+					InOutTooltip.AddNameValueTextLine(TEXT("Completed:"), FString::Printf(TEXT("%f (+%s)"), Task.FinishedTimestamp, *TimeUtils::FormatTimeAuto(Task.CompletedTimestamp - Task.FinishedTimestamp)));
+				}
+			}
+		};
+
+		do // info about a task
+		{
+			const TraceServices::FTaskInfo* Task = TasksProvider->TryGetTask(ThreadId, TooltipEvent.GetStartTime());
+			if (Task == nullptr)
+			{
+				break;
+			}
+
+			if (Task->FinishedTimestamp < TooltipEvent.GetEndTime())
+			{
+				break;
+			}
+
+			AddTaskInfo(*Task);
+		} while (false);
+
+		do // info about blocking
+		{
+			if (TimerName.StartsWith(TEXT("WaitUntilTasksComplete")) || TimerName.StartsWith(TEXT("GameThreadWaitForTask")))
+			{
+				const TraceServices::FWaitingForTasks* Waiting = TasksProvider->TryGetWaiting(ThreadId, TooltipEvent.GetStartTime());
+				if (Waiting == nullptr)
+				{
+					break;
+				}
+
+				InOutTooltip.AddTextLine(TEXT("-------- Wating for tasks --------"), FLinearColor::Red);
+				FString TaskIdsStr = FString::JoinBy(Waiting->Tasks, TEXT(", "), [](TaskTrace::FId TaskId) { return FString::FromInt(TaskId); });
+				InOutTooltip.AddNameValueTextLine(TEXT("Tasks:"), FString::Printf(TEXT("[%s]"), *TaskIdsStr));
+				InOutTooltip.AddNameValueTextLine(TEXT("Started waiting:"), FString::Printf(TEXT("%f"), Waiting->StartedTimestamp));
+				InOutTooltip.AddNameValueTextLine(TEXT("Finished waiting:"), FString::Printf(TEXT("%f (+%s)"), Waiting->FinishedTimestamp, *TimeUtils::FormatTimeAuto(Waiting->FinishedTimestamp - Waiting->StartedTimestamp)));
+
+				for (TaskTrace::FId TaskId : Waiting->Tasks)
+				{
+					const TraceServices::FTaskInfo* Task = TasksProvider->TryGetTask(TaskId);
+					if (Task != nullptr)
+					{
+						AddTaskInfo(*Task);
+					}
+				}
+			}
+		} while (false);
 	}
 
 	InOutTooltip.UpdateLayout();
diff --git a/Engine/Source/Developer/TraceInsights/Private/Insights/ViewModels/ThreadTimingTrack.h b/Engine/Source/Developer/TraceInsights/Private/Insights/ViewModels/ThreadTimingTrack.h
--- a/Engine/Source/Developer/TraceInsights/Private/Insights/ViewModels/ThreadTimingTrack.h
+++ b/Engine/Source/Developer/TraceInsights/Private/Insights/ViewModels/ThreadTimingTrack.h
@@ -129,6 +129,7 @@ public:
 		, GroupName(InGroupName)
 		, TimelineIndex(InTimelineIndex)
 		, ThreadId(InThreadId)
+		, SharedState(InSharedState)
 	{
 	}
 
@@ -180,6 +181,8 @@ private:
 	uint32 TimelineIndex;
 	uint32 ThreadId;
 
+	FThreadTimingSharedState& SharedState;
+
 	TSharedPtr<Insights::FFilterConfigurator> FilterConfigurator;
 	FDelegateHandle OnFilterChangesCommitedHandle;
 
diff --git a/Engine/Source/Developer/TraceServices/Private/Analyzers/TasksAnalysis.cpp b/Engine/Source/Developer/TraceServices/Private/Analyzers/TasksAnalysis.cpp
new file mode 100644
--- /dev/null
+++ b/Engine/Source/Developer/TraceServices/Private/Analyzers/TasksAnalysis.cpp
@@ -0,0 +1,136 @@
+// Copyright Epic Games, Inc. All Rights Reserved.
+
+#include "TasksAnalysis.h"
+#include "AnalysisServicePrivate.h"
+#include "Model/TasksProfilerPrivate.h"
+#include "Containers/UnrealString.h"
+
+namespace TraceServices
+{
+	FTasksAnalyzer::FTasksAnalyzer(IAnalysisSession& InSession, FTasksProvider& InTasksProvider)
+		: Session(InSession)
+		, TasksProvider(InTasksProvider)
+	{
+	}
+
+	void FTasksAnalyzer::OnAnalysisBegin(const FOnAnalysisContext& Context)
+	{
+		auto& Builder = Context.InterfaceBuilder;
+
+		Builder.RouteEvent(RouteId_Created, "TaskTrace", "Created");
+		Builder.RouteEvent(RouteId_Launched, "TaskTrace", "Launched");
+		Builder.RouteEvent(RouteId_Scheduled, "TaskTrace", "Scheduled");
+		Builder.RouteEvent(RouteId_SubsequentAdded, "TaskTrace", "SubsequentAdded");
+		Builder.RouteEvent(RouteId_Started, "TaskTrace", "Started");
+		Builder.RouteEvent(RouteId_NestedAdded, "TaskTrace", "NestedAdded");
+		Builder.RouteEvent(RouteId_Finished, "TaskTrace", "Finished");
+		Builder.RouteEvent(RouteId_Completed, "TaskTrace", "Completed");
+
+		Builder.RouteEvent(RouteId_WaitingStarted, "TaskTrace", "WaitingStarted");
+		Builder.RouteEvent(RouteId_WaitingFinished, "TaskTrace", "WaitingFinished");
+	}
+
+	bool FTasksAnalyzer::OnEvent(uint16 RouteId, EStyle Style, const FOnEventContext& Context)
+	{
+		FAnalysisSessionEditScope _(Session);
+	
+		// the protocol must match TaskTrace.cpp
+
+		const auto& EventData = Context.EventData;
+		uint32 ThreadId = Context.ThreadInfo.GetId();
+		switch (RouteId)
+		{
+		case RouterId_Init:
+		{
+			uint32 Version = EventData.GetValue<uint32>("Version");
+			TasksProvider.Init(Version);
+			break;
+		}
+		case RouteId_Created:
+		{
+			double Timestamp = Context.EventTime.AsSeconds(EventData.GetValue<uint64>("Timestamp"));
+			TaskTrace::FId TaskId = EventData.GetValue<TaskTrace::FId>("TaskId");
+			TasksProvider.TaskCreated(TaskId, Timestamp, ThreadId);
+			break;
+		}
+		case RouteId_Launched:
+		{
+			double Timestamp = Context.EventTime.AsSeconds(EventData.GetValue<uint64>("Timestamp"));
+			TaskTrace::FId TaskId = EventData.GetValue<TaskTrace::FId>("TaskId");
+
+			FString DebugNameStr;
+			const TCHAR* DebugName = nullptr;
+			if (EventData.GetString("DebugName", DebugNameStr) && DebugNameStr.Len() != 0)
+			{
+				DebugName = Session.StoreString(*DebugNameStr);
+			}
+
+			bool bTracked = EventData.GetValue<bool>("Tracked");
+			int32 ThreadToExecuteOn = EventData.GetValue<int32>("ThreadToExecuteOn");
+
+			TasksProvider.TaskLaunched(TaskId, DebugName, bTracked, ThreadToExecuteOn, Timestamp, ThreadId);
+
+			break;
+		}
+		case RouteId_Scheduled:
+		{
+			double Timestamp = Context.EventTime.AsSeconds(EventData.GetValue<uint64>("Timestamp"));
+			uint64 TaskId = EventData.GetValue<TaskTrace::FId>("TaskId");
+			TasksProvider.TaskScheduled(TaskId, Timestamp, ThreadId);
+			break;
+		}
+		case RouteId_SubsequentAdded:
+		{
+			double Timestamp = Context.EventTime.AsSeconds(EventData.GetValue<uint64>("Timestamp"));
+			TaskTrace::FId TaskId = EventData.GetValue<TaskTrace::FId>("TaskId");
+			TaskTrace::FId SubsequentId = EventData.GetValue<TaskTrace::FId>("SubsequentId");
+			TasksProvider.SubsequentAdded(TaskId, SubsequentId, Timestamp, ThreadId);
+			break;
+		}
+		case RouteId_Started:
+		{
+			double Timestamp = Context.EventTime.AsSeconds(EventData.GetValue<uint64>("Timestamp"));
+			TaskTrace::FId TaskId = EventData.GetValue<TaskTrace::FId>("TaskId");
+			TasksProvider.TaskStarted(TaskId, Timestamp, ThreadId);
+			break;
+		}
+		case RouteId_NestedAdded:
+		{
+			double Timestamp = Context.EventTime.AsSeconds(EventData.GetValue<uint64>("Timestamp"));
+			TaskTrace::FId TaskId = EventData.GetValue<TaskTrace::FId>("TaskId");
+			TaskTrace::FId NestedId = EventData.GetValue<TaskTrace::FId>("NestedId");
+			TasksProvider.NestedAdded(TaskId, NestedId, Timestamp, ThreadId);
+			break;
+		}
+		case RouteId_Finished:
+		{
+			double Timestamp = Context.EventTime.AsSeconds(EventData.GetValue<uint64>("Timestamp"));
+			TaskTrace::FId TaskId = EventData.GetValue<TaskTrace::FId>("TaskId");
+			TasksProvider.TaskFinished(TaskId, Timestamp);
+			break;
+		}
+		case RouteId_Completed:
+		{
+			double Timestamp = Context.EventTime.AsSeconds(EventData.GetValue<uint64>("Timestamp"));
+			TaskTrace::FId TaskId = EventData.GetValue<TaskTrace::FId>("TaskId");
+			TasksProvider.TaskCompleted(TaskId, Timestamp, ThreadId);
+			break;
+		}
+		case RouteId_WaitingStarted:
+		{
+			double Timestamp = Context.EventTime.AsSeconds(EventData.GetValue<uint64>("Timestamp"));
+			const TArrayReader<TaskTrace::FId>& Tasks = EventData.GetArray<TaskTrace::FId>("Tasks");
+			TasksProvider.WaitingStarted({ Tasks.GetData(), static_cast<int32>(Tasks.Num()) }, Timestamp, ThreadId);
+			break;
+		}
+		case RouteId_WaitingFinished:
+		{
+			double Timestamp = Context.EventTime.AsSeconds(EventData.GetValue<uint64>("Timestamp"));
+			TasksProvider.WaitingFinished(Timestamp, ThreadId);
+			break;
+		}
+		}
+
+		return true;
+	}
+}
\ No newline at end of file
diff --git a/Engine/Source/Developer/TraceServices/Private/Analyzers/TasksAnalysis.h b/Engine/Source/Developer/TraceServices/Private/Analyzers/TasksAnalysis.h
new file mode 100644
--- /dev/null
+++ b/Engine/Source/Developer/TraceServices/Private/Analyzers/TasksAnalysis.h
@@ -0,0 +1,40 @@
+// Copyright Epic Games, Inc. All Rights Reserved.
+
+#pragma once
+
+#include "Trace/Analyzer.h"
+
+namespace TraceServices
+{
+	class IAnalysisSession;
+	class FTasksProvider;
+
+	// deserialises tasks traces and feeds them to FTasksProvider
+	class FTasksAnalyzer : public UE::Trace::IAnalyzer
+	{
+	public:
+		FTasksAnalyzer(IAnalysisSession& Session, FTasksProvider& TasksProvider);
+		virtual void OnAnalysisBegin(const FOnAnalysisContext& Context) override;
+		virtual bool OnEvent(uint16 RouteId, EStyle Style, const FOnEventContext& Context) override;
+
+	private:
+		enum : uint16
+		{
+			RouterId_Init,
+			RouteId_Created,
+			RouteId_Launched,
+			RouteId_Scheduled,
+			RouteId_SubsequentAdded,
+			RouteId_Started,
+			RouteId_NestedAdded,
+			RouteId_Finished,
+			RouteId_Completed,
+
+			RouteId_WaitingStarted,
+			RouteId_WaitingFinished,
+		};
+
+		IAnalysisSession& Session;
+		FTasksProvider& TasksProvider;
+	};
+}
diff --git a/Engine/Source/Developer/TraceServices/Private/Model/TasksProfiler.cpp b/Engine/Source/Developer/TraceServices/Private/Model/TasksProfiler.cpp
new file mode 100644
--- /dev/null
+++ b/Engine/Source/Developer/TraceServices/Private/Model/TasksProfiler.cpp
@@ -0,0 +1,389 @@
+// Copyright Epic Games, Inc. All Rights Reserved.
+
+#include "TraceServices/Model/TasksProfiler.h"
+#include "Model/TasksProfilerPrivate.h"
+#include "AnalysisServicePrivate.h"
+#include "Misc/ScopeLock.h"
+#include "Misc/Timespan.h"
+#include "Algo/BinarySearch.h"
+#include "Common/Utils.h"
+
+namespace TraceServices
+{
+	FTasksProvider::FTasksProvider(IAnalysisSession& InSession)
+		: Session(InSession)
+		, CounterProvider(EditCounterProvider(Session))
+	{
+#if 0
+		////////////////////////////////
+		// tests
+
+		FAnalysisSessionEditScope _(Session);
+
+		const int32 ValidThreadId = 42;
+		const int32 InvalidThreadId = 0;
+		const double InvalidTimestamp = 0;
+
+		TArray64<TaskTrace::FId>& Thread = ExecutionThreads.FindOrAdd(ValidThreadId);
+
+		auto MockTaskExecution = [this, Thread = &Thread](TaskTrace::FId TaskId, double StartedTimestamp, double FinishedTimestamp) -> FTaskInfo&
+		{
+			FTaskInfo& Task = GetOrCreateTask(TaskId);
+			Task.Id = TaskId;
+			Task.StartedTimestamp = StartedTimestamp;
+			Task.FinishedTimestamp = FinishedTimestamp;
+			Thread->Add(Task.Id);
+			return Task;
+		};
+
+		FTaskInfo& Task1 = MockTaskExecution(0, 5, 10);
+		FTaskInfo& Task2 = MockTaskExecution(1, 15, 20);
+		
+		check(GetTask(InvalidThreadId, InvalidTimestamp) == nullptr);
+		check(GetTask(ValidThreadId, InvalidTimestamp) == nullptr);
+		check(GetTask(ValidThreadId, 5) == &Task1);
+		check(GetTask(ValidThreadId, 7) == &Task1);
+		check(GetTask(ValidThreadId, 10) == nullptr);
+		check(GetTask(ValidThreadId, 12) == nullptr);
+		check(GetTask(ValidThreadId, 17) == &Task2);
+		check(GetTask(ValidThreadId, 22) == nullptr);
+
+		// reset
+		FirstTaskId = TaskTrace::InvalidId;
+		ExecutionThreads.Empty();
+		Tasks.Empty();
+#endif
+	}
+
+	void FTasksProvider::CreateCounters()
+	{
+		FAnalysisSessionEditScope _(Session);
+
+		WaitingForPrerequisitesTasksCounter = CounterProvider.CreateCounter();
+		WaitingForPrerequisitesTasksCounter->SetName(TEXT("Tasks::WaitingForPrerequisitesTasks"));
+		WaitingForPrerequisitesTasksCounter->SetDescription(TEXT("Tasks: the number of tasks waiting for prerequisites (blocked by dependency)"));
+		WaitingForPrerequisitesTasksCounter->SetIsFloatingPoint(false);
+
+		TaskLatencyCounter = CounterProvider.CreateCounter();
+		TaskLatencyCounter->SetName(TEXT("Tasks::TaskLatency"));
+		TaskLatencyCounter->SetDescription(TEXT("Tasks: tasks latency - the time from scheduling to execution start"));
+		TaskLatencyCounter->SetIsFloatingPoint(true);
+
+		ScheduledTasksCounter = CounterProvider.CreateCounter();
+		ScheduledTasksCounter->SetName(TEXT("Tasks::ScheduledTasks"));
+		ScheduledTasksCounter->SetDescription(TEXT("Tasks: number of scheduled tasks excluding named threads (the size of the queue)"));
+		ScheduledTasksCounter->SetIsFloatingPoint(false);
+
+		NamedThreadsScheduledTasksCounter = CounterProvider.CreateCounter();
+		NamedThreadsScheduledTasksCounter->SetName(TEXT("Tasks::NamedThreadsScheduledTasks"));
+		NamedThreadsScheduledTasksCounter->SetDescription(TEXT("Tasks: number of scheduled tasks for named threads"));
+		NamedThreadsScheduledTasksCounter->SetIsFloatingPoint(false);
+
+		RunningTasksCounter = CounterProvider.CreateCounter();
+		RunningTasksCounter->SetName(TEXT("Tasks::RunningTasks"));
+		RunningTasksCounter->SetDescription(TEXT("Tasks: level of parallelism - the number of tasks being executed"));
+		RunningTasksCounter->SetIsFloatingPoint(false);
+	}
+
+	void FTasksProvider::Init(uint32 InVersion)
+	{
+		Version = InVersion;
+
+		CreateCounters();
+	}
+
+	void FTasksProvider::TaskCreated(TaskTrace::FId TaskId, double Timestamp, uint32 ThreadId)
+	{
+		UE_LOG(LogTraceServices, Verbose, TEXT("TaskCreated(TaskId: %d, Timestamp %.6f)"), TaskId, Timestamp);
+
+		InitTaskIdToIndexConversion(TaskId);
+
+		FTaskInfo* Task = TryGetOrCreateTask(TaskId);
+		if (Task == nullptr)
+		{
+			UE_LOG(LogTraceServices, Log, TEXT("TaskCreated(TaskId %d, Timestamp %d) skipped"), TaskId, Timestamp);
+			return;
+		}
+
+		checkf(Task->CreatedTimestamp == FTaskInfo::InvalidTimestamp, TEXT("%d"), TaskId);
+
+		Task->Id = TaskId;
+		Task->CreatedTimestamp = Timestamp;
+		Task->CreatedThreadId = ThreadId;
+	}
+
+	void FTasksProvider::TaskLaunched(TaskTrace::FId TaskId, const TCHAR* DebugName, bool bTracked, int32 ThreadToExecuteOn, double Timestamp, uint32 ThreadId)
+	{
+		UE_LOG(LogTraceServices, Verbose, TEXT("TaskLaunched(TaskId: %d, DebugName: %s, bTracked: %d, Timestamp %.6f)"), TaskId, DebugName, bTracked, Timestamp);
+
+		InitTaskIdToIndexConversion(TaskId);
+
+		FTaskInfo* Task = TryGetOrCreateTask(TaskId);
+		if (Task == nullptr)
+		{
+			UE_LOG(LogTraceServices, Log, TEXT("TaskLaunched(TaskId %d, DebugName %s, bTracked %d, Timestamp %.6f) skipped"), TaskId, DebugName, bTracked, Timestamp);
+			return;
+		}
+
+		checkf(Task->LaunchedTimestamp == FTaskInfo::InvalidTimestamp, TEXT("%d"), TaskId);
+			
+		if (Task->Id == TaskTrace::InvalidId) // created and launched in one go
+		{
+			Task->Id = TaskId;
+			Task->CreatedTimestamp = Timestamp;
+		}
+
+		Task->DebugName = DebugName;
+		Task->bTracked = bTracked;
+		Task->ThreadToExecuteOn = ThreadToExecuteOn;
+		Task->LaunchedTimestamp = Timestamp;
+		Task->LaunchedThreadId = ThreadId;
+
+		WaitingForPrerequisitesTasksCounter->SetValue(Timestamp, ++WaitingForPrerequisitesTasksNum);
+	}
+
+	void FTasksProvider::TaskScheduled(TaskTrace::FId TaskId, double Timestamp, uint32 ThreadId)
+	{
+		if (!TryRegisterEvent(TEXT("TaskScheduled"), TaskId, &FTaskInfo::ScheduledTimestamp, Timestamp, &FTaskInfo::ScheduledThreadId, ThreadId))
+		{
+			return;
+		}
+
+		WaitingForPrerequisitesTasksCounter->SetValue(Timestamp, --WaitingForPrerequisitesTasksNum);
+		if (IsNamedThread(TryGetTask(TaskId)->ThreadToExecuteOn))
+		{
+			NamedThreadsScheduledTasksCounter->SetValue(Timestamp, ++ScheduledTasksNum);
+		}
+		else
+		{
+			ScheduledTasksCounter->SetValue(Timestamp, ++ScheduledTasksNum);
+		}
+	}
+
+	void FTasksProvider::SubsequentAdded(TaskTrace::FId TaskId, TaskTrace::FId SubsequentId, double Timestamp, uint32 ThreadId)
+	{
+		// when FGraphEvent is used to wait for a notification, it doesn't have an associated task and so is not created or launched. 
+		// In this case we need to create it and initialise it before registering the completion event
+		FTaskInfo* Task = TryGetOrCreateTask(TaskId);
+		if (Task == nullptr)
+		{
+			UE_LOG(LogTraceServices, Log, TEXT("SubsequentAdded(TaskId %d, SubsequentId %d, Timestamp %.6f) skipped"), TaskId, SubsequentId, Timestamp);
+			return;
+		}
+
+		Task->Id = TaskId;
+
+		AddRelative(TEXT("Subsequent"), TaskId, &FTaskInfo::Subsequents, SubsequentId, Timestamp, ThreadId);
+	}
+
+	void FTasksProvider::TaskStarted(TaskTrace::FId TaskId, double Timestamp, uint32 ThreadId)
+	{
+		if (!TryRegisterEvent(TEXT("TaskStarted"), TaskId, &FTaskInfo::StartedTimestamp, Timestamp, &FTaskInfo::StartedThreadId, ThreadId))
+		{
+			return;
+		}
+
+		ExecutionThreads.FindOrAdd(ThreadId).Add(TaskId);
+
+		FTaskInfo* Task = TryGetTask(TaskId);
+
+		if (IsNamedThread(Task->ThreadToExecuteOn))
+		{
+			NamedThreadsScheduledTasksCounter->SetValue(Timestamp, --ScheduledTasksNum);
+		}
+		else
+		{
+			ScheduledTasksCounter->SetValue(Timestamp, --ScheduledTasksNum);
+		}
+		RunningTasksCounter->SetValue(Timestamp, ++RunningTasksNum);
+
+		double LatencyMicrosecs = Task->StartedTimestamp * 1000000 - Task->ScheduledTimestamp * 1000000;
+		TaskLatencyCounter->SetValue(Timestamp, LatencyMicrosecs);
+	}
+
+	void FTasksProvider::NestedAdded(TaskTrace::FId TaskId, TaskTrace::FId NestedId, double Timestamp, uint32 ThreadId)
+	{
+		AddRelative(TEXT("Nested"), TaskId, &FTaskInfo::NestedTasks, NestedId, Timestamp, ThreadId);
+	}
+
+	void FTasksProvider::TaskFinished(TaskTrace::FId TaskId, double Timestamp)
+	{
+		if (!TryRegisterEvent(TEXT("TaskFinished"), TaskId, &FTaskInfo::FinishedTimestamp, Timestamp))
+		{
+			return;
+		}
+
+		RunningTasksCounter->SetValue(Timestamp, --RunningTasksNum);
+	}
+
+	void FTasksProvider::TaskCompleted(TaskTrace::FId TaskId, double Timestamp, uint32 ThreadId)
+	{
+		// when FGraphEvent is used to wait for a notification, it doesn't have an associated task and so is not created or launched. 
+		// In this case we need to create it and initialise it before registering the completion event
+		FTaskInfo* Task = TryGetOrCreateTask(TaskId);
+		if (Task == nullptr)
+		{
+			UE_LOG(LogTraceServices, Log, TEXT("TaskCompleted(TaskId %d, Timestamp %.6f) skipped"), TaskId, Timestamp);
+			return;
+		}
+
+		Task->Id = TaskId;
+
+		TryRegisterEvent(TEXT("TaskCompleted"), TaskId, &FTaskInfo::CompletedTimestamp, Timestamp, &FTaskInfo::CompletedThreadId, ThreadId);
+	}
+
+	void FTasksProvider::WaitingStarted(TArray<TaskTrace::FId> InTasks, double Timestamp, uint32 ThreadId)
+	{
+		FWaitingForTasks Waiting;
+		Waiting.Tasks = MoveTemp(InTasks);
+		Waiting.StartedTimestamp = Timestamp;
+
+		WaitingThreads.FindOrAdd(ThreadId).Add(MoveTemp(Waiting));
+	}
+
+	void FTasksProvider::WaitingFinished(double Timestamp, uint32 ThreadId)
+	{
+		TArray64<FWaitingForTasks>* Thread = WaitingThreads.Find(ThreadId);
+		checkf(Thread != nullptr, TEXT("%d"), ThreadId);
+
+		Thread->Last().FinishedTimestamp = Timestamp;
+	}
+
+	void FTasksProvider::InitTaskIdToIndexConversion(TaskTrace::FId InFirstTaskId)
+	{
+		check(InFirstTaskId != TaskTrace::InvalidId);
+		if (FirstTaskId == TaskTrace::InvalidId)
+		{
+			FirstTaskId = InFirstTaskId;
+		}
+	}
+
+	int64 FTasksProvider::GetTaskIndex(TaskTrace::FId TaskId) const
+	{
+		return (int64)TaskId - FirstTaskId;
+	}
+
+	const FTaskInfo* FTasksProvider::TryGetTask(TaskTrace::FId TaskId) const
+	{
+		check(TaskId != TaskTrace::InvalidId);
+		int64 TaskIndex = GetTaskIndex(TaskId);
+		return Tasks.IsValidIndex(TaskIndex) ? &Tasks[TaskIndex] : nullptr;
+	}
+
+	FTaskInfo* FTasksProvider::TryGetTask(TaskTrace::FId TaskId)
+	{
+		return const_cast<FTaskInfo*>(const_cast<const FTasksProvider*>(this)->TryGetTask(TaskId)); // reuse the const version
+	}
+
+	FTaskInfo* FTasksProvider::TryGetOrCreateTask(TaskTrace::FId TaskId)
+	{
+		int64 TaskIndex = GetTaskIndex(TaskId);
+		// traces can race, it's possible a trace with `TaskId = X` can come first, initialize `FirstTaskId` and only then a trace with 
+		// `TaskId = X - 1` arrives. This will produce `TaskIndex < 0`. Such traces can happen only at the very beginning of the capture 
+		// and are ignored
+		if (TaskIndex < 0)
+		{
+			return nullptr;
+		}
+
+		if (TaskIndex >= Tasks.Num())
+		{
+			Tasks.AddDefaulted(TaskIndex - Tasks.Num() + 1);
+		}
+
+		return &Tasks[TaskIndex];
+	}
+
+	bool FTasksProvider::IsNamedThread(int32 Thread)
+	{
+		return Thread != ENamedThreads::AnyThread;
+	}
+
+	bool FTasksProvider::TryRegisterEvent(const TCHAR* EventName, TaskTrace::FId TaskId, double FTaskInfo::* TimestampPtr, double TimestampValue, uint32 FTaskInfo::* ThreadIdPtr/* = nullptr*/, uint32 ThreadIdValue/* = 0*/)
+	{
+		UE_LOG(LogTraceServices, Verbose, TEXT("%s(TaskId: %d, Timestamp %.6f)"), EventName, TaskId, TimestampValue);
+
+		FTaskInfo* Task = TryGetTask(TaskId);
+		if (Task == nullptr)
+		{
+			UE_LOG(LogTraceServices, Log, TEXT("%s(TaskId %d, Timestamp %.6f) skipped"), EventName, TaskId, TimestampValue);
+			return false;
+		}
+
+		checkf(Task->*TimestampPtr == FTaskInfo::InvalidTimestamp, TEXT("TaskId %d, old TS %.6f, new TS %.6f"), TaskId, Task->*TimestampPtr, TimestampValue);
+		Task->*TimestampPtr = TimestampValue;
+		if (ThreadIdPtr != nullptr)
+		{
+			Task->*ThreadIdPtr = ThreadIdValue;
+		}
+
+		return true;
+	}
+
+	void FTasksProvider::AddRelative(const TCHAR* RelationType, TaskTrace::FId TaskId, TArray<FTaskInfo::FRelationInfo> FTaskInfo::* RelationsPtr, TaskTrace::FId RelativeId, double Timestamp, uint32 ThreadId)
+	{
+		UE_LOG(LogTraceServices, Verbose, TEXT("%s (%d) added to TaskId: %d, Timestamp %.6f)"), RelationType, RelativeId, TaskId, Timestamp);
+
+		FTaskInfo* Task = TryGetTask(TaskId);
+		if (Task == nullptr)
+		{
+			UE_LOG(LogTraceServices, Log, TEXT("Add%s(TaskId %d, OtherId: %d, Timestamp %.6f) skipped"), RelationType, TaskId, RelativeId, Timestamp);
+			return;
+		}
+
+		(Task->*RelationsPtr).Emplace(RelativeId, Timestamp, ThreadId);
+	}
+
+	/////////////////////////////////////////////////////////////////////////////////
+	// ITasksProvider impl
+
+	const FTaskInfo* FTasksProvider::TryGetTask(uint32 ThreadId, double Timestamp) const
+	{
+		const TArray64<TaskTrace::FId>* Thread = ExecutionThreads.Find(ThreadId);
+		if (Thread == nullptr)
+		{
+			return nullptr;
+		}
+
+		int64 NextTaskIndex = Algo::LowerBound(*Thread, Timestamp, 
+			[this](TaskTrace::FId TaskId, double Timestamp)
+			{
+				const FTaskInfo* Task = TryGetTask(TaskId);
+				return Task != nullptr && Task->StartedTimestamp <= Timestamp;
+			}
+		);
+
+		if (NextTaskIndex == 0)
+		{
+			return nullptr;
+		}
+
+		TaskTrace::FId TaskId = (*Thread)[NextTaskIndex - 1];
+		const FTaskInfo* Task = TryGetTask(TaskId);
+		return Task != nullptr && Task->FinishedTimestamp > Timestamp ? Task : nullptr;
+	}
+
+	const FWaitingForTasks* FTasksProvider::TryGetWaiting(uint32 ThreadId, double Timestamp) const
+	{
+		const TArray64<FWaitingForTasks>* Thread = WaitingThreads.Find(ThreadId);
+		if (Thread == nullptr)
+		{
+			return nullptr;
+		}
+
+		int64 NextWaitingIndex = Algo::LowerBound(*Thread, Timestamp,
+			[this](const FWaitingForTasks& Waiting, double Timestamp)
+			{
+				return Waiting.StartedTimestamp <= Timestamp;
+			}
+		);
+
+		if (NextWaitingIndex == 0)
+		{
+			return nullptr;
+		}
+
+		const FWaitingForTasks& Waiting = (*Thread)[NextWaitingIndex - 1];
+		return Waiting.FinishedTimestamp > Timestamp || Waiting.FinishedTimestamp == FTaskInfo::InvalidTimestamp ? &Waiting : nullptr;
+	}
+}
diff --git a/Engine/Source/Developer/TraceServices/Private/Model/TasksProfilerPrivate.h b/Engine/Source/Developer/TraceServices/Private/Model/TasksProfilerPrivate.h
new file mode 100644
--- /dev/null
+++ b/Engine/Source/Developer/TraceServices/Private/Model/TasksProfilerPrivate.h
@@ -0,0 +1,97 @@
+// Copyright Epic Games, Inc. All Rights Reserved.
+
+#pragma once
+
+#include "TraceServices/Model/AnalysisSession.h"
+#include "TraceServices/Model/TasksProfiler.h"
+#include "TraceServices/Model/Counters.h"
+#include "HAL/CriticalSection.h"
+#include "Async/TaskTrace.h"
+#include "Containers/Map.h"
+
+namespace TraceServices
+{
+	// tasks database, filled by the analyser and consumed by UI.
+	// Creates and fills a number of counters. Unprocessed info can be queried by `ITasksProvider` interface
+	class FTasksProvider : public ITasksProvider
+	{
+	public:
+		explicit FTasksProvider(IAnalysisSession& Session);
+
+		void Init(uint32 Version);
+
+		// traces consumption
+		void TaskCreated(TaskTrace::FId TaskId, double Timestamp, uint32 ThreadId);
+		void TaskLaunched(TaskTrace::FId TaskId, const TCHAR* DebugName, bool bTracked, int32 ThreadToExecuteOn, double Timestamp, uint32 ThreadId);
+		void TaskScheduled(TaskTrace::FId TaskId, double Timestamp, uint32 ThreadId);
+		void SubsequentAdded(TaskTrace::FId TaskId, TaskTrace::FId SubsequentId, double Timestamp, uint32 ThreadId);
+		void TaskStarted(TaskTrace::FId TaskId, double Timestamp, uint32 ThreadId);
+		void NestedAdded(TaskTrace::FId TaskId, TaskTrace::FId NestedId, double Timestamp, uint32 ThreadId);
+		void TaskFinished(TaskTrace::FId TaskId, double Timestamp);
+		void TaskCompleted(TaskTrace::FId TaskId, double Timestamp, uint32 ThreadId);
+
+		void WaitingStarted(TArray<TaskTrace::FId> Tasks, double Timestamp, uint32 ThreadId);
+		void WaitingFinished(double Timestamp, uint32 ThreadId);
+
+		// ITasksProvider impl
+		virtual const FTaskInfo* TryGetTask(uint32 ThreadId, double Timestamp) const override;
+		virtual const FTaskInfo* TryGetTask(TaskTrace::FId TaskId) const;
+		virtual const FWaitingForTasks* TryGetWaiting(uint32 ThreadId, double Timestamp) const override;
+
+	private:
+		// initializes conversion of task id to its index in the internal container, if it's not initialized already
+		void InitTaskIdToIndexConversion(TaskTrace::FId FirstTaskId);
+
+		// translates `TaskId` to the index into `Tasks` array, 
+		int64 GetTaskIndex(TaskTrace::FId TaskId) const; // `FistTaskId` must be initialised before the first use
+
+		FTaskInfo& GetTask(TaskTrace::FId TaskId);
+		FTaskInfo* TryGetTask(TaskTrace::FId TaskId);
+		// default-constructs a task if required
+		FTaskInfo* TryGetOrCreateTask(TaskTrace::FId TaskId);
+
+		static bool IsNamedThread(int32 Thread);
+
+		void CreateCounters();
+
+		// registers a minimal event with a name, task id, when and where it happened
+		bool TryRegisterEvent(const TCHAR* EventName, TaskTrace::FId TaskId, double FTaskInfo::* TimestampPtr, double TimestampValue, uint32 FTaskInfo::* ThreadIdPtr = nullptr, uint32 ThreadIdValue = 0);
+		// registers a relation with another task, when and where it was established
+		void AddRelative(const TCHAR* RelativeType, TaskTrace::FId TaskId, TArray<FTaskInfo::FRelationInfo> FTaskInfo::* RelationsPtr, TaskTrace::FId RelativeId, double Timestamp, uint32 ThreadId);
+
+	private:
+		IAnalysisSession& Session; // for synchronisation with consumers
+
+		uint32 Version = ~(uint32)0;
+
+		// task IDs are auto-incremented almost monotonic numbers, they are used as indices into tasks array minus the first task ID
+		// "almost monotonic" because a trace with TaskId==N+1 can arrive before all traces with TaskId==N
+		TArray64<FTaskInfo> Tasks;
+		
+		// on execution start, when it becomes known which thread will execute a task, all tasks are grouped by their execution threads, where
+		// they are naturally sorted by `StartedTimestamp`
+		TMap<uint32, TArray64<TaskTrace::FId>> ExecutionThreads; // ThreadId -> [TaskId]
+
+		TMap<uint32, TArray64<FWaitingForTasks>> WaitingThreads; // ThreadId -> [Waiting]
+
+		// task ID of the first trace. as traces can race, it's possible some traces of the previous task will be lost, which is not a probem
+		TaskTrace::FId FirstTaskId = TaskTrace::InvalidId;
+
+		// number of tasks waiting to be scheduled, because they are blocked by dependencies or they will be immediately scheduled
+		// and reported by another task trace
+		int64 WaitingForPrerequisitesTasksNum = 0;
+		// number of non named threads tasks sheduled for execution
+		int64 ScheduledTasksNum = 0;
+		// number of named threads tasks scheduled for execution
+		int64 NamedThreadsScheduledTasksNum = 0;
+		// number of tasks being executed
+		int64 RunningTasksNum = 0;
+
+		ICounterProvider& CounterProvider;
+		IEditableCounter* TaskLatencyCounter;
+		IEditableCounter* WaitingForPrerequisitesTasksCounter;
+		IEditableCounter* ScheduledTasksCounter;
+		IEditableCounter* NamedThreadsScheduledTasksCounter;
+		IEditableCounter* RunningTasksCounter;
+	};
+}
\ No newline at end of file
diff --git a/Engine/Source/Developer/TraceServices/Private/Model/ThreadsPrivate.h b/Engine/Source/Developer/TraceServices/Private/Model/ThreadsPrivate.h
--- a/Engine/Source/Developer/TraceServices/Private/Model/ThreadsPrivate.h
+++ b/Engine/Source/Developer/TraceServices/Private/Model/ThreadsPrivate.h
@@ -3,6 +3,7 @@
 #pragma once
 
 #include "TraceServices/Model/AnalysisSession.h"
+#include "TraceServices/Model/Threads.h"
 #include "ProfilingDebugging/MiscTrace.h"
 #include "HAL/PlatformAffinity.h"
 #include "Containers/Map.h"
diff --git a/Engine/Source/Developer/TraceServices/Private/Modules/TasksModule.cpp b/Engine/Source/Developer/TraceServices/Private/Modules/TasksModule.cpp
new file mode 100644
--- /dev/null
+++ b/Engine/Source/Developer/TraceServices/Private/Modules/TasksModule.cpp
@@ -0,0 +1,31 @@
+// Copyright Epic Games, Inc. All Rights Reserved.
+
+#include "TasksModule.h"
+#include "Analyzers/TasksAnalysis.h"
+#include "Model/TasksProfilerPrivate.h"
+#include "AnalysisServicePrivate.h"
+#include "TraceServices/ModuleService.h"
+
+namespace TraceServices
+{
+	static const FName TasksModuleName("TraceModule_TasksProfiler");
+	static const FName TasksProviderName("TasksProvider");
+
+	void FTasksModule::GetModuleInfo(FModuleInfo& OutModuleInfo)
+	{
+		OutModuleInfo.Name = TasksModuleName;
+		OutModuleInfo.DisplayName = TEXT("TasksProfiler");
+	}
+
+	void FTasksModule::OnAnalysisBegin(IAnalysisSession& Session)
+	{
+		FTasksProvider* TasksProvider = new FTasksProvider(Session);
+		Session.AddProvider(TasksProviderName, TasksProvider);
+		Session.AddAnalyzer(new FTasksAnalyzer(Session, *TasksProvider));
+	}
+
+	const ITasksProvider* ReadTasksProvider(const IAnalysisSession& Session)
+	{
+		return Session.ReadProvider<ITasksProvider>(TasksProviderName);
+	}
+} // namespace TraceServices
diff --git a/Engine/Source/Developer/TraceServices/Private/Modules/TasksModule.h b/Engine/Source/Developer/TraceServices/Private/Modules/TasksModule.h
new file mode 100644
--- /dev/null
+++ b/Engine/Source/Developer/TraceServices/Private/Modules/TasksModule.h
@@ -0,0 +1,26 @@
+// Copyright Epic Games, Inc. All Rights Reserved.
+
+#pragma once
+
+#include "TraceServices/ModuleService.h"
+#include "UObject/NameTypes.h"
+
+namespace TraceServices
+{
+	class FTasksModule : public IModule
+	{
+	public:
+		virtual void GetModuleInfo(FModuleInfo& OutModuleInfo) override;
+		virtual void OnAnalysisBegin(IAnalysisSession& Session) override;
+		virtual void GetLoggers(TArray<const TCHAR*>& OutLoggers) override
+		{
+		};
+		virtual const TCHAR* GetCommandLineArgument() override
+		{
+			return nullptr;
+		}
+		virtual void GenerateReports(const IAnalysisSession& Session, const TCHAR* CmdLine, const TCHAR* OutputDirectory) override 
+		{
+		};
+	};
+} // namespace TraceServices
diff --git a/Engine/Source/Developer/TraceServices/Private/TraceServicesModule.cpp b/Engine/Source/Developer/TraceServices/Private/TraceServicesModule.cpp
--- a/Engine/Source/Developer/TraceServices/Private/TraceServicesModule.cpp
+++ b/Engine/Source/Developer/TraceServices/Private/TraceServicesModule.cpp
@@ -13,6 +13,7 @@
 #include "Modules/NetProfilerModule.h"
 #include "Modules/MemoryModule.h"
 #include "Modules/DiagnosticsModule.h"
+#include "Modules/TasksModule.h"
 
 class FTraceServicesModule
 	: public ITraceServicesModule
@@ -38,6 +39,7 @@ private:
 	TraceServices::FNetProfilerModule NetProfilerModule;
 	TraceServices::FMemoryModule MemoryModule;
 	TraceServices::FDiagnosticsModule DiagnosticsModule;
+	TraceServices::FTasksModule TasksModule;
 };
 
 TSharedPtr<TraceServices::IAnalysisService> FTraceServicesModule::GetAnalysisService()
@@ -86,10 +88,12 @@ void FTraceServicesModule::StartupModule()
 	IModularFeatures::Get().RegisterModularFeature(TraceServices::ModuleFeatureName, &NetProfilerModule);
 	IModularFeatures::Get().RegisterModularFeature(TraceServices::ModuleFeatureName, &MemoryModule);
 	IModularFeatures::Get().RegisterModularFeature(TraceServices::ModuleFeatureName, &DiagnosticsModule);
+	IModularFeatures::Get().RegisterModularFeature(TraceServices::ModuleFeatureName, &TasksModule);
 }
 
 void FTraceServicesModule::ShutdownModule()
 {
+	IModularFeatures::Get().UnregisterModularFeature(TraceServices::ModuleFeatureName, &TasksModule);
 	IModularFeatures::Get().UnregisterModularFeature(TraceServices::ModuleFeatureName, &DiagnosticsModule);
 	IModularFeatures::Get().UnregisterModularFeature(TraceServices::ModuleFeatureName, &MemoryModule);
 	IModularFeatures::Get().UnregisterModularFeature(TraceServices::ModuleFeatureName, &NetProfilerModule);
diff --git a/Engine/Source/Developer/TraceServices/Public/TraceServices/Model/TasksProfiler.h b/Engine/Source/Developer/TraceServices/Public/TraceServices/Model/TasksProfiler.h
new file mode 100644
--- /dev/null
+++ b/Engine/Source/Developer/TraceServices/Public/TraceServices/Model/TasksProfiler.h
@@ -0,0 +1,87 @@
+// Copyright Epic Games, Inc. All Rights Reserved.
+
+#pragma once
+
+#include "TraceServices/Model/AnalysisSession.h"
+#include "Async/TaskTrace.h"
+#include "Containers/Array.h"
+
+namespace TraceServices
+{
+	// task biography
+	struct FTaskInfo
+	{
+		FTaskInfo() = default;
+
+		TaskTrace::FId Id = TaskTrace::InvalidId;
+
+		// this members are filled during task launch
+		const TCHAR* DebugName; // user provided or automatically generated (in format `full_path(lineno)`)
+		bool bTracked; // if the task is awaitable, not awaitable tasks will never be "completed", they are done when their execution is finished
+		int32 ThreadToExecuteOn; // ENamedThreads::Type - which thread type the task will be executed on, one of named threads or a worker thread
+
+		static constexpr double InvalidTimestamp = 0.0;
+
+		// some tasks can be created with delayed launch, otherwise "created" and "launched" timestamps and thread IDs are equal
+		double CreatedTimestamp = InvalidTimestamp;
+		uint32 CreatedThreadId;
+		// the moment and the place when the task was launched
+		double LaunchedTimestamp = InvalidTimestamp;
+		uint32 LaunchedThreadId;
+		// the moment and the place when all task's dependencies are resolved and it's added to the execution queue
+		double ScheduledTimestamp = InvalidTimestamp;
+		uint32 ScheduledThreadId;
+		// execution started
+		double StartedTimestamp = InvalidTimestamp;
+		uint32 StartedThreadId;
+		// execution finished, but the task can be still not completed waiting for all nested tasks to complete
+		double FinishedTimestamp = InvalidTimestamp;
+		// execution done and all nested tasks are completed
+		double CompletedTimestamp = InvalidTimestamp;
+		uint32 CompletedThreadId;
+
+		//  relation with another task, when and where it was established
+		struct FRelationInfo
+		{
+			FRelationInfo(TaskTrace::FId InRelativeId, double InTimestamp, uint32 InThreadId)
+				: RelativeId(InRelativeId)
+				, Timestamp(InTimestamp)
+				, ThreadId(InThreadId)
+			{
+			}
+
+			TaskTrace::FId RelativeId;
+			double Timestamp;
+			uint32 ThreadId;
+		};
+
+		// other tasks that wait for this task completion before they'll start execution
+		TArray<FRelationInfo> Subsequents;
+		// the task is completed only after all nested tasks are completed
+		TArray<FRelationInfo> NestedTasks;
+	};
+
+	struct FWaitingForTasks
+	{
+		TArray<TaskTrace::FId> Tasks;
+		double StartedTimestamp;
+		double FinishedTimestamp;
+	};
+
+	// query interface to tasks info
+	class ITasksProvider : public IProvider
+	{
+	public:
+		virtual ~ITasksProvider() = default;
+
+		// returns a task for given thread and timestamp, if any, otherwise `nullptr`
+		virtual const FTaskInfo* TryGetTask(uint32 ThreadId, double Timestamp) const = 0;
+		// returns task info for given task ID
+		virtual const FTaskInfo* TryGetTask(TaskTrace::FId TaskId) const = 0;
+
+		// returns an info about waiting for tasks completion for given thread and timestamp, if any, otherwise `nullptr`
+		virtual const FWaitingForTasks* TryGetWaiting(uint32 ThreadId, double Timestamp) const = 0;
+	};
+
+	TRACESERVICES_API const ITasksProvider* ReadTasksProvider(const IAnalysisSession& Session);
+}
\ No newline at end of file
diff --git a/Engine/Source/Runtime/Core/Private/Async/TaskGraph.cpp b/Engine/Source/Runtime/Core/Private/Async/TaskGraph.cpp
--- a/Engine/Source/Runtime/Core/Private/Async/TaskGraph.cpp
+++ b/Engine/Source/Runtime/Core/Private/Async/TaskGraph.cpp
@@ -214,7 +214,22 @@ static FAutoConsoleVariableRef CVarUseHiPriThreads(
 	#define TASKGRAPH_SCOPE_CYCLE_COUNTER(Index, Name)
 #endif
 
+TArray<TaskTrace::FId> GetTraceIds(const FGraphEventArray& Tasks)
+{
+#if UE_TASK_TRACE_ENABLED
+	TArray<TaskTrace::FId> TasksIds;
+	TasksIds.Reserve(Tasks.Num());
+
+	for (const FGraphEventRef& Task : Tasks)
+	{
+		TasksIds.Add(Task->GetTraceId());
+	}
 
+	return TasksIds;
+#else
+	return {};
+#endif
+}
 
 /** 
  *	Pointer to the task graph implementation singleton.
@@ -1188,6 +1203,8 @@ public:
 	**/
 	FTaskGraphImplementation(int32)
 	{
+		TaskTrace::Init();
+
 		bCreatedHiPriorityThreads = !!ENamedThreads::bHasHighPriorityThreads;
 		bCreatedBackgroundPriorityThreads = !!ENamedThreads::bHasBackgroundThreads;
 
@@ -1501,7 +1518,9 @@ public:
 
 	virtual void WaitUntilTasksComplete(const FGraphEventArray& Tasks, ENamedThreads::Type CurrentThreadIfKnown = ENamedThreads::AnyThread) final override
 	{
+		TaskTrace::FWaitingScope WaitingScope(GetTraceIds(Tasks));
 		TRACE_CPUPROFILER_EVENT_SCOPE(WaitUntilTasksComplete);
+
 		ENamedThreads::Type CurrentThread = CurrentThreadIfKnown;
 		if (ENamedThreads::GetThreadIndex(CurrentThreadIfKnown) == ENamedThreads::AnyThread)
 		{
@@ -1537,6 +1556,7 @@ public:
 					return;
 				}
 			}
+
 			// named thread process tasks while we wait
 			TGraphTask<FReturnGraphTask>::CreateTask(&Tasks, CurrentThread).ConstructAndDispatchWhenReady(CurrentThread);
 			ProcessThreadUntilRequestReturn(CurrentThread);
@@ -1561,6 +1581,7 @@ public:
 				}
 				UE_LOG(LogTaskGraph, Fatal, TEXT("Recursive waits are not allowed in single threaded mode."));
 			}
+		
 			// We will just stall this thread on an event while we wait
 			FScopedEvent Event;
 			TriggerEventWhenTasksComplete(Event.Get(), Tasks, CurrentThreadIfKnown);
@@ -1637,7 +1658,6 @@ public:
 		int32 MyIndex = int32((uint32(ThreadInNeed) - NumNamedThreads) % NumTaskThreadsPerSet);
 		int32 Priority = int32((uint32(ThreadInNeed) - NumNamedThreads) / NumTaskThreadsPerSet);
 		check(MyIndex >= 0 && MyIndex < LocalNumWorkingThread &&
-			MyIndex < (PLATFORM_64BITS ? 63 : 32) &&
 			Priority >= 0 && Priority < ENamedThreads::NumThreadPriorities);
 
 		return IncomingAnyThreadTasks[Priority].Pop(MyIndex, true);
@@ -1719,7 +1739,7 @@ private:
 	enum
 	{
 		/** Compile time maximum number of threads. Didn't really need to be a compile time constant, but task thread are limited by MAX_LOCK_FREE_LINKS_AS_BITS **/
-		MAX_THREADS = 26 * (CREATE_HIPRI_TASK_THREADS + CREATE_BACKGROUND_TASK_THREADS + 1) + ENamedThreads::ActualRenderingThread + 1,
+		MAX_THREADS = 0xFFFF,
 		MAX_THREAD_PRIORITIES = 3
 	};
 
@@ -1779,6 +1799,8 @@ public:
 	FTaskGraphCompatibilityImplementation(int32 InNumWorkerThreads) 
 		: NumWorkerThreads(FForkProcessHelper::IsForkedMultithreadInstance() ? CVar_ForkedProcess_MaxWorkerThreads : InNumWorkerThreads)
 	{
+		TaskTrace::Init();
+
 		if (FTaskGraphInterface::IsMultithread())
 		{
 			int32 NumBackgroundWorkers = FMath::Max(1, NumWorkerThreads - FMath::Min<int>(GNumForegroundWorkers, NumWorkerThreads));
@@ -1974,7 +1996,9 @@ private:
 
 	void WaitUntilTasksComplete(const FGraphEventArray& Tasks, ENamedThreads::Type CurrentThreadIfKnown = ENamedThreads::AnyThread) final override
 	{
+		TaskTrace::FWaitingScope WaitingScope(GetTraceIds(Tasks));
 		TRACE_CPUPROFILER_EVENT_SCOPE(WaitUntilTasksComplete);
+	
 		ENamedThreads::Type CurrentThread = CurrentThreadIfKnown;
 		if (ENamedThreads::GetThreadIndex(CurrentThreadIfKnown) == ENamedThreads::AnyThread)
 		{
@@ -2010,6 +2034,7 @@ private:
 					return;
 				}
 			}
+
 			// named thread process tasks while we wait
 			TGraphTask<FReturnGraphTask>::CreateTask(&Tasks, CurrentThread).ConstructAndDispatchWhenReady(CurrentThread);
 			ProcessThreadUntilRequestReturn(CurrentThread);
@@ -2225,7 +2250,7 @@ static TLockFreeClassAllocator_TLSCache<FGraphEvent, PLATFORM_CACHE_LINE_SIZE> T
 
 FGraphEventRef FGraphEvent::CreateGraphEvent()
 {
-	return TheGraphEventAllocator.New();
+	return new(TheGraphEventAllocator.Allocate()) FGraphEvent{};
 }
 
 void FGraphEvent::Recycle(FGraphEvent* ToRecycle)
@@ -2288,6 +2313,8 @@ void FGraphEvent::DispatchSubsequents(TArray<FBaseGraphTask*>& NewTasks, ENamedT
 		NewTask->ConditionalQueueTask(CurrentThreadIfKnown, bWakeUpWorker);
 	}
 	NewTasks.Reset();
+
+	TaskTrace::Completed(GetTraceId());
 }
 
 FGraphEvent::~FGraphEvent()
diff --git a/Engine/Source/Runtime/Core/Private/Async/TaskTrace.cpp b/Engine/Source/Runtime/Core/Private/Async/TaskTrace.cpp
new file mode 100644
--- /dev/null
+++ b/Engine/Source/Runtime/Core/Private/Async/TaskTrace.cpp
@@ -0,0 +1,206 @@
+// Copyright Epic Games, Inc. All Rights Reserved.
+
+#include "Async/TaskTrace.h"
+
+#if UE_TASK_TRACE_ENABLED
+
+#include "CoreMinimal.h"
+#include "Trace/Trace.h"
+#include "Trace/Trace.inl"
+#include "HAL/PlatformTime.h"
+
+namespace TaskTrace
+{
+	UE_TRACE_CHANNEL(TaskChannel);
+
+	UE_TRACE_EVENT_BEGIN(TaskTrace, InitTaskTrace)
+		UE_TRACE_EVENT_FIELD(uint32, Version)
+	UE_TRACE_EVENT_END()
+
+	UE_TRACE_EVENT_BEGIN(TaskTrace, Created)
+		UE_TRACE_EVENT_FIELD(uint64, Timestamp)
+		UE_TRACE_EVENT_FIELD(uint32, TaskId)
+	UE_TRACE_EVENT_END()
+
+	UE_TRACE_EVENT_BEGIN(TaskTrace, Launched)
+		UE_TRACE_EVENT_FIELD(uint64, Timestamp)
+		UE_TRACE_EVENT_FIELD(uint32, TaskId)
+		UE_TRACE_EVENT_FIELD(UE::Trace::AnsiString, DebugName)
+		UE_TRACE_EVENT_FIELD(bool, Tracked)
+		UE_TRACE_EVENT_FIELD(int32, ThreadToExecuteOn)
+	UE_TRACE_EVENT_END()
+
+	UE_TRACE_EVENT_BEGIN(TaskTrace, Scheduled)
+		UE_TRACE_EVENT_FIELD(uint64, Timestamp)
+		UE_TRACE_EVENT_FIELD(uint32, TaskId)
+	UE_TRACE_EVENT_END()
+
+	UE_TRACE_EVENT_BEGIN(TaskTrace, SubsequentAdded)
+		UE_TRACE_EVENT_FIELD(uint64, Timestamp)
+		UE_TRACE_EVENT_FIELD(uint32, TaskId)
+		UE_TRACE_EVENT_FIELD(uint32, SubsequentId)
+	UE_TRACE_EVENT_END()
+
+	UE_TRACE_EVENT_BEGIN(TaskTrace, Started)
+		UE_TRACE_EVENT_FIELD(uint64, Timestamp)
+		UE_TRACE_EVENT_FIELD(uint32, TaskId)
+	UE_TRACE_EVENT_END()
+
+	UE_TRACE_EVENT_BEGIN(TaskTrace, NestedAdded)
+		UE_TRACE_EVENT_FIELD(uint64, Timestamp)
+		UE_TRACE_EVENT_FIELD(uint32, TaskId)
+		UE_TRACE_EVENT_FIELD(uint32, NestedId)
+	UE_TRACE_EVENT_END()
+
+	UE_TRACE_EVENT_BEGIN(TaskTrace, Finished)
+		UE_TRACE_EVENT_FIELD(uint64, Timestamp)
+		UE_TRACE_EVENT_FIELD(uint32, TaskId)
+	UE_TRACE_EVENT_END()
+
+	UE_TRACE_EVENT_BEGIN(TaskTrace, Completed)
+		UE_TRACE_EVENT_FIELD(uint64, Timestamp)
+		UE_TRACE_EVENT_FIELD(uint32, TaskId)
+	UE_TRACE_EVENT_END()
+
+	UE_TRACE_EVENT_BEGIN(TaskTrace, WaitingStarted)
+		UE_TRACE_EVENT_FIELD(uint64, Timestamp)
+		UE_TRACE_EVENT_FIELD(uint32[], Tasks)
+	UE_TRACE_EVENT_END()
+
+	UE_TRACE_EVENT_BEGIN(TaskTrace, WaitingFinished)
+		UE_TRACE_EVENT_FIELD(uint64, Timestamp)
+	UE_TRACE_EVENT_END()
+
+	FId GenerateTaskId()
+	{
+		static std::atomic<FId> UId{ 0 };
+		return UId.fetch_add(1, std::memory_order_relaxed);
+	}
+
+	static bool bGTaskTraceInitialized = false;
+
+	void Init()
+	{
+		UE_TRACE_LOG(TaskTrace, InitTaskTrace, TaskChannel)
+			<< InitTaskTrace.Version(0);
+
+		bGTaskTraceInitialized = true;
+	}
+
+	void Created(FId TaskId)
+	{
+		check(bGTaskTraceInitialized);
+		check(TaskId != InvalidId);
+		check(TaskId != 0);
+
+		UE_TRACE_LOG(TaskTrace, Created, TaskChannel)
+			<< Created.Timestamp(FPlatformTime::Cycles64())
+			<< Created.TaskId(TaskId);
+	}
+
+	void Launched(FId TaskId, const char* DebugName, bool bTracked, ENamedThreads::Type ThreadToExecuteOn)
+	{
+		check(bGTaskTraceInitialized);
+		check(TaskId != InvalidId);
+
+		UE_TRACE_LOG(TaskTrace, Launched, TaskChannel)
+			<< Launched.Timestamp(FPlatformTime::Cycles64())
+			<< Launched.TaskId(TaskId)
+			<< Launched.DebugName(DebugName != nullptr ? DebugName : "")
+			<< Launched.Tracked(bTracked)
+			<< Launched.ThreadToExecuteOn(ThreadToExecuteOn);
+	}
+
+	void Scheduled(FId TaskId)
+	{
+		check(bGTaskTraceInitialized);
+		check(TaskId != InvalidId);
+
+		UE_TRACE_LOG(TaskTrace, Scheduled, TaskChannel)
+			<< Scheduled.Timestamp(FPlatformTime::Cycles64())
+			<< Scheduled.TaskId(TaskId);
+	}
+
+	void SubsequentAdded(FId TaskId, FId SubsequentId)
+	{
+		check(bGTaskTraceInitialized);
+
+		// "empty" FGraphEvent is used for synchronisation only, to wait for a notification. It doesn't have an associated task and ID.
+		if (TaskId == InvalidId)
+		{
+			TaskId = GenerateTaskId();
+		}
+
+		check(SubsequentId != InvalidId);
+		UE_TRACE_LOG(TaskTrace, SubsequentAdded, TaskChannel)
+			<< SubsequentAdded.Timestamp(FPlatformTime::Cycles64())
+			<< SubsequentAdded.TaskId(TaskId)
+			<< SubsequentAdded.SubsequentId(SubsequentId);
+	}
+
+	void Started(FId TaskId)
+	{
+		check(bGTaskTraceInitialized);
+		check(TaskId != InvalidId);
+
+		UE_TRACE_LOG(TaskTrace, Started, TaskChannel)
+			<< Started.Timestamp(FPlatformTime::Cycles64())
+			<< Started.TaskId(TaskId);
+	}
+
+	void NestedAdded(FId TaskId, FId NestedId)
+	{
+		check(bGTaskTraceInitialized);
+		check(TaskId != InvalidId);
+		check(NestedId != InvalidId);
+
+		UE_TRACE_LOG(TaskTrace, NestedAdded, TaskChannel)
+			<< NestedAdded.Timestamp(FPlatformTime::Cycles64())
+			<< NestedAdded.TaskId(TaskId)
+			<< NestedAdded.NestedId(NestedId);
+	}
+
+	void Finished(FId TaskId)
+	{
+		check(bGTaskTraceInitialized);
+		check(TaskId != InvalidId);
+
+		UE_TRACE_LOG(TaskTrace, Finished, TaskChannel)
+			<< Finished.Timestamp(FPlatformTime::Cycles64())
+			<< Finished.TaskId(TaskId);
+	}
+
+	void Completed(FId TaskId)
+	{
+		check(bGTaskTraceInitialized);
+
+		// "empty" FGraphEvent is used for synchronisation only, to wait for a notification. It doesn't have an associated task and ID.
+		if (TaskId == InvalidId)
+		{
+			TaskId = GenerateTaskId();
+		}
+
+		UE_TRACE_LOG(TaskTrace, Completed, TaskChannel)
+			<< Completed.Timestamp(FPlatformTime::Cycles64())
+			<< Completed.TaskId(TaskId);
+	}
+
+	FWaitingScope::FWaitingScope(const TArray<FId>& Tasks)
+	{
+		check(bGTaskTraceInitialized);
+
+		UE_TRACE_LOG(TaskTrace, WaitingStarted, TaskChannel)
+			<< WaitingStarted.Timestamp(FPlatformTime::Cycles64())
+			<< WaitingStarted.Tasks(Tasks.GetData(), Tasks.Num());
+	}
+
+	FWaitingScope::~FWaitingScope()
+	{
+		check(bGTaskTraceInitialized);
+
+		UE_TRACE_LOG(TaskTrace, WaitingFinished, TaskChannel)
+			<< WaitingFinished.Timestamp(FPlatformTime::Cycles64());
+	}
+}
+
+#endif // UE_TASK_TRACE_ENABLED
diff --git a/Engine/Source/Runtime/Core/Public/Async/TaskGraphInterfaces.h b/Engine/Source/Runtime/Core/Public/Async/TaskGraphInterfaces.h
--- a/Engine/Source/Runtime/Core/Public/Async/TaskGraphInterfaces.h
+++ b/Engine/Source/Runtime/Core/Public/Async/TaskGraphInterfaces.h
@@ -26,6 +26,8 @@
 
 #include "Async/Fundamental/Task.h"
 
+#include "Async/TaskTrace.h"
+
 #if !defined(STATS)
 #error "STATS must be defined as either zero or one."
 #endif
@@ -209,7 +211,6 @@ namespace ENamedThreads
 			);
 		return Type(ThreadAndIndex | TaskPriority);
 	}
-
 }
 
 DECLARE_INTRINSIC_TYPE_LAYOUT(ENamedThreads::Type);
@@ -278,6 +279,9 @@ typedef TRefCountPtr<class FGraphEvent> FGraphEventRef;
 /** Convenience typedef for a an array a graph events **/
 typedef TArray<FGraphEventRef, TInlineAllocator<4> > FGraphEventArray;
 
+/** returns trace IDs of given tasks **/
+CORE_API TArray<TaskTrace::FId> GetTraceIds(const FGraphEventArray& Tasks);
+
 /** Interface tot he task graph system **/
 class FTaskGraphInterface
 {
@@ -430,7 +434,6 @@ public:
  *	Tasks go through a very specific life stage progression, and this is verified.
  **/
 
-
 class FBaseGraphTask
 {
 public:
@@ -509,6 +512,15 @@ protected:
 		}
 	}
 
+	TaskTrace::FId GetTraceId() const
+	{
+#if UE_TASK_TRACE_ENABLED
+		return TraceId;
+#else
+		return TaskTrace::InvalidId;
+#endif
+	}
+
 private:
 	friend class FNamedTaskThread;
 	friend class FTaskThreadBase;
@@ -557,6 +569,7 @@ private:
 	void QueueTask(ENamedThreads::Type CurrentThreadIfKnown, bool bWakeUpWorker)
 	{
 		checkThreadGraph(LifeStage.Increment() == int32(LS_Queued));
+		TaskTrace::Scheduled(GetTraceId());
 		FTaskGraphInterface::Get().QueueTask(this, bWakeUpWorker, ThreadToExecuteOn, CurrentThreadIfKnown);
 	}
 
@@ -586,6 +599,10 @@ private:
 	int32 InheritedTraceTag;
 #endif
 	LLM(const UE::LLMPrivate::FTagData* InheritedLLMTag);
+
+#if UE_TASK_TRACE_ENABLED
+	TaskTrace::FId TraceId{ TaskTrace::GenerateTaskId() };
+#endif
 };
 
 /** 
@@ -605,13 +622,19 @@ public:
 
 	// the returned event will have ref count zero; be sure to add one!
 	static CORE_API FGraphEvent* CreateGraphEventWithInlineStorage();
+
 	/**
 	 *	Attempts to a new subsequent task. If this event has already fired, false is returned and action must be taken to ensure that the task will still fire even though this event cannot be a prerequisite (because it is already finished).
 	 *	@return true if the task was successfully set up as a subsequent. false if the event has already fired.
 	**/
-	bool AddSubsequent(class FBaseGraphTask* Task)
+	bool AddSubsequent(class FBaseGraphTask* Subsequent)
 	{
-		return SubsequentList.PushIfNotClosed(Task);
+		bool bSucceeded = SubsequentList.PushIfNotClosed(Subsequent);
+		if (bSucceeded)
+		{
+			TaskTrace::SubsequentAdded(GetTraceId(), Subsequent->GetTraceId());
+		}
+		return bSucceeded;
 	}
 
 	/**
@@ -631,6 +654,7 @@ public:
 	{
 		checkThreadGraph(!IsComplete()); // it is not legal to add a DontCompleteUntil after the event has been completed. Basically, this is only legal within a task function.
 		new (EventsToWaitFor) FGraphEventRef(EventToWaitFor);
+		TaskTrace::NestedAdded(GetTraceId(), EventToWaitFor->GetTraceId());
 	}
 
 	/**
@@ -679,12 +703,28 @@ public:
 	/**
 	 * Sets a name for the event for debugging purposes.
 	 */
-#if !UE_BUILD_SHIPPING && !UE_BUILD_TEST
 	void SetDebugName(const TCHAR* Name)
 	{
+#if !UE_BUILD_SHIPPING && !UE_BUILD_TEST
 		DebugName = Name;
+#endif
 	}
+
+	TaskTrace::FId GetTraceId() const
+	{
+#if UE_TASK_TRACE_ENABLED
+		return TraceId;
+#else
+		return TaskTrace::InvalidId;
+#endif
+	}
+
+	void SetTraceId(TaskTrace::FId InTraceId)
+	{
+#if UE_TASK_TRACE_ENABLED
+		TraceId = InTraceId;
 #endif
+	}
 
 private:
 	friend class TRefCountPtr<FGraphEvent>;
@@ -700,7 +740,7 @@ private:
 	 *	Hidden Constructor
 	**/
 	friend struct FGraphEventAndSmallTaskStorage;
-	FGraphEvent(bool bInInline = false)
+	FGraphEvent()
 		: ThreadToDoGatherOn(ENamedThreads::AnyHiPriThreadHiPriTask)
 	{
 	}
@@ -738,7 +778,6 @@ public:
 		return RefCount;
 	}
 
-
 private:
 
 	/** Threadsafe list of subsequents for the event **/
@@ -752,9 +791,11 @@ private:
 #if !UE_BUILD_SHIPPING && !UE_BUILD_TEST
 	const TCHAR* DebugName = nullptr;
 #endif
-};
-
 
+#if UE_TASK_TRACE_ENABLED
+	TaskTrace::FId TraceId = TaskTrace::InvalidId;
+#endif
+};
 
 /** 
  The user defined task type can take arguments to a constructor. These arguments (unfortunately) must not be references.
@@ -873,6 +914,7 @@ public:
 	{
 		bool bWakeUpWorker = true;
 		ConditionalQueueTask(CurrentThreadIfKnown, bWakeUpWorker);
+		TaskTrace::Launched(GetTraceId(), nullptr, Subsequents.IsValid(), ENamedThreads::GetThreadIndex(((TTask*)&TaskStorage)->GetDesiredThread()));
 	}
 
 	FGraphEventRef GetCompletionEvent()
@@ -909,12 +951,14 @@ private:
 		}
 		
 		TTask& Task = *(TTask*)&TaskStorage;
+		TaskTrace::Started(GetTraceId());
 		{
 			FScopeCycleCounter Scope(Task.GetStatId(), true); 
 			Task.DoTask(CurrentThread, Subsequents);
 			Task.~TTask();
 			checkThreadGraph(ENamedThreads::GetThreadIndex(CurrentThread) <= ENamedThreads::GetRenderThread() || FMemStack::Get().IsEmpty()); // you must mark and pop memstacks if you use them in tasks! Named threads are excepted.
 		}
+		TaskTrace::Finished(GetTraceId());
 		
 		TaskConstructed = false;
 
@@ -947,7 +991,7 @@ private:
 
 	/** 
 	 *	Private constructor, constructs the base class with the number of prerequisites.
-	 *	@param InSubsequents subsequents to associate with this task. Thsi refernence is destroyed in the process!
+	 *	@param InSubsequents subsequents to associate with this task. This refernence is destroyed in the process!
 	 *	@param NumberOfPrerequistitesOutstanding the number of prerequisites this task will have when it is built.
 	**/
 	TGraphTask(FGraphEventRef InSubsequents, int32 NumberOfPrerequistitesOutstanding)
@@ -955,6 +999,10 @@ private:
 		, TaskConstructed(false)
 	{
 		Subsequents.Swap(InSubsequents);
+		if (Subsequents.IsValid())
+		{
+			Subsequents->SetTraceId(GetTraceId());
+		}
 	}
 
 	/** 
@@ -1010,6 +1058,8 @@ private:
 	 **/
 	FGraphEventRef Setup(const FGraphEventArray* Prerequisites = NULL, ENamedThreads::Type CurrentThreadIfKnown = ENamedThreads::AnyThread)
 	{
+		TaskTrace::Launched(GetTraceId(), nullptr, Subsequents.IsValid(), ENamedThreads::GetThreadIndex(((TTask*)&TaskStorage)->GetDesiredThread()));
+
 		FGraphEventRef ReturnedEventRef = Subsequents; // very important so that this doesn't get destroyed before we return
 		SetupPrereqs(Prerequisites, CurrentThreadIfKnown, true);
 		return ReturnedEventRef;
@@ -1028,6 +1078,8 @@ private:
 	 **/
 	TGraphTask* Hold(const FGraphEventArray* Prerequisites = NULL, ENamedThreads::Type CurrentThreadIfKnown = ENamedThreads::AnyThread)
 	{
+		TaskTrace::Created(GetTraceId());
+
 		SetupPrereqs(Prerequisites, CurrentThreadIfKnown, false);
 		return this;
 	}
diff --git a/Engine/Source/Runtime/Core/Public/Async/TaskTrace.h b/Engine/Source/Runtime/Core/Public/Async/TaskTrace.h
new file mode 100644
--- /dev/null
+++ b/Engine/Source/Runtime/Core/Public/Async/TaskTrace.h
@@ -0,0 +1,55 @@
+// Copyright Epic Games, Inc. All Rights Reserved.
+
+#pragma once
+
+#include "Trace/Config.h"
+
+#if UE_TRACE_ENABLED && !IS_PROGRAM && !UE_BUILD_SHIPPING
+#define UE_TASK_TRACE_ENABLED 1
+#else
+#define UE_TASK_TRACE_ENABLED 0
+#endif
+
+namespace TaskTrace
+{
+	using FId = uint32;
+
+	const FId InvalidId = ~FId(0);
+
+	static constexpr uint32 TaskTraceVersion = 1;
+
+	FId CORE_API GenerateTaskId();
+
+	void CORE_API Init();
+	void CORE_API Created(FId TaskId); // optional, used only if a task was created but not launched immediately
+	void CORE_API Launched(FId TaskId, const char* DebugName, bool bTracked, ENamedThreads::Type ThreadToExecuteOn);
+	void CORE_API Scheduled(FId TaskId);
+	void CORE_API SubsequentAdded(FId TaskId, FId SubsequentId);
+	void CORE_API Started(FId TaskId);
+	void CORE_API NestedAdded(FId TaskId, FId NestedId);
+	void CORE_API Finished(FId TaskId);
+	void CORE_API Completed(FId TaskId);
+
+	struct CORE_API FWaitingScope
+	{
+		explicit FWaitingScope(const TArray<FId>& Tasks); // waiting for given tasks completion
+		~FWaitingScope();
+	};
+
+#if !UE_TASK_TRACE_ENABLED
+	// NOOP implementation
+	inline FId GenerateTaskId() { return InvalidId; }
+	inline void Init() {}
+	inline void Created(FId TaskId) {}
+	inline void Launched(FId TaskId, const char* DebugName, bool bTracked, ENamedThreads::Type ThreadToExecuteOn) {}
+	inline void Scheduled(FId TaskId) {}
+	inline void SubsequentAdded(FId TaskId, FId SubsequentId) {}
+	inline void Started(FId TaskId) {}
+	inline void NestedAdded(FId TaskId, FId NestedId) {}
+	inline void Finished(FId TaskId) {}
+	inline void Completed(FId TaskId) {}
+	inline FWaitingScope::FWaitingScope(const TArray<FId>& Tasks) {}
+	inline FWaitingScope::~FWaitingScope() {}
+#endif // UE_TASK_TRACE_ENABLED
+}
+
diff --git a/Engine/Source/Runtime/RenderCore/Private/RenderingThread.cpp b/Engine/Source/Runtime/RenderCore/Private/RenderingThread.cpp
--- a/Engine/Source/Runtime/RenderCore/Private/RenderingThread.cpp
+++ b/Engine/Source/Runtime/RenderCore/Private/RenderingThread.cpp
@@ -25,6 +25,7 @@
 #include "HAL/LowLevelMemTracker.h"
 #include "ProfilingDebugging/MiscTrace.h"
 #include "ProfilingDebugging/CsvProfiler.h"
+#include "Async/TaskTrace.h"
 
 
 //
@@ -1125,6 +1126,7 @@ static void GameThreadWaitForTask(const FGraphEventRef& Task, ENamedThreads::Typ
 {
 	TRACE_CPUPROFILER_EVENT_SCOPE(GameThreadWaitForTask);
 	SCOPE_TIME_GUARD(TEXT("GameThreadWaitForTask"));
+	TaskTrace::FWaitingScope WaitingScope(GetTraceIds({ Task }));
 
 	check(IsInGameThread());
 	check(IsValidRef(Task));	
