Non-unity fixes

#rnx

diff --git a/Engine/Source/Developer/TraceServices/Private/Common/StringStore.cpp b/Engine/Source/Developer/TraceServices/Private/Common/StringStore.cpp
--- a/Engine/Source/Developer/TraceServices/Private/Common/StringStore.cpp
+++ b/Engine/Source/Developer/TraceServices/Private/Common/StringStore.cpp
@@ -1,6 +1,8 @@
 // Copyright Epic Games, Inc. All Rights Reserved.
 
 #include "Common/StringStore.h"
+#include "Misc/ScopeLock.h"
+
 
 namespace TraceServices
 {
diff --git a/Engine/Source/Developer/TraceServices/Private/Common/StringStore.h b/Engine/Source/Developer/TraceServices/Private/Common/StringStore.h
--- a/Engine/Source/Developer/TraceServices/Private/Common/StringStore.h
+++ b/Engine/Source/Developer/TraceServices/Private/Common/StringStore.h
@@ -7,6 +7,7 @@
 #include "Common/PagedArray.h"
 #include "Containers/Map.h"
 #include "Containers/StringView.h"
+#include "HAL/CriticalSection.h"
 
 namespace TraceServices
 {
