fixed TaskGraph traces, was crashing UnrealInsights

diff --git a/Engine/Source/Developer/TraceServices/Private/Analyzers/TasksAnalysis.cpp b/Engine/Source/Developer/TraceServices/Private/Analyzers/TasksAnalysis.cpp
--- a/Engine/Source/Developer/TraceServices/Private/Analyzers/TasksAnalysis.cpp
+++ b/Engine/Source/Developer/TraceServices/Private/Analyzers/TasksAnalysis.cpp
@@ -17,6 +17,7 @@ namespace TraceServices
 	{
 		auto& Builder = Context.InterfaceBuilder;
 
+		Builder.RouteEvent(RouteId_Init, "TaskTrace", "Init");
 		Builder.RouteEvent(RouteId_Created, "TaskTrace", "Created");
 		Builder.RouteEvent(RouteId_Launched, "TaskTrace", "Launched");
 		Builder.RouteEvent(RouteId_Scheduled, "TaskTrace", "Scheduled");
@@ -40,7 +41,7 @@ namespace TraceServices
 		uint32 ThreadId = Context.ThreadInfo.GetId();
 		switch (RouteId)
 		{
-		case RouterId_Init:
+		case RouteId_Init:
 		{
 			uint32 Version = EventData.GetValue<uint32>("Version");
 			TasksProvider.Init(Version);
diff --git a/Engine/Source/Developer/TraceServices/Private/Analyzers/TasksAnalysis.h b/Engine/Source/Developer/TraceServices/Private/Analyzers/TasksAnalysis.h
--- a/Engine/Source/Developer/TraceServices/Private/Analyzers/TasksAnalysis.h
+++ b/Engine/Source/Developer/TraceServices/Private/Analyzers/TasksAnalysis.h
@@ -20,7 +20,7 @@ namespace TraceServices
 	private:
 		enum : uint16
 		{
-			RouterId_Init,
+			RouteId_Init,
 			RouteId_Created,
 			RouteId_Launched,
 			RouteId_Scheduled,
diff --git a/Engine/Source/Runtime/Core/Private/Async/TaskTrace.cpp b/Engine/Source/Runtime/Core/Private/Async/TaskTrace.cpp
--- a/Engine/Source/Runtime/Core/Private/Async/TaskTrace.cpp
+++ b/Engine/Source/Runtime/Core/Private/Async/TaskTrace.cpp
@@ -13,7 +13,7 @@ namespace TaskTrace
 {
 	UE_TRACE_CHANNEL(TaskChannel);
 
-	UE_TRACE_EVENT_BEGIN(TaskTrace, InitTaskTrace)
+	UE_TRACE_EVENT_BEGIN(TaskTrace, Init)
 		UE_TRACE_EVENT_FIELD(uint32, Version)
 	UE_TRACE_EVENT_END()
 
@@ -81,8 +81,8 @@ namespace TaskTrace
 
 	void Init()
 	{
-		UE_TRACE_LOG(TaskTrace, InitTaskTrace, TaskChannel)
-			<< InitTaskTrace.Version(0);
+		UE_TRACE_LOG(TaskTrace, Init, TaskChannel)
+			<< Init.Version(0);
 
 		bGTaskTraceInitialized = true;
 	}
