[Insights] Fix crash on platform that doesn't have module provider yet

For some platforms/module types a provider cannot be created.

#rnx

diff --git a/Engine/Source/Developer/TraceServices/Private/Analyzers/ModuleAnalysis.cpp b/Engine/Source/Developer/TraceServices/Private/Analyzers/ModuleAnalysis.cpp
--- a/Engine/Source/Developer/TraceServices/Private/Analyzers/ModuleAnalysis.cpp
+++ b/Engine/Source/Developer/TraceServices/Private/Analyzers/ModuleAnalysis.cpp
@@ -45,7 +45,10 @@ bool FModuleAnalyzer::OnEvent(uint16 RouteId, EStyle Style, const FOnEventContex
 				{
 					check(Provider == nullptr); // Should only get one init message
 					Provider = CreateModuleProvider(Session, SymbolFormat);
-					Session.AddProvider(TEXT("ModuleProvider"), Provider);
+					if (Provider)
+					{
+						Session.AddProvider(TEXT("ModuleProvider"), Provider);
+					}
 				}
 			}
 			break;
@@ -59,7 +62,10 @@ bool FModuleAnalyzer::OnEvent(uint16 RouteId, EStyle Style, const FOnEventContex
 				FStringView ModuleName;
 				if (Context.EventData.GetString("Name", ModuleName))
 				{
-					Provider->OnModuleLoad(ModuleName, Base, Size);
+					if (Provider)
+					{
+						Provider->OnModuleLoad(ModuleName, Base, Size);
+					}
 				}
 			}
 			break;
@@ -67,7 +73,10 @@ bool FModuleAnalyzer::OnEvent(uint16 RouteId, EStyle Style, const FOnEventContex
 		case RouteId_ModuleUnload:
 			{
 				const uint32 Base = Context.EventData.GetValue<uint32>("Base");
-				Provider->OnModuleUnload(Base);
+				if (Provider)
+				{
+					Provider->OnModuleUnload(Base);
+				}
 			}
 			break;
 	}
