[Insights] Add mismatch query result type and consolidate strings

#rnx

diff --git a/Engine/Source/Developer/TraceInsights/Private/Insights/MemoryProfiler/ViewModels/MemAllocGroupingByCallstack.cpp b/Engine/Source/Developer/TraceInsights/Private/Insights/MemoryProfiler/ViewModels/MemAllocGroupingByCallstack.cpp
--- a/Engine/Source/Developer/TraceInsights/Private/Insights/MemoryProfiler/ViewModels/MemAllocGroupingByCallstack.cpp
+++ b/Engine/Source/Developer/TraceInsights/Private/Insights/MemoryProfiler/ViewModels/MemAllocGroupingByCallstack.cpp
@@ -58,8 +58,6 @@ void FMemAllocGroupingByCallstack::GroupNodes(const TArray<FTableTreeNodePtr>& N
 
 	FTableTreeNodePtr UnsetGroupPtr = nullptr;
 
-	static FName PendingName(TEXT("Pending..."));
-	static FName NotFoundName(TEXT("Not found"));
 	static FName NotAvailableName(TEXT("N/A"));
 
 	for (FTableTreeNodePtr NodePtr : Nodes)
@@ -98,21 +96,20 @@ void FMemAllocGroupingByCallstack::GroupNodes(const TArray<FTableTreeNodePtr>& N
 						GroupPtr->GroupMap.Add(Frame->Addr, NewGroupPtr);
 
 						FName GroupName = NotAvailableName;
-						const TraceServices::QueryResult Result = Frame->Symbol->Result.load(std::memory_order_acquire);
+						const TraceServices::QueryResult Result = Frame->Symbol->GetResult();
 						switch (Result)
 						{
-						case TraceServices::QueryResult::QR_NotLoaded:
+						case TraceServices::QueryResult::Pending:
 							//GroupName = PendingName;
 							GroupName = FName(*FString::Printf(TEXT("0x%X (...)"), Frame->Addr));
 							break;
-						case TraceServices::QueryResult::QR_NotFound:
-							//GroupName = NotFoundName;
-							GroupName = FName(*FString::Printf(TEXT("0x%X"), Frame->Addr));
-							break;
-						case TraceServices::QueryResult::QR_OK:
+						case TraceServices::QueryResult::OK:
 							GroupName = FName(Frame->Symbol->Name);
 							//GroupName = FName(*FString::Printf(TEXT("0x%X %s"), Frame->Addr, Frame->Symbol->Name));
 							break;
+						default:
+							GroupName = FName(*FString::Printf(TEXT("0x%X"), Frame->Addr));
+							break;
 						}
 
 						FTableTreeNodePtr NewGroupNodePtr = MakeShared<FTableTreeNode>(GroupName, InParentTable);
diff --git a/Engine/Source/Developer/TraceInsights/Private/Insights/MemoryProfiler/ViewModels/MemAllocTable.cpp b/Engine/Source/Developer/TraceInsights/Private/Insights/MemoryProfiler/ViewModels/MemAllocTable.cpp
--- a/Engine/Source/Developer/TraceInsights/Private/Insights/MemoryProfiler/ViewModels/MemAllocTable.cpp
+++ b/Engine/Source/Developer/TraceInsights/Private/Insights/MemoryProfiler/ViewModels/MemAllocTable.cpp
@@ -515,13 +515,8 @@ void FMemAllocTable::AddDefaultColumns()
 					const FMemoryAlloc* Alloc = MemAllocNode.GetMemAlloc();
 					if (Alloc)
 					{
-						static const TCHAR* DisplayStrings[] = {
-							TEXT("Pending..."),
-							TEXT("Not found"),
-							TEXT("N/A"),
-						};
-
-						const TCHAR* Value = DisplayStrings[2]; // not available
+						static const TCHAR* NotAvailable = TEXT("N/A");
+						const TCHAR* Value = NotAvailable;
 
 						const TraceServices::FCallstack* Callstack = Alloc->GetCallstack();
 
@@ -531,18 +526,14 @@ void FMemAllocTable::AddDefaultColumns()
 							const TraceServices::FStackFrame* Frame = Callstack->Frame(FMath::Min(2u, Callstack->Num() - 1));
 							check(Frame != nullptr);
 
-							const TraceServices::QueryResult Result = Frame->Symbol->Result.load(std::memory_order_acquire);
-							switch (Result)
+							const TraceServices::QueryResult Result = Frame->Symbol->GetResult();
+							if (Result == TraceServices::QueryResult::OK)
 							{
-							case TraceServices::QueryResult::QR_NotLoaded:
-								Value = DisplayStrings[0]; // pending
-								break;
-							case TraceServices::QueryResult::QR_NotFound:
-								Value = DisplayStrings[1]; // not found
-								break;
-							case TraceServices::QueryResult::QR_OK:
 								Value = Frame->Symbol->Name;
-								break;
+							}
+							else
+							{
+								Value = TraceServices::QueryResultToString(Result);
 							}
 						}
 
diff --git a/Engine/Source/Developer/TraceInsights/Private/Insights/MemoryProfiler/ViewModels/MemoryAlloc.cpp b/Engine/Source/Developer/TraceInsights/Private/Insights/MemoryProfiler/ViewModels/MemoryAlloc.cpp
--- a/Engine/Source/Developer/TraceInsights/Private/Insights/MemoryProfiler/ViewModels/MemoryAlloc.cpp
+++ b/Engine/Source/Developer/TraceInsights/Private/Insights/MemoryProfiler/ViewModels/MemoryAlloc.cpp
@@ -27,12 +27,6 @@ FMemoryAlloc::~FMemoryAlloc()
 
 FText FMemoryAlloc::GetFullCallstack() const
 {
-	static const TCHAR* DisplayStrings[] = {
-	TEXT("Pending..."),
-	TEXT("Not found"),
-	TEXT("N/A"),
-	};
-
 	FString Tooltip;
 	if (Callstack)
 	{
@@ -42,18 +36,14 @@ FText FMemoryAlloc::GetFullCallstack() const
 			const TraceServices::FStackFrame* Frame = Callstack->Frame(Index);
 			check(Frame != nullptr);
 
-			const TraceServices::QueryResult Result = Frame->Symbol->Result.load(std::memory_order_acquire);
-			switch (Result)
+			const TraceServices::QueryResult Result = Frame->Symbol->GetResult();
+			if (Result == TraceServices::QueryResult::OK)
 			{
-			case TraceServices::QueryResult::QR_NotLoaded:
-				Tooltip.Appendf(TEXT("%s\n"), DisplayStrings[0]); // pending
-				break;
-			case TraceServices::QueryResult::QR_NotFound:
-				Tooltip.Appendf(TEXT("%s\n"), DisplayStrings[1]); // not found
-				break;
-			case TraceServices::QueryResult::QR_OK:
 				Tooltip.Appendf(TEXT("%s\n"), Frame->Symbol->Name);
-				break;
+			}
+			else
+			{
+				Tooltip.Appendf(TEXT("%s\n"), QueryResultToString(Result));
 			}
 		}
 
@@ -63,7 +53,7 @@ FText FMemoryAlloc::GetFullCallstack() const
 
 	if (Tooltip.Len() == 0)
 	{
-		Tooltip.Append(DisplayStrings[2]);
+		Tooltip.Append(TEXT("N/A"));
 	}
 
 	return FText::FromString(Tooltip);
diff --git a/Engine/Source/Developer/TraceServices/Private/Model/ModuleProvider.cpp b/Engine/Source/Developer/TraceServices/Private/Model/ModuleProvider.cpp
--- a/Engine/Source/Developer/TraceServices/Private/Model/ModuleProvider.cpp
+++ b/Engine/Source/Developer/TraceServices/Private/Model/ModuleProvider.cpp
@@ -98,7 +98,7 @@ const FResolvedSymbol* TModuleProvider<SymbolProvider>::GetSymbol(uint64 Address
 			return SymbolCache[Address];
 		}
 		ResolvedSymbol = &Symbols.EmplaceBack();
-		ResolvedSymbol->Result = QR_NotLoaded;
+		ResolvedSymbol->Result = QueryResult::NotLoaded;
 		ResolvedSymbol->Name = nullptr;
 		ResolvedSymbol->FileAndLine = nullptr;
 		SymbolCache.Add(Address, ResolvedSymbol);
diff --git a/Engine/Source/Developer/TraceServices/Private/Model/PbdSymbols.inl b/Engine/Source/Developer/TraceServices/Private/Model/PbdSymbols.inl
--- a/Engine/Source/Developer/TraceServices/Private/Model/PbdSymbols.inl
+++ b/Engine/Source/Developer/TraceServices/Private/Model/PbdSymbols.inl
@@ -187,7 +187,7 @@ private:
 		{
 			DWORD Err = GetLastError();
 			SymbolsFailed++;
-			UpdateResolvedSymbol(Target, QR_NotFound, UNKNOWN_MODULE_TEXT, UNKNOWN_MODULE_TEXT);
+			UpdateResolvedSymbol(Target, QueryResult::NotFound, UNKNOWN_MODULE_TEXT, UNKNOWN_MODULE_TEXT);
 			return;
 		}
 
@@ -203,7 +203,7 @@ private:
 		if (!SymGetLineFromAddr(Handle, Address, &dwDisplacement, &Line))
 		{
 			SymbolsFailed++;
-			UpdateResolvedSymbol(Target, QR_OK, SymbolNameStr, UNKNOWN_MODULE_TEXT);
+			UpdateResolvedSymbol(Target, QueryResult::OK, SymbolNameStr, UNKNOWN_MODULE_TEXT);
 			return;
 		}
 
@@ -212,7 +212,7 @@ private:
 		const TCHAR* FileAndLineStr = Session.StoreString(FStringView(FileAndLine));
 		
 		SymbolsResolved++;
-		UpdateResolvedSymbol(Target, QR_OK, SymbolNameStr, FileAndLineStr);
+		UpdateResolvedSymbol(Target, QueryResult::OK, SymbolNameStr, FileAndLineStr);
 	}
 
 	bool LoadModuleSymbols(uint64 Base, uint64 Size, const TCHAR* Path)
diff --git a/Engine/Source/Developer/TraceServices/Public/TraceServices/Model/Modules.h b/Engine/Source/Developer/TraceServices/Public/TraceServices/Model/Modules.h
--- a/Engine/Source/Developer/TraceServices/Public/TraceServices/Model/Modules.h
+++ b/Engine/Source/Developer/TraceServices/Public/TraceServices/Model/Modules.h
@@ -11,13 +11,33 @@ namespace TraceServices {
   * Result of a query. Since symbol resolving can be deferred this signals if a
   * symbol has been resolved, waiting to be resolved or wasn't found at all.
   */
-enum QueryResult 
+enum class QueryResult : uint8
 {
-	QR_OK,
-	QR_NotFound,
-	QR_NotLoaded,
+	Pending,		// Symbol is pending resolution
+	OK,			// Symbol has been correctly resolved
+	NotLoaded,	// Module debug data could not be loaded or found
+	Mismatch,	// Module debug data could not be loaded because debug data did not match traced binary
+	NotFound,	// Symbol was not found in module debug data
+	StatusNum
 };
 
+////////////////////////////////////////////////////////////////////////////////
+/**
+ * Helper method to get a string representation of the query result.
+ */
+inline const TCHAR* QueryResultToString(QueryResult Result)
+{
+	static const TCHAR* DisplayStrings[] = {
+		TEXT("Pending..."),
+		TEXT("Ok"),
+		TEXT("Not loaded"),
+		TEXT("Mismatch"),
+		TEXT("Not found")
+	};
+	static_assert(UE_ARRAY_COUNT(DisplayStrings) == (uint8) QueryResult::StatusNum, "Missing QueryResult");
+	return DisplayStrings[(uint8)Result];
+}
+
 ////////////////////////////////////////////////////////////////////////////////
 /**
   * Represent a resolved symbol. The resolve status and string values may change
@@ -29,6 +49,11 @@ struct FResolvedSymbol
 	std::atomic<QueryResult> Result;
 	const TCHAR* Name;
 	const TCHAR* FileAndLine;
+
+	inline QueryResult GetResult() const
+	{
+		return Result.load(std::memory_order_acquire);
+	}
 };
 
 ////////////////////////////////////////////////////////////////////////////////
@@ -53,4 +78,5 @@ public:
 };
 
 ////////////////////////////////////////////////////////////////////////////////
+
 } // namespace TraceServices
