[Insights] Handle cases where module provider doesn't exist.

Handle cases where module provider doesn't exist, for example when using the application on a platform that doesn't supportt the required symbol resolving. Symbol field then falls back to a shared const value.

diff --git a/Engine/Source/Developer/TraceInsights/Private/Insights/MemoryProfiler/Widgets/SMemAllocTableTreeView.cpp b/Engine/Source/Developer/TraceInsights/Private/Insights/MemoryProfiler/Widgets/SMemAllocTableTreeView.cpp
--- a/Engine/Source/Developer/TraceInsights/Private/Insights/MemoryProfiler/Widgets/SMemAllocTableTreeView.cpp
+++ b/Engine/Source/Developer/TraceInsights/Private/Insights/MemoryProfiler/Widgets/SMemAllocTableTreeView.cpp
@@ -678,9 +678,15 @@ FText SMemAllocTableTreeView::GetSymbolResolutionStatus() const
 	auto ModuleProvider = Session->ReadProvider<IModuleProvider>(FName("ModuleProvider"));
 
 	IModuleProvider::FStats Stats;
-	ModuleProvider->GetStats(&Stats);
-
-	return FText::Format(LOCTEXT("SymbolsResolved", "{0} / {1} symbols resolved. {2} failed."), Stats.SymbolsResolved, Stats.SymbolsDiscovered, Stats.SymbolsFailed);
+	if (ModuleProvider)
+	{
+		ModuleProvider->GetStats(&Stats);
+		return FText::Format(LOCTEXT("SymbolsResolved", "{0} / {1} symbols resolved. {2} failed."), Stats.SymbolsResolved, Stats.SymbolsDiscovered, Stats.SymbolsFailed);
+	}
+	else
+	{
+		return LOCTEXT("SymbolsResolutionNotPossible", "No symbols resolution possible.");
+	}
 }
 
 ////////////////////////////////////////////////////////////////////////////////////////////////////
diff --git a/Engine/Source/Developer/TraceServices/Private/Model/CallstacksProvider.cpp b/Engine/Source/Developer/TraceServices/Private/Model/CallstacksProvider.cpp
--- a/Engine/Source/Developer/TraceServices/Private/Model/CallstacksProvider.cpp
+++ b/Engine/Source/Developer/TraceServices/Private/Model/CallstacksProvider.cpp
@@ -11,6 +11,13 @@
 namespace TraceServices
 {
 
+/////////////////////////////////////////////////////////////////////
+static const FResolvedSymbol GNeverResolveSymbol = {
+	QueryResult::NotLoaded,
+	nullptr,
+	nullptr
+};
+
 /////////////////////////////////////////////////////////////////////
 #ifdef TRACE_CALLSTACK_STATS
 static struct FCallstackProviderStats
@@ -66,9 +73,17 @@ void FCallstacksProvider::AddCallstack(uint64 InCallstackId, const uint64* InFra
 	{
 		FStackFrame& F = Frames.PushBack();
 		F.Addr = InFrames[FrameIdx];
-		// This will return immediately. The result will be empty if the symbol 
-		// has not been encountered before, and resolution has been queued up.
-		F.Symbol = ModuleProvider->GetSymbol(InFrames[FrameIdx]);
+		
+		if (ModuleProvider)
+		{
+			// This will return immediately. The result will be empty if the symbol 
+			// has not been encountered before, and resolution has been queued up.
+			F.Symbol = ModuleProvider->GetSymbol(InFrames[FrameIdx]);
+		}
+		else
+		{
+			F.Symbol = &GNeverResolveSymbol;
+		}
 	}
 
 	{
