[Insights] Fix initialization issue on clang.

#rnx

diff --git a/Engine/Source/Developer/TraceServices/Private/Model/CallstacksProvider.cpp b/Engine/Source/Developer/TraceServices/Private/Model/CallstacksProvider.cpp
--- a/Engine/Source/Developer/TraceServices/Private/Model/CallstacksProvider.cpp
+++ b/Engine/Source/Developer/TraceServices/Private/Model/CallstacksProvider.cpp
@@ -12,11 +12,7 @@ namespace TraceServices
 {
 
 /////////////////////////////////////////////////////////////////////
-static const FResolvedSymbol GNeverResolveSymbol = {
-	QueryResult::NotLoaded,
-	nullptr,
-	nullptr
-};
+static FResolvedSymbol GNeverResolveSymbol;
 
 /////////////////////////////////////////////////////////////////////
 #ifdef TRACE_CALLSTACK_STATS
@@ -35,6 +31,9 @@ FCallstacksProvider::FCallstacksProvider(IAnalysisSession& InSession)
 	, Callstacks(InSession.GetLinearAllocator(), CallstacksPerPage)
 	, Frames(InSession.GetLinearAllocator(), FramesPerPage)
 {
+	GNeverResolveSymbol.Result.store(QueryResult::NotLoaded);
+	GNeverResolveSymbol.Name = nullptr;
+	GNeverResolveSymbol.FileAndLine = nullptr;
 }
 
 /////////////////////////////////////////////////////////////////////
