TaskTrace: added "execution time" counter

diff --git a/Engine/Source/Developer/TraceInsights/Private/Insights/ViewModels/ThreadTimingTrack.cpp b/Engine/Source/Developer/TraceInsights/Private/Insights/ViewModels/ThreadTimingTrack.cpp
--- a/Engine/Source/Developer/TraceInsights/Private/Insights/ViewModels/ThreadTimingTrack.cpp
+++ b/Engine/Source/Developer/TraceInsights/Private/Insights/ViewModels/ThreadTimingTrack.cpp
@@ -1042,7 +1042,8 @@ void FThreadTimingTrack::InitTooltip(FTooltipDrawState& InOutTooltip, const ITim
 			
 			InOutTooltip.AddNameValueTextLine(TEXT("Created:"), FString::Printf(TEXT("%f on %s"), Task.CreatedTimestamp, *SharedState.GetCpuTrack(Task.CreatedThreadId)->GetName()));
 			InOutTooltip.AddNameValueTextLine(TEXT("Launched:"), FString::Printf(TEXT("%f (+%s) on %s"), Task.LaunchedTimestamp, *TimeUtils::FormatTimeAuto(Task.LaunchedTimestamp - Task.CreatedTimestamp), *SharedState.GetCpuTrack(Task.LaunchedThreadId)->GetName()));
-			InOutTooltip.AddNameValueTextLine(TEXT("Started:"), FString::Printf(TEXT("%f (+%s)"), Task.StartedTimestamp, *TimeUtils::FormatTimeAuto(Task.StartedTimestamp - Task.LaunchedTimestamp)));
+			InOutTooltip.AddNameValueTextLine(TEXT("Scheduled:"), FString::Printf(TEXT("%f (+%s) on %s"), Task.ScheduledTimestamp, *TimeUtils::FormatTimeAuto(Task.ScheduledTimestamp - Task.LaunchedTimestamp), *SharedState.GetCpuTrack(Task.ScheduledThreadId)->GetName()));
+			InOutTooltip.AddNameValueTextLine(TEXT("Started:"), FString::Printf(TEXT("%f (+%s)"), Task.StartedTimestamp, *TimeUtils::FormatTimeAuto(Task.StartedTimestamp - Task.ScheduledTimestamp)));
 			if (Task.FinishedTimestamp != TraceServices::FTaskInfo::InvalidTimestamp)
 			{
 				InOutTooltip.AddNameValueTextLine(TEXT("Finished:"), FString::Printf(TEXT("%f (+%s)"), Task.FinishedTimestamp, *TimeUtils::FormatTimeAuto(Task.FinishedTimestamp - Task.StartedTimestamp)));
diff --git a/Engine/Source/Developer/TraceServices/Private/Model/TasksProfiler.cpp b/Engine/Source/Developer/TraceServices/Private/Model/TasksProfiler.cpp
--- a/Engine/Source/Developer/TraceServices/Private/Model/TasksProfiler.cpp
+++ b/Engine/Source/Developer/TraceServices/Private/Model/TasksProfiler.cpp
@@ -84,6 +84,11 @@ namespace TraceServices
 		RunningTasksCounter->SetName(TEXT("Tasks::RunningTasks"));
 		RunningTasksCounter->SetDescription(TEXT("Tasks: level of parallelism - the number of tasks being executed"));
 		RunningTasksCounter->SetIsFloatingPoint(false);
+
+		ExecutionTimeCounter = CounterProvider.CreateCounter();
+		ExecutionTimeCounter->SetName(TEXT("Tasks::ExecutionTime"));
+		ExecutionTimeCounter->SetDescription(TEXT("Tasks: execution time"));
+		ExecutionTimeCounter->SetIsFloatingPoint(true);
 	}
 
 	void FTasksProvider::Init(uint32 InVersion)
@@ -188,6 +193,7 @@ namespace TraceServices
 		ExecutionThreads.FindOrAdd(ThreadId).Add(TaskId);
 
 		FTaskInfo* Task = TryGetTask(TaskId);
+		check(Task != nullptr);
 
 		if (IsNamedThread(Task->ThreadToExecuteOn))
 		{
@@ -199,8 +205,7 @@ namespace TraceServices
 		}
 		RunningTasksCounter->SetValue(Timestamp, ++RunningTasksNum);
 
-		double LatencyMicrosecs = Task->StartedTimestamp * 1000000 - Task->ScheduledTimestamp * 1000000;
-		TaskLatencyCounter->SetValue(Timestamp, LatencyMicrosecs);
+		TaskLatencyCounter->SetValue(Timestamp, Task->StartedTimestamp - Task->ScheduledTimestamp);
 	}
 
 	void FTasksProvider::NestedAdded(TaskTrace::FId TaskId, TaskTrace::FId NestedId, double Timestamp, uint32 ThreadId)
@@ -216,6 +221,10 @@ namespace TraceServices
 		}
 
 		RunningTasksCounter->SetValue(Timestamp, --RunningTasksNum);
+	
+		FTaskInfo* Task = TryGetTask(TaskId);
+		check(Task != nullptr);
+		ExecutionTimeCounter->SetValue(Timestamp, Task->FinishedTimestamp - Task->StartedTimestamp);
 	}
 
 	void FTasksProvider::TaskCompleted(TaskTrace::FId TaskId, double Timestamp, uint32 ThreadId)
diff --git a/Engine/Source/Developer/TraceServices/Private/Model/TasksProfilerPrivate.h b/Engine/Source/Developer/TraceServices/Private/Model/TasksProfilerPrivate.h
--- a/Engine/Source/Developer/TraceServices/Private/Model/TasksProfilerPrivate.h
+++ b/Engine/Source/Developer/TraceServices/Private/Model/TasksProfilerPrivate.h
@@ -92,5 +92,6 @@ namespace TraceServices
 		IEditableCounter* ScheduledTasksCounter;
 		IEditableCounter* NamedThreadsScheduledTasksCounter;
 		IEditableCounter* RunningTasksCounter;
+		IEditableCounter* ExecutionTimeCounter;
 	};
 }
\ No newline at end of file
