[Insights] Rename QueryResult to ESymbolQueryResult

#rnx

diff --git a/Engine/Source/Developer/TraceInsights/Private/Insights/MemoryProfiler/ViewModels/MemAllocGroupingByCallstack.cpp b/Engine/Source/Developer/TraceInsights/Private/Insights/MemoryProfiler/ViewModels/MemAllocGroupingByCallstack.cpp
--- a/Engine/Source/Developer/TraceInsights/Private/Insights/MemoryProfiler/ViewModels/MemAllocGroupingByCallstack.cpp
+++ b/Engine/Source/Developer/TraceInsights/Private/Insights/MemoryProfiler/ViewModels/MemAllocGroupingByCallstack.cpp
@@ -96,14 +96,14 @@ void FMemAllocGroupingByCallstack::GroupNodes(const TArray<FTableTreeNodePtr>& N
 						GroupPtr->GroupMap.Add(Frame->Addr, NewGroupPtr);
 
 						FName GroupName = NotAvailableName;
-						const TraceServices::QueryResult Result = Frame->Symbol->GetResult();
+						const TraceServices::ESymbolQueryResult Result = Frame->Symbol->GetResult();
 						switch (Result)
 						{
-						case TraceServices::QueryResult::Pending:
+						case TraceServices::ESymbolQueryResult::Pending:
 							//GroupName = PendingName;
 							GroupName = FName(*FString::Printf(TEXT("0x%X (...)"), Frame->Addr));
 							break;
-						case TraceServices::QueryResult::OK:
+						case TraceServices::ESymbolQueryResult::OK:
 							GroupName = FName(Frame->Symbol->Name);
 							//GroupName = FName(*FString::Printf(TEXT("0x%X %s"), Frame->Addr, Frame->Symbol->Name));
 							break;
diff --git a/Engine/Source/Developer/TraceInsights/Private/Insights/MemoryProfiler/ViewModels/MemAllocTable.cpp b/Engine/Source/Developer/TraceInsights/Private/Insights/MemoryProfiler/ViewModels/MemAllocTable.cpp
--- a/Engine/Source/Developer/TraceInsights/Private/Insights/MemoryProfiler/ViewModels/MemAllocTable.cpp
+++ b/Engine/Source/Developer/TraceInsights/Private/Insights/MemoryProfiler/ViewModels/MemAllocTable.cpp
@@ -526,8 +526,8 @@ void FMemAllocTable::AddDefaultColumns()
 							const TraceServices::FStackFrame* Frame = Callstack->Frame(FMath::Min(2u, Callstack->Num() - 1));
 							check(Frame != nullptr);
 
-							const TraceServices::QueryResult Result = Frame->Symbol->GetResult();
-							if (Result == TraceServices::QueryResult::OK)
+							const TraceServices::ESymbolQueryResult Result = Frame->Symbol->GetResult();
+							if (Result == TraceServices::ESymbolQueryResult::OK)
 							{
 								Value = Frame->Symbol->Name;
 							}
diff --git a/Engine/Source/Developer/TraceInsights/Private/Insights/MemoryProfiler/ViewModels/MemoryAlloc.cpp b/Engine/Source/Developer/TraceInsights/Private/Insights/MemoryProfiler/ViewModels/MemoryAlloc.cpp
--- a/Engine/Source/Developer/TraceInsights/Private/Insights/MemoryProfiler/ViewModels/MemoryAlloc.cpp
+++ b/Engine/Source/Developer/TraceInsights/Private/Insights/MemoryProfiler/ViewModels/MemoryAlloc.cpp
@@ -36,8 +36,8 @@ FText FMemoryAlloc::GetFullCallstack() const
 			const TraceServices::FStackFrame* Frame = Callstack->Frame(Index);
 			check(Frame != nullptr);
 
-			const TraceServices::QueryResult Result = Frame->Symbol->GetResult();
-			if (Result == TraceServices::QueryResult::OK)
+			const TraceServices::ESymbolQueryResult Result = Frame->Symbol->GetResult();
+			if (Result == TraceServices::ESymbolQueryResult::OK)
 			{
 				Tooltip.Appendf(TEXT("%s\n"), Frame->Symbol->Name);
 			}
diff --git a/Engine/Source/Developer/TraceInsights/Private/Insights/MemoryProfiler/Widgets/SMemAllocTableTreeView.cpp b/Engine/Source/Developer/TraceInsights/Private/Insights/MemoryProfiler/Widgets/SMemAllocTableTreeView.cpp
--- a/Engine/Source/Developer/TraceInsights/Private/Insights/MemoryProfiler/Widgets/SMemAllocTableTreeView.cpp
+++ b/Engine/Source/Developer/TraceInsights/Private/Insights/MemoryProfiler/Widgets/SMemAllocTableTreeView.cpp
@@ -685,7 +685,7 @@ FText SMemAllocTableTreeView::GetSymbolResolutionStatus() const
 	}
 	else
 	{
-		return LOCTEXT("SymbolsResolutionNotPossible", "No symbols resolution possible.");
+		return LOCTEXT("SymbolsResolutionNotPossible", "Symbol resolution was not possible.");
 	}
 }
 
diff --git a/Engine/Source/Developer/TraceServices/Private/Model/CallstacksProvider.cpp b/Engine/Source/Developer/TraceServices/Private/Model/CallstacksProvider.cpp
--- a/Engine/Source/Developer/TraceServices/Private/Model/CallstacksProvider.cpp
+++ b/Engine/Source/Developer/TraceServices/Private/Model/CallstacksProvider.cpp
@@ -31,7 +31,7 @@ FCallstacksProvider::FCallstacksProvider(IAnalysisSession& InSession)
 	, Callstacks(InSession.GetLinearAllocator(), CallstacksPerPage)
 	, Frames(InSession.GetLinearAllocator(), FramesPerPage)
 {
-	GNeverResolveSymbol.Result.store(QueryResult::NotLoaded);
+	GNeverResolveSymbol.Result.store(ESymbolQueryResult::NotLoaded);
 	GNeverResolveSymbol.Name = nullptr;
 	GNeverResolveSymbol.FileAndLine = nullptr;
 }
diff --git a/Engine/Source/Developer/TraceServices/Private/Model/ModuleProvider.cpp b/Engine/Source/Developer/TraceServices/Private/Model/ModuleProvider.cpp
--- a/Engine/Source/Developer/TraceServices/Private/Model/ModuleProvider.cpp
+++ b/Engine/Source/Developer/TraceServices/Private/Model/ModuleProvider.cpp
@@ -98,7 +98,7 @@ const FResolvedSymbol* TModuleProvider<SymbolProvider>::GetSymbol(uint64 Address
 			return SymbolCache[Address];
 		}
 		ResolvedSymbol = &Symbols.EmplaceBack();
-		ResolvedSymbol->Result = QueryResult::NotLoaded;
+		ResolvedSymbol->Result = ESymbolQueryResult::NotLoaded;
 		ResolvedSymbol->Name = nullptr;
 		ResolvedSymbol->FileAndLine = nullptr;
 		SymbolCache.Add(Address, ResolvedSymbol);
diff --git a/Engine/Source/Developer/TraceServices/Private/Model/PbdSymbols.inl b/Engine/Source/Developer/TraceServices/Private/Model/PbdSymbols.inl
--- a/Engine/Source/Developer/TraceServices/Private/Model/PbdSymbols.inl
+++ b/Engine/Source/Developer/TraceServices/Private/Model/PbdSymbols.inl
@@ -166,7 +166,7 @@ private:
 		bRunWorkerThread = false; 
 	}
 
-	void UpdateResolvedSymbol(FResolvedSymbol* Symbol, QueryResult Result, const TCHAR* Name, const TCHAR* FileAndLine)
+	void UpdateResolvedSymbol(FResolvedSymbol* Symbol, ESymbolQueryResult Result, const TCHAR* Name, const TCHAR* FileAndLine)
 	{
 		Symbol->Name = Name;
 		Symbol->FileAndLine = FileAndLine;
@@ -187,7 +187,7 @@ private:
 		{
 			DWORD Err = GetLastError();
 			SymbolsFailed++;
-			UpdateResolvedSymbol(Target, QueryResult::NotFound, UNKNOWN_MODULE_TEXT, UNKNOWN_MODULE_TEXT);
+			UpdateResolvedSymbol(Target, ESymbolQueryResult::NotFound, UNKNOWN_MODULE_TEXT, UNKNOWN_MODULE_TEXT);
 			return;
 		}
 
@@ -203,7 +203,7 @@ private:
 		if (!SymGetLineFromAddr(Handle, Address, &dwDisplacement, &Line))
 		{
 			SymbolsFailed++;
-			UpdateResolvedSymbol(Target, QueryResult::OK, SymbolNameStr, UNKNOWN_MODULE_TEXT);
+			UpdateResolvedSymbol(Target, ESymbolQueryResult::OK, SymbolNameStr, UNKNOWN_MODULE_TEXT);
 			return;
 		}
 
@@ -212,7 +212,7 @@ private:
 		const TCHAR* FileAndLineStr = Session.StoreString(FStringView(FileAndLine));
 		
 		SymbolsResolved++;
-		UpdateResolvedSymbol(Target, QueryResult::OK, SymbolNameStr, FileAndLineStr);
+		UpdateResolvedSymbol(Target, ESymbolQueryResult::OK, SymbolNameStr, FileAndLineStr);
 	}
 
 	bool LoadModuleSymbols(uint64 Base, uint64 Size, const TCHAR* Path)
diff --git a/Engine/Source/Developer/TraceServices/Public/TraceServices/Model/Modules.h b/Engine/Source/Developer/TraceServices/Public/TraceServices/Model/Modules.h
--- a/Engine/Source/Developer/TraceServices/Public/TraceServices/Model/Modules.h
+++ b/Engine/Source/Developer/TraceServices/Public/TraceServices/Model/Modules.h
@@ -11,7 +11,7 @@ namespace TraceServices {
   * Result of a query. Since symbol resolving can be deferred this signals if a
   * symbol has been resolved, waiting to be resolved or wasn't found at all.
   */
-enum class QueryResult : uint8
+enum class ESymbolQueryResult : uint8
 {
 	Pending,		// Symbol is pending resolution
 	OK,			// Symbol has been correctly resolved
@@ -25,7 +25,7 @@ enum class QueryResult : uint8
 /**
  * Helper method to get a string representation of the query result.
  */
-inline const TCHAR* QueryResultToString(QueryResult Result)
+inline const TCHAR* QueryResultToString(ESymbolQueryResult Result)
 {
 	static const TCHAR* DisplayStrings[] = {
 		TEXT("Pending..."),
@@ -34,7 +34,7 @@ inline const TCHAR* QueryResultToString(QueryResult Result)
 		TEXT("Mismatch"),
 		TEXT("Not found")
 	};
-	static_assert(UE_ARRAY_COUNT(DisplayStrings) == (uint8) QueryResult::StatusNum, "Missing QueryResult");
+	static_assert(UE_ARRAY_COUNT(DisplayStrings) == (uint8) ESymbolQueryResult::StatusNum, "Missing QueryResult");
 	return DisplayStrings[(uint8)Result];
 }
 
@@ -46,11 +46,11 @@ inline const TCHAR* QueryResultToString(QueryResult Result)
   */
 struct FResolvedSymbol
 {
-	std::atomic<QueryResult> Result;
+	std::atomic<ESymbolQueryResult> Result;
 	const TCHAR* Name;
 	const TCHAR* FileAndLine;
 
-	inline QueryResult GetResult() const
+	inline ESymbolQueryResult GetResult() const
 	{
 		return Result.load(std::memory_order_acquire);
 	}
