[Insights] Initialize symbol entries as pending.

#rnx

diff --git a/Engine/Source/Developer/TraceServices/Private/Model/ModuleProvider.cpp b/Engine/Source/Developer/TraceServices/Private/Model/ModuleProvider.cpp
--- a/Engine/Source/Developer/TraceServices/Private/Model/ModuleProvider.cpp
+++ b/Engine/Source/Developer/TraceServices/Private/Model/ModuleProvider.cpp
@@ -98,7 +98,7 @@ const FResolvedSymbol* TModuleProvider<SymbolProvider>::GetSymbol(uint64 Address
 			return SymbolCache[Address];
 		}
 		ResolvedSymbol = &Symbols.EmplaceBack();
-		ResolvedSymbol->Result = ESymbolQueryResult::NotLoaded;
+		ResolvedSymbol->Result = ESymbolQueryResult::Pending;
 		ResolvedSymbol->Name = nullptr;
 		ResolvedSymbol->FileAndLine = nullptr;
 		SymbolCache.Add(Address, ResolvedSymbol);
