Unity fix

#rnx

diff --git a/Engine/Source/Developer/TraceServices/Public/TraceServices/Model/Modules.h b/Engine/Source/Developer/TraceServices/Public/TraceServices/Model/Modules.h
--- a/Engine/Source/Developer/TraceServices/Public/TraceServices/Model/Modules.h
+++ b/Engine/Source/Developer/TraceServices/Public/TraceServices/Model/Modules.h
@@ -2,6 +2,7 @@
 
 #pragma once
 #include "TraceServices/Model/AnalysisSession.h"
+#include "Templates/UnrealTemplate.h"
 #include <atomic>
 
 namespace TraceServices {
