Fix crash when toggling the task channel on

diff --git a/Engine/Source/Developer/TraceServices/Private/Model/TasksProfiler.cpp b/Engine/Source/Developer/TraceServices/Private/Model/TasksProfiler.cpp
--- a/Engine/Source/Developer/TraceServices/Private/Model/TasksProfiler.cpp
+++ b/Engine/Source/Developer/TraceServices/Private/Model/TasksProfiler.cpp
@@ -58,6 +58,7 @@ namespace TraceServices
 
 	void FTasksProvider::CreateCounters()
 	{
+		check(bCountersCreated == false);
 		FAnalysisSessionEditScope _(Session);
 
 		WaitingForPrerequisitesTasksCounter = CounterProvider.CreateCounter();
@@ -89,13 +90,18 @@ namespace TraceServices
 		ExecutionTimeCounter->SetName(TEXT("Tasks::ExecutionTime"));
 		ExecutionTimeCounter->SetDescription(TEXT("Tasks: execution time"));
 		ExecutionTimeCounter->SetIsFloatingPoint(true);
+
+		bCountersCreated = true;
 	}
 
 	void FTasksProvider::Init(uint32 InVersion)
 	{
 		Version = InVersion;
 
-		CreateCounters();
+		if (!bCountersCreated)
+		{
+			CreateCounters();
+		}
 	}
 
 	void FTasksProvider::TaskCreated(TaskTrace::FId TaskId, double Timestamp, uint32 ThreadId)
@@ -146,6 +152,11 @@ namespace TraceServices
 		Task->LaunchedTimestamp = Timestamp;
 		Task->LaunchedThreadId = ThreadId;
 
+		if (!bCountersCreated)
+		{
+			CreateCounters();
+		}
+
 		WaitingForPrerequisitesTasksCounter->SetValue(Timestamp, ++WaitingForPrerequisitesTasksNum);
 	}
 
@@ -156,6 +167,11 @@ namespace TraceServices
 			return;
 		}
 
+		if (!bCountersCreated)
+		{
+			CreateCounters();
+		}
+
 		WaitingForPrerequisitesTasksCounter->SetValue(Timestamp, --WaitingForPrerequisitesTasksNum);
 		if (IsNamedThread(TryGetTask(TaskId)->ThreadToExecuteOn))
 		{
@@ -195,6 +211,11 @@ namespace TraceServices
 		FTaskInfo* Task = TryGetTask(TaskId);
 		check(Task != nullptr);
 
+		if (!bCountersCreated)
+		{
+			CreateCounters();
+		}
+
 		if (IsNamedThread(Task->ThreadToExecuteOn))
 		{
 			NamedThreadsScheduledTasksCounter->SetValue(Timestamp, --ScheduledTasksNum);
@@ -220,6 +241,11 @@ namespace TraceServices
 			return;
 		}
 
+		if (!bCountersCreated)
+		{
+			CreateCounters();
+		}
+
 		RunningTasksCounter->SetValue(Timestamp, --RunningTasksNum);
 	
 		FTaskInfo* Task = TryGetTask(TaskId);
diff --git a/Engine/Source/Developer/TraceServices/Private/Model/TasksProfilerPrivate.h b/Engine/Source/Developer/TraceServices/Private/Model/TasksProfilerPrivate.h
--- a/Engine/Source/Developer/TraceServices/Private/Model/TasksProfilerPrivate.h
+++ b/Engine/Source/Developer/TraceServices/Private/Model/TasksProfilerPrivate.h
@@ -85,6 +85,8 @@ namespace TraceServices
 		int64 NamedThreadsScheduledTasksNum = 0;
 		// number of tasks being executed
 		int64 RunningTasksNum = 0;
+		// a bool that is set when the counters are created
+		bool bCountersCreated = false;
 
 		ICounterProvider& CounterProvider;
 		IEditableCounter* TaskLatencyCounter;
