[TraceLog] Fix wrapping memory markers.

diff --git a/Engine/Source/Developer/TraceServices/Private/Analyzers/AllocationsAnalysis.cpp b/Engine/Source/Developer/TraceServices/Private/Analyzers/AllocationsAnalysis.cpp
--- a/Engine/Source/Developer/TraceServices/Private/Analyzers/AllocationsAnalysis.cpp
+++ b/Engine/Source/Developer/TraceServices/Private/Analyzers/AllocationsAnalysis.cpp
@@ -40,6 +40,7 @@ void FAllocationsAnalyzer::OnAnalysisBegin(const FOnAnalysisContext& Context)
 	Builder.RouteEvent(RouteId_ReallocFree,  "Memory", "ReallocFree");
 	Builder.RouteEvent(RouteId_Marker,       "Memory", "Marker");
 	Builder.RouteEvent(RouteId_TagSpec,      "Memory", "TagSpec");
+	Builder.RouteEvent(RouteId_MarkerNewBase,"Memory", "MarkerNewBase");
 
 	Builder.RouteLoggerEvents(RouteId_MemScope, "Memory", true);
 }
@@ -141,7 +142,7 @@ bool FAllocationsAnalyzer::OnEvent(uint16 RouteId, EStyle Style, const FOnEventC
 			{
 				const double Seconds = Context.EventTime.AsSeconds(Cycle);
 				check(Seconds >= LastMarkerSeconds);
-				if (ensure(Seconds - LastMarkerSeconds < 60.0))
+				if (ensure((Seconds - LastMarkerSeconds < 60.0) || LastMarkerSeconds == 0.0f))
 				{
 					LastMarkerCycle = Cycle;
 					LastMarkerSeconds = Seconds;
@@ -157,6 +158,13 @@ bool FAllocationsAnalyzer::OnEvent(uint16 RouteId, EStyle Style, const FOnEventC
 			}
 			break;
 		}
+		case RouteId_MarkerNewBase:
+		{
+			const uint64 NewBaseCycle = EventData.GetValue<uint64>("BaseCycle");
+			check(NewBaseCycle > BaseCycle);
+			BaseCycle = NewBaseCycle;
+			break;
+		}
 		case RouteId_TagSpec:
 		{
 			const int32 Tag = Context.EventData.GetValue<int32>("Tag");
diff --git a/Engine/Source/Developer/TraceServices/Private/Analyzers/AllocationsAnalysis.h b/Engine/Source/Developer/TraceServices/Private/Analyzers/AllocationsAnalysis.h
--- a/Engine/Source/Developer/TraceServices/Private/Analyzers/AllocationsAnalysis.h
+++ b/Engine/Source/Developer/TraceServices/Private/Analyzers/AllocationsAnalysis.h
@@ -24,6 +24,7 @@ private:
 		RouteId_ReallocAlloc,
 		RouteId_ReallocFree,
 		RouteId_Marker,
+		RouteId_MarkerNewBase,
 		RouteId_TagSpec,
 		RouteId_MemScope,
 	};
diff --git a/Engine/Source/Runtime/Core/Private/ProfilingDebugging/Windows/MemoryTrace.cpp b/Engine/Source/Runtime/Core/Private/ProfilingDebugging/Windows/MemoryTrace.cpp
--- a/Engine/Source/Runtime/Core/Private/ProfilingDebugging/Windows/MemoryTrace.cpp
+++ b/Engine/Source/Runtime/Core/Private/ProfilingDebugging/Windows/MemoryTrace.cpp
@@ -107,6 +107,10 @@ UE_TRACE_EVENT_BEGIN(Memory, Marker)
 	UE_TRACE_EVENT_FIELD(uint32, Cycle)
 UE_TRACE_EVENT_END()
 
+UE_TRACE_EVENT_BEGIN(Memory, MarkerNewBase)
+	UE_TRACE_EVENT_FIELD(uint64, BaseCycle)
+UE_TRACE_EVENT_END()
+
 UE_TRACE_EVENT_BEGIN(Memory, CoreAdd)
 	UE_TRACE_EVENT_FIELD(uint64, Owner)
 	UE_TRACE_EVENT_FIELD(void*, Base)
@@ -331,6 +335,16 @@ void FAllocationTrace::Update()
 	if ((TheCount & MarkerSamplePeriod) == 0)
 	{
 		uint64 Cycle = FPlatformTime::Cycles64();
+
+		// We assume here that that the marker sample period is large
+		// enough that we don't have to syncronize across threads.
+		if ((Cycle - BaseCycle) > uint32(~0))
+		{
+			BaseCycle = Cycle;
+			UE_TRACE_LOG(Memory, MarkerNewBase, MemAllocChannel)
+				<< MarkerNewBase.BaseCycle(Cycle);
+		}
+		
 		UE_TRACE_LOG(Memory, Marker, MemAllocChannel)
 			<< Marker.Cycle(uint32(Cycle - BaseCycle));
 	}
