[TraceLog] Use 64-bit cycle fields

Back out of previous new base cycle event change and use the absolute 64-bit cycle value instead. The previous solution did not support late connect scenarios. The memory cost is not significant based on the current marker frequency.

diff --git a/Engine/Source/Developer/TraceServices/Private/Analyzers/AllocationsAnalysis.cpp b/Engine/Source/Developer/TraceServices/Private/Analyzers/AllocationsAnalysis.cpp
--- a/Engine/Source/Developer/TraceServices/Private/Analyzers/AllocationsAnalysis.cpp
+++ b/Engine/Source/Developer/TraceServices/Private/Analyzers/AllocationsAnalysis.cpp
@@ -40,7 +40,6 @@ void FAllocationsAnalyzer::OnAnalysisBegin(const FOnAnalysisContext& Context)
 	Builder.RouteEvent(RouteId_ReallocFree,  "Memory", "ReallocFree");
 	Builder.RouteEvent(RouteId_Marker,       "Memory", "Marker");
 	Builder.RouteEvent(RouteId_TagSpec,      "Memory", "TagSpec");
-	Builder.RouteEvent(RouteId_MarkerNewBase,"Memory", "MarkerNewBase");
 
 	Builder.RouteLoggerEvents(RouteId_MemScope, "Memory", true);
 }
@@ -78,7 +77,7 @@ bool FAllocationsAnalyzer::OnEvent(uint16 RouteId, EStyle Style, const FOnEventC
 			const uint8 SizeShift = EventData.GetValue<uint8>("SizeShift");
 			const uint8 SummarySizeShift = EventData.GetValue<uint8>("SummarySizeShift");
 
-			BaseCycle = EventData.GetValue<uint64>("BaseCycle");
+			BaseCycle = EventData.GetValue<uint64>("BaseCycle", 0);
 			MarkerPeriod = EventData.GetValue<uint32>("MarkerPeriod");
 
 			FAllocationsProvider::FEditScopeLock _(AllocationsProvider);
@@ -137,7 +136,9 @@ bool FAllocationsAnalyzer::OnEvent(uint16 RouteId, EStyle Style, const FOnEventC
 		}
 		case RouteId_Marker:
 		{
-			const uint64 Cycle = BaseCycle + EventData.GetValue<uint32>("Cycle");
+			// If BaseCycle is set Cycle is a relative 32-bit value, otherwise a 64-bit absolute value.
+			const uint64 Cycle = BaseCycle + (BaseCycle > 0 ? EventData.GetValue<uint32>("Cycle") : EventData.GetValue<uint64>("Cycle"));
+
 			if (ensure(Cycle >= LastMarkerCycle))
 			{
 				const double Seconds = Context.EventTime.AsSeconds(Cycle);
@@ -158,13 +159,6 @@ bool FAllocationsAnalyzer::OnEvent(uint16 RouteId, EStyle Style, const FOnEventC
 			}
 			break;
 		}
-		case RouteId_MarkerNewBase:
-		{
-			const uint64 NewBaseCycle = EventData.GetValue<uint64>("BaseCycle");
-			check(NewBaseCycle > BaseCycle);
-			BaseCycle = NewBaseCycle;
-			break;
-		}
 		case RouteId_TagSpec:
 		{
 			const int32 Tag = Context.EventData.GetValue<int32>("Tag");
diff --git a/Engine/Source/Developer/TraceServices/Private/Analyzers/AllocationsAnalysis.h b/Engine/Source/Developer/TraceServices/Private/Analyzers/AllocationsAnalysis.h
--- a/Engine/Source/Developer/TraceServices/Private/Analyzers/AllocationsAnalysis.h
+++ b/Engine/Source/Developer/TraceServices/Private/Analyzers/AllocationsAnalysis.h
@@ -24,7 +24,6 @@ private:
 		RouteId_ReallocAlloc,
 		RouteId_ReallocFree,
 		RouteId_Marker,
-		RouteId_MarkerNewBase,
 		RouteId_TagSpec,
 		RouteId_MemScope,
 	};
diff --git a/Engine/Source/Runtime/Core/Private/ProfilingDebugging/Windows/MemoryTrace.cpp b/Engine/Source/Runtime/Core/Private/ProfilingDebugging/Windows/MemoryTrace.cpp
--- a/Engine/Source/Runtime/Core/Private/ProfilingDebugging/Windows/MemoryTrace.cpp
+++ b/Engine/Source/Runtime/Core/Private/ProfilingDebugging/Windows/MemoryTrace.cpp
@@ -83,7 +83,6 @@ UE_TRACE_CHANNEL(MemSummaryChannel)
 UE_TRACE_CHANNEL_DEFINE(MemAllocChannel)
 
 UE_TRACE_EVENT_BEGIN(Memory, Init, NoSync|Important)
-	UE_TRACE_EVENT_FIELD(uint64, BaseCycle)
 	UE_TRACE_EVENT_FIELD(uint32, MarkerPeriod)
 	UE_TRACE_EVENT_FIELD(uint8, MinAlignment)
 	UE_TRACE_EVENT_FIELD(uint8, SizeShift)
@@ -104,11 +103,7 @@ UE_TRACE_EVENT_END()
 #endif
 
 UE_TRACE_EVENT_BEGIN(Memory, Marker)
-	UE_TRACE_EVENT_FIELD(uint32, Cycle)
-UE_TRACE_EVENT_END()
-
-UE_TRACE_EVENT_BEGIN(Memory, MarkerNewBase)
-	UE_TRACE_EVENT_FIELD(uint64, BaseCycle)
+	UE_TRACE_EVENT_FIELD(uint64, Cycle)
 UE_TRACE_EVENT_END()
 
 UE_TRACE_EVENT_BEGIN(Memory, CoreAdd)
@@ -308,10 +303,7 @@ static FUndestructed<FAllocationTrace> GAllocationTrace;
 ////////////////////////////////////////////////////////////////////////////////
 void FAllocationTrace::Initialize()
 {
-	BaseCycle = FPlatformTime::Cycles64();
-
 	UE_TRACE_LOG(Memory, Init, MemAllocChannel)
-		<< Init.BaseCycle(BaseCycle)
 		<< Init.MarkerPeriod(MarkerSamplePeriod + 1)
 		<< Init.MinAlignment(uint8(MIN_ALIGNMENT))
 #if USE_OVERVIEW_TRACE
@@ -334,19 +326,8 @@ void FAllocationTrace::Update()
 	uint32 TheCount = MarkerCounter.fetch_add(1, std::memory_order_relaxed);
 	if ((TheCount & MarkerSamplePeriod) == 0)
 	{
-		uint64 Cycle = FPlatformTime::Cycles64();
-
-		// We assume here that that the marker sample period is large
-		// enough that we don't have to syncronize across threads.
-		if ((Cycle - BaseCycle) > uint32(~0))
-		{
-			BaseCycle = Cycle;
-			UE_TRACE_LOG(Memory, MarkerNewBase, MemAllocChannel)
-				<< MarkerNewBase.BaseCycle(Cycle);
-		}
-		
 		UE_TRACE_LOG(Memory, Marker, MemAllocChannel)
-			<< Marker.Cycle(uint32(Cycle - BaseCycle));
+			<< Marker.Cycle(FPlatformTime::Cycles64());
 	}
 
 	if (bPumpTrace)
