[Insights] Fixed issue with duplicated cpu timers when used before being declared (resulting in incorrectly assigning some cpu timing events to first timer and others to the second timer). In this case (when two timers have same name and the EventSpec trace event is analyzed after the timers were already used in an EventBatch event) the two timers will not be merged by name, instead, two separate timers will appear in Insights and the timing events will be properly assigned to the two timers.

diff --git a/Engine/Source/Developer/TraceServices/Private/Analyzers/CpuProfilerTraceAnalysis.cpp b/Engine/Source/Developer/TraceServices/Private/Analyzers/CpuProfilerTraceAnalysis.cpp
--- a/Engine/Source/Developer/TraceServices/Private/Analyzers/CpuProfilerTraceAnalysis.cpp
+++ b/Engine/Source/Developer/TraceServices/Private/Analyzers/CpuProfilerTraceAnalysis.cpp
@@ -288,23 +288,52 @@ void FCpuProfilerAnalyzer::OnCpuScopeLeave(const FOnEventContext& Context)
 
 void FCpuProfilerAnalyzer::DefineScope(uint32 SpecId, const TCHAR* Name)
 {
+	// The cpu scoped events (timers) will be merged by name.
+	// Ex.: If there are multiple timers defined in code with same name,
+	//      those will appear in Insights as a single timer.
+
+	// Check if a timer with same name was already defined.
 	uint32* FindTimerIdByName = ScopeNameToTimerIdMap.Find(Name);
 	if (FindTimerIdByName)
 	{
-		SpecIdToTimerIdMap.Add(SpecId, *FindTimerIdByName);
+		// Yes, a timer with same name was already defined.
+		// Check if SpecId is already mapped to timer.
+		uint32* FindTimerId = SpecIdToTimerIdMap.Find(SpecId);
+		if (FindTimerId)
+		{
+			// Yes, SpecId was already mapped to a timer (ex. as an <unknown> timer).
+			// Update name for mapped timer.
+			TimingProfilerProvider.SetTimerName(*FindTimerId, Name);
+			// In this case, we do not remap the SpecId to the previously defined timer with same name.
+			// This is becasue the two timers are already used in timelines.
+			// So we will continue to use separate timers, even if those have same name.
+		}
+		else
+		{
+			// Map this SpecId to the previously defined timer with same name.
+			SpecIdToTimerIdMap.Add(SpecId, *FindTimerIdByName);
+		}
 	}
 	else
 	{
+		// No, a timer with same name was not defined.
+		// Check if SpecId is already mapped to timer.
 		uint32* FindTimerId = SpecIdToTimerIdMap.Find(SpecId);
 		if (FindTimerId)
 		{
+			// Yes, SpecId was already mapped to a timer (ex. as an <unknown> timer).
+			// Update name for mapped timer.
 			TimingProfilerProvider.SetTimerName(*FindTimerId, Name);
+			// Map the name to the timer.
 			ScopeNameToTimerIdMap.Add(Name, *FindTimerId);
 		}
 		else
 		{
+			// Define a new Cpu timer.
 			uint32 NewTimerId = TimingProfilerProvider.AddCpuTimer(Name);
+			// Map the SpecId to the timer.
 			SpecIdToTimerIdMap.Add(SpecId, NewTimerId);
+			// Map the name to the timer.
 			ScopeNameToTimerIdMap.Add(Name, NewTimerId);
 		}
 	}
