Fix atomic static initialization

#rnx

diff --git a/Engine/Source/Developer/TraceServices/Private/Model/CallstacksProvider.cpp b/Engine/Source/Developer/TraceServices/Private/Model/CallstacksProvider.cpp
--- a/Engine/Source/Developer/TraceServices/Private/Model/CallstacksProvider.cpp
+++ b/Engine/Source/Developer/TraceServices/Private/Model/CallstacksProvider.cpp
@@ -12,7 +12,7 @@ namespace TraceServices
 {
 
 /////////////////////////////////////////////////////////////////////
-static FResolvedSymbol GNeverResolveSymbol;
+static const FResolvedSymbol GNeverResolveSymbol(ESymbolQueryResult::NotLoaded, nullptr, nullptr);
 
 /////////////////////////////////////////////////////////////////////
 #ifdef TRACE_CALLSTACK_STATS
@@ -31,9 +31,6 @@ FCallstacksProvider::FCallstacksProvider(IAnalysisSession& InSession)
 	, Callstacks(InSession.GetLinearAllocator(), CallstacksPerPage)
 	, Frames(InSession.GetLinearAllocator(), FramesPerPage)
 {
-	GNeverResolveSymbol.Result.store(ESymbolQueryResult::NotLoaded);
-	GNeverResolveSymbol.Name = nullptr;
-	GNeverResolveSymbol.FileAndLine = nullptr;
 }
 
 /////////////////////////////////////////////////////////////////////
diff --git a/Engine/Source/Developer/TraceServices/Private/Model/ModuleProvider.cpp b/Engine/Source/Developer/TraceServices/Private/Model/ModuleProvider.cpp
--- a/Engine/Source/Developer/TraceServices/Private/Model/ModuleProvider.cpp
+++ b/Engine/Source/Developer/TraceServices/Private/Model/ModuleProvider.cpp
@@ -97,10 +97,7 @@ const FResolvedSymbol* TModuleProvider<SymbolProvider>::GetSymbol(uint64 Address
 		{
 			return SymbolCache[Address];
 		}
-		ResolvedSymbol = &Symbols.EmplaceBack();
-		ResolvedSymbol->Result = ESymbolQueryResult::Pending;
-		ResolvedSymbol->Name = nullptr;
-		ResolvedSymbol->FileAndLine = nullptr;
+		ResolvedSymbol = &Symbols.EmplaceBack(ESymbolQueryResult::Pending, nullptr, nullptr);
 		SymbolCache.Add(Address, ResolvedSymbol);
 	}
 
diff --git a/Engine/Source/Developer/TraceServices/Public/TraceServices/Model/Modules.h b/Engine/Source/Developer/TraceServices/Public/TraceServices/Model/Modules.h
--- a/Engine/Source/Developer/TraceServices/Public/TraceServices/Model/Modules.h
+++ b/Engine/Source/Developer/TraceServices/Public/TraceServices/Model/Modules.h
@@ -55,6 +55,12 @@ struct FResolvedSymbol
 	{
 		return Result.load(std::memory_order_acquire);
 	}
+
+	FResolvedSymbol(ESymbolQueryResult InResult, const TCHAR* InName, const TCHAR* InFileAndLine) 
+		: Result(InResult)
+		, Name(InName)
+		, FileAndLine(InFileAndLine)
+	{}
 };
 
 ////////////////////////////////////////////////////////////////////////////////
