[Insights] Frees symbol resources when complete

Added hooks to allow the symbol resolver to know when analysis is complete. At this point no more modules or symbol requests will come in and resources can be freed, releasing the file locks on pdb files.

diff --git a/Engine/Source/Developer/TraceServices/Private/Analyzers/ModuleAnalysis.cpp b/Engine/Source/Developer/TraceServices/Private/Analyzers/ModuleAnalysis.cpp
--- a/Engine/Source/Developer/TraceServices/Private/Analyzers/ModuleAnalysis.cpp
+++ b/Engine/Source/Developer/TraceServices/Private/Analyzers/ModuleAnalysis.cpp
@@ -31,6 +31,14 @@ void FModuleAnalyzer::OnAnalysisBegin(const FOnAnalysisContext& Context)
 	Builder.RouteEvent(RouteId_ModuleUnload, "Diagnostics", "ModuleUnload");
 }
 
+////////////////////////////////////////////////////////////////////////////////
+void FModuleAnalyzer::OnAnalysisEnd()
+{
+	if (Provider)
+	{
+		Provider->OnAnalysisComplete();
+	}
+}
 
 ////////////////////////////////////////////////////////////////////////////////
 bool FModuleAnalyzer::OnEvent(uint16 RouteId, EStyle Style, const FOnEventContext& Context)
diff --git a/Engine/Source/Developer/TraceServices/Private/Analyzers/ModuleAnalysis.h b/Engine/Source/Developer/TraceServices/Private/Analyzers/ModuleAnalysis.h
--- a/Engine/Source/Developer/TraceServices/Private/Analyzers/ModuleAnalysis.h
+++ b/Engine/Source/Developer/TraceServices/Private/Analyzers/ModuleAnalysis.h
@@ -15,6 +15,7 @@ class FModuleAnalyzer
 public:
 					FModuleAnalyzer(IAnalysisSession& Session);
 	virtual void	OnAnalysisBegin(const FOnAnalysisContext& Context) override;
+	virtual void	OnAnalysisEnd() override;
 	virtual bool 	OnEvent(uint16 RouteId, EStyle Style, const FOnEventContext& Context) override;
 
 private:
diff --git a/Engine/Source/Developer/TraceServices/Private/Model/ModuleProvider.cpp b/Engine/Source/Developer/TraceServices/Private/Model/ModuleProvider.cpp
--- a/Engine/Source/Developer/TraceServices/Private/Model/ModuleProvider.cpp
+++ b/Engine/Source/Developer/TraceServices/Private/Model/ModuleProvider.cpp
@@ -51,6 +51,7 @@ public:
 	//Analysis interface
 	void						OnModuleLoad(const FStringView& Module, uint64 Base, uint32 Size) override;
 	void 						OnModuleUnload(uint64 Base) override;
+	void						OnAnalysisComplete() override;
 
 private:
 
@@ -143,6 +144,13 @@ void TModuleProvider<SymbolProvider>::OnModuleUnload(uint64 Base)
 	//todo: Find entry, set bLoaded to false
 }
 
+/////////////////////////////////////////////////////////////////////
+template<typename SymbolResolver>
+void TModuleProvider<SymbolResolver>::OnAnalysisComplete()
+{
+	Resolver.OnAnalysisComplete();
+}
+
 /////////////////////////////////////////////////////////////////////
 IModuleAnalysisProvider* CreateModuleProvider(IAnalysisSession& InSession, const FAnsiStringView& InSymbolFormat)
 {
diff --git a/Engine/Source/Developer/TraceServices/Private/Model/ModuleProvider.h b/Engine/Source/Developer/TraceServices/Private/Model/ModuleProvider.h
--- a/Engine/Source/Developer/TraceServices/Private/Model/ModuleProvider.h
+++ b/Engine/Source/Developer/TraceServices/Private/Model/ModuleProvider.h
@@ -16,6 +16,7 @@ class IModuleAnalysisProvider : public IModuleProvider
 public:
 	virtual void	OnModuleLoad(const FStringView& Module, uint64 Base, uint32 Size) = 0;
 	virtual void	OnModuleUnload(uint64 Base) = 0;
+	virtual void	OnAnalysisComplete() = 0;
 };
 
 /** Create a module provider with the given symbol format */
diff --git a/Engine/Source/Developer/TraceServices/Private/Model/PbdSymbols.inl b/Engine/Source/Developer/TraceServices/Private/Model/PbdSymbols.inl
--- a/Engine/Source/Developer/TraceServices/Private/Model/PbdSymbols.inl
+++ b/Engine/Source/Developer/TraceServices/Private/Model/PbdSymbols.inl
@@ -33,23 +33,12 @@ class FPdbSymbols : public FRunnable
 {
 public:
 	FPdbSymbols(IAnalysisSession& InSession)
-		: bInitialized(false)
+		: bRunWorkerThread(false)
+		, bDrainThenStop(false)
+		, bInitialized(false)
 		, Session(InSession)
 	{
-		static uint64 BaseHandle = 0x493;
-		Handle = (HANDLE)(++BaseHandle);
-
-		// Load DbgHelp interface
-		ULONG SymOpts = 0;
-		SymOpts |= SYMOPT_LOAD_LINES;
-		SymOpts |= SYMOPT_OMAP_FIND_NEAREST;
-		SymOpts |= SYMOPT_DEFERRED_LOADS;
-		SymOpts |= SYMOPT_EXACT_SYMBOLS;
-		SymOpts |= SYMOPT_IGNORE_NT_SYMPATH;
-		SymOpts |= SYMOPT_UNDNAME;
-
-		SymSetOptions(SymOpts);
-		bInitialized = SymInitialize(Handle, NULL, FALSE); 
+		bInitialized = SetupSyms();
 
 		if (bInitialized)
 		{
@@ -62,9 +51,10 @@ public:
 	~FPdbSymbols()
 	{
 		bRunWorkerThread = false;
-		Thread->WaitForCompletion();
-
-		SymCleanup(Handle);
+		if (Thread)
+		{
+			Thread->WaitForCompletion();
+		}
 	}
 
 	void QueueModuleLoad(const FStringView& ModulePath, uint64 Base, uint32 Size)
@@ -101,6 +91,13 @@ public:
 		OutStats->SymbolsResolved = SymbolsResolved.load();
 	}
 
+	void OnAnalysisComplete()
+	{
+		// At this point no more module loads or symbol requests will be queued,
+		// we drain the current queue, then release resources and file locks.
+		bDrainThenStop = true;
+	}
+
 private:
 
 	struct ModuleEntry
@@ -130,6 +127,33 @@ private:
 	};
 
 
+	bool SetupSyms() 
+	{
+		// Create a unique handle
+		static uint64 BaseHandle = 0x493;
+		Handle = (HANDLE)(++BaseHandle);
+
+		// Load DbgHelp interface
+		ULONG SymOpts = 0;
+		SymOpts |= SYMOPT_LOAD_LINES;
+		SymOpts |= SYMOPT_OMAP_FIND_NEAREST;
+		SymOpts |= SYMOPT_DEFERRED_LOADS;
+		SymOpts |= SYMOPT_EXACT_SYMBOLS;
+		SymOpts |= SYMOPT_IGNORE_NT_SYMPATH;
+		SymOpts |= SYMOPT_UNDNAME;
+
+		SymSetOptions(SymOpts);
+		return SymInitialize(Handle, NULL, FALSE);
+	}
+
+
+	void FreeSyms()
+	{
+		// This release file locks on debug files
+		SymCleanup(Handle);
+	}
+
+
 	uint32 Run() override 
 	{
 		while (bRunWorkerThread)
@@ -153,10 +177,18 @@ private:
 					ResolveSymbol(Item.Address, Item.Target);
 				}
 			}
+
+			if (bDrainThenStop && ResolveQueue.IsEmpty() && LoadSymbolsQueue.IsEmpty())
+			{
+				bRunWorkerThread = false;
+			}
 			
 			// ...and breathe...
 			FPlatformProcess::Sleep(0.2f);
 		}
+
+		// We don't need the syms library anymore
+		FreeSyms();
 		
 		return 0;
 	}
@@ -272,6 +304,7 @@ private:
 	std::atomic<uint32> SymbolsResolved;
 
 	bool bRunWorkerThread;
+	bool bDrainThenStop;
 	bool bInitialized;
 	HANDLE Handle;
 	IAnalysisSession& Session;
diff --git a/Engine/Source/Developer/TraceServices/Public/TraceServices/Model/Modules.h b/Engine/Source/Developer/TraceServices/Public/TraceServices/Model/Modules.h
--- a/Engine/Source/Developer/TraceServices/Public/TraceServices/Model/Modules.h
+++ b/Engine/Source/Developer/TraceServices/Public/TraceServices/Model/Modules.h
@@ -77,7 +77,10 @@ public:
 		uint32 SymbolsFailed;
 	};
 	
-	/** Queries the name of the symbol at address */
+	/** Queries the name of the symbol at address. This function returns immedately, 
+	 * but the lookup is async. See \ref FResolvedSymbol for details. It assumed that 
+	 * all calls to this function happens before analysis has ended.
+	 */
 	virtual const FResolvedSymbol* GetSymbol(uint64 Address) = 0;
 
 	/** Gets statistics from provider */
