Allocations use realloc scopes

Realloc scopes were not applied. Additionally we extracted incorrect information from the event and constructed the combined thread and tracker information wrong.

diff --git a/Engine/Source/Developer/TraceServices/Private/Analyzers/AllocationsAnalysis.cpp b/Engine/Source/Developer/TraceServices/Private/Analyzers/AllocationsAnalysis.cpp
--- a/Engine/Source/Developer/TraceServices/Private/Analyzers/AllocationsAnalysis.cpp
+++ b/Engine/Source/Developer/TraceServices/Private/Analyzers/AllocationsAnalysis.cpp
@@ -187,7 +187,7 @@ bool FAllocationsAnalyzer::OnEvent(uint16 RouteId, EStyle Style, const FOnEventC
 				}
 				else // "MemoryScopeRealloc"
 				{
-					const uint64 Ptr = Context.EventData.GetValue<int32>("Ptr");
+					const uint64 Ptr = Context.EventData.GetValue<uint64>("Ptr");
 					FAllocationsProvider::FEditScopeLock _(AllocationsProvider);
 					AllocationsProvider.EditPushRealloc(ThreadId, Tracker, Ptr);
 				}
diff --git a/Engine/Source/Developer/TraceServices/Private/Model/AllocationsProvider.cpp b/Engine/Source/Developer/TraceServices/Private/Model/AllocationsProvider.cpp
--- a/Engine/Source/Developer/TraceServices/Private/Model/AllocationsProvider.cpp
+++ b/Engine/Source/Developer/TraceServices/Private/Model/AllocationsProvider.cpp
@@ -150,11 +150,12 @@ const TCHAR* FTagTracker::GetTagString(uint32 InTag) const
 
 ////////////////////////////////////////////////////////////////////////////////////////////////////
 
-void FTagTracker::PushRealloc(uint32 InThreadId, uint8 InTracker, uint64 InPtr)
+void FTagTracker::PushRealloc(uint32 InThreadId, uint8 InTracker, uint32 InTag)
 {
 	const uint32 TrackerThreadId = GetTrackerThreadId(InThreadId, InTracker);
 	ThreadState& State = TrackerThreadStates.FindOrAdd(TrackerThreadId);
-	State.ReallocPtr.Push(InPtr);
+	State.TagStack.Push(InTag);
+	State.bReallocTagActive = true;
 }
 
 ////////////////////////////////////////////////////////////////////////////////////////////////////
@@ -163,9 +164,10 @@ void FTagTracker::PopRealloc(uint32 InThreadId, uint8 InTracker)
 {
 	const uint32 TrackerThreadId = GetTrackerThreadId(InThreadId, InTracker);
 	ThreadState* State = TrackerThreadStates.Find(TrackerThreadId);
-	if (ensure(State && !State->ReallocPtr.IsEmpty()))
+	if (ensure(State && !State->TagStack.IsEmpty() && State->bReallocTagActive))
 	{
-		State->ReallocPtr.Pop();
+		State->TagStack.Pop();
+		State->bReallocTagActive = false;
 	}
 	else
 	{
@@ -183,16 +185,7 @@ bool FTagTracker::HasReallocScope(uint32 InThreadId, uint8 InTracker) const
 {
 	const uint32 TrackerThreadId = GetTrackerThreadId(InThreadId, InTracker);
 	const ThreadState* State = TrackerThreadStates.Find(TrackerThreadId);
-	return State && !State->ReallocPtr.IsEmpty();
-}
-
-////////////////////////////////////////////////////////////////////////////////////////////////////
-
-bool FTagTracker::IsReallocScope(uint64 InSourcePtr, uint32 InThreadId, uint8 InTracker) const
-{
-	const uint32 TrackerThreadId = GetTrackerThreadId(InThreadId, InTracker);
-	const ThreadState* State = TrackerThreadStates.Find(TrackerThreadId);
-	return State && !State->ReallocPtr.IsEmpty() && State->ReallocPtr.Contains(InSourcePtr);
+	return State && State->bReallocTagActive;
 }
 
 ////////////////////////////////////////////////////////////////////////////////////////////////////
@@ -635,6 +628,23 @@ void FAllocationsProvider::AdvanceTimelines(double Time)
 
 ////////////////////////////////////////////////////////////////////////////////////////////////////
 
+void FAllocationsProvider::EditPushRealloc(uint32 ThreadId, uint8 Tracker, uint64 Ptr) 
+{ 
+	EditAccessCheck(); 
+	int32 Tag(0); // If ptr is not found use "Untagged"
+	if (FAllocationItem* Alloc = LiveAllocs.Find(Ptr))
+	{
+		Tag = Alloc->Tag;
+	}
+	TagTracker.PushRealloc(ThreadId, Tracker, Tag);
+}
+
+void FAllocationsProvider::EditPopRealloc(uint32 ThreadId, uint8 Tracker) 
+{ 
+	EditAccessCheck(); 
+	TagTracker.PopRealloc(ThreadId, Tracker); 
+}
+
 void FAllocationsProvider::EditOnAnalysisCompleted(double Time)
 {
 	Lock.WriteAccessCheck();
diff --git a/Engine/Source/Developer/TraceServices/Private/Model/AllocationsProvider.h b/Engine/Source/Developer/TraceServices/Private/Model/AllocationsProvider.h
--- a/Engine/Source/Developer/TraceServices/Private/Model/AllocationsProvider.h
+++ b/Engine/Source/Developer/TraceServices/Private/Model/AllocationsProvider.h
@@ -44,7 +44,7 @@ private:
 	struct ThreadState
 	{
 		TArray<uint32> TagStack;
-		TArray<uint64> ReallocPtr;
+		bool bReallocTagActive;
 	};
 
 	struct TagEntry
@@ -60,15 +60,14 @@ public:
 	uint32 GetCurrentTag(uint32 ThreadId, uint8 Tracker) const;
 	const TCHAR* GetTagString(uint32 Tag) const;
 
-	void PushRealloc(uint32 ThreadId, uint8 Tracker, uint64 Ptr);
+	void PushRealloc(uint32 ThreadId, uint8 Tracker, uint32 Tag);
 	void PopRealloc(uint32 ThreadId, uint8 Tracker);
 	bool HasReallocScope(uint32 ThreadId, uint8 Tracker) const;
-	bool IsReallocScope(uint64 SourcePtr, uint32 ThreadId, uint8 Tracker) const;
 
 private:
 	inline uint32 GetTrackerThreadId(uint32 ThreadId, uint8 Tracker) const
 	{
-		return (Tracker << TrackerIdShift) & (~TrackerIdMask & ThreadId);
+		return (Tracker << TrackerIdShift) | (~TrackerIdMask & ThreadId);
 	}
 
 	TMap<uint32, ThreadState> TrackerThreadStates;
@@ -183,8 +182,8 @@ public:
 	void EditPushTag(uint32 ThreadId, uint8 Tracker, uint32 Tag)           { EditAccessCheck(); TagTracker.PushTag(ThreadId, Tracker, Tag); }
 	void EditPopTag(uint32 ThreadId, uint8 Tracker)                        { EditAccessCheck(); TagTracker.PopTag(ThreadId, Tracker); }
 
-	void EditPushRealloc(uint32 ThreadId, uint8 Tracker, uint64 Ptr)       { EditAccessCheck(); TagTracker.PushRealloc(ThreadId, Tracker, Ptr); }
-	void EditPopRealloc(uint32 ThreadId, uint8 Tracker)                    { EditAccessCheck(); TagTracker.PopRealloc(ThreadId, Tracker); }
+	void EditPushRealloc(uint32 ThreadId, uint8 Tracker, uint64 Ptr);
+	void EditPopRealloc(uint32 ThreadId, uint8 Tracker);
 
 	void EditOnAnalysisCompleted(double Time);
 
