[Insights]
  - Memory: Improved speed of live allocs matching in memory analysis.
  - Memory: Fixed N/A callstacks in a query made before analysis is completed.

diff --git a/Engine/Source/Developer/TraceServices/Private/Model/AllocationsProvider.cpp b/Engine/Source/Developer/TraceServices/Private/Model/AllocationsProvider.cpp
--- a/Engine/Source/Developer/TraceServices/Private/Model/AllocationsProvider.cpp
+++ b/Engine/Source/Developer/TraceServices/Private/Model/AllocationsProvider.cpp
@@ -11,11 +11,40 @@
 
 #include <limits>
 
+////////////////////////////////////////////////////////////////////////////////////////////////////
+
+#define INSIGHTS_SLOW_CHECK(expr) //check(expr)
+
+#define INSIGHTS_DEBUG_WATCH 0
+
+// Use optimized path for the short living allocations.
+// If enabled, caches the short living allocations, as there is a very high chance to be freed in the next few "free" events.
+// ~66% of all allocs are expected to have an "event distance" < 64 events
+// ~70% of all allocs are expected to have an "event distance" < 512 events
+#define INSIGHTS_USE_SHORT_LIVING_ALLOCS 1
+
+// Use optimized path for the last alloc.
+// If enabled, caches the last added alloc, as there is a high chance to be freed in the next "free" event.
+// ~10% to ~30% of all allocs are expected to have an "event distance" == 1 event ("free" event follows the "alloc" event imediatelly)
+#define INSIGHTS_USE_LAST_ALLOC 1
+
+#define INSIGHTS_VALIDATE_ALLOC_EVENTS 0
+
+////////////////////////////////////////////////////////////////////////////////////////////////////
+
 namespace TraceServices
 {
 
 constexpr uint32 MaxLogMessagesPerErrorType = 100;
 
+#if INSIGHTS_DEBUG_WATCH
+static uint64 GWatchAddresses[] =
+{
+	// add here addresses to watch
+	0x0,
+};
+#endif // INSIGHTS_DEBUG_WATCH
+
 ////////////////////////////////////////////////////////////////////////////////////////////////////
 // FAllocationsProviderLock
 ////////////////////////////////////////////////////////////////////////////////////////////////////
@@ -291,6 +320,331 @@ IAllocationsProvider::FQueryResult IAllocationsProvider::FQueryStatus::NextResul
 	return FQueryResult(Ret);
 }
 
+////////////////////////////////////////////////////////////////////////////////////////////////////
+// FShortLivingAllocs
+////////////////////////////////////////////////////////////////////////////////////////////////////
+
+FShortLivingAllocs::FShortLivingAllocs()
+{
+	AllNodes = new FNode[MaxAllocCount];
+
+	// Build the "unused" simple linked list.
+	FirstUnusedNode = AllNodes;
+	for (int32 Index = 0; Index < MaxAllocCount - 1; ++Index)
+	{
+		AllNodes[Index].Next = &AllNodes[Index + 1];
+	}
+	AllNodes[MaxAllocCount - 1].Next = nullptr;
+}
+
+////////////////////////////////////////////////////////////////////////////////////////////////////
+
+FShortLivingAllocs::~FShortLivingAllocs()
+{
+	FNode* Node = LastAddedAllocNode;
+	while (Node != nullptr)
+	{
+		delete Node->Alloc;
+		Node = Node->Prev;
+	}
+	delete[] AllNodes;
+}
+
+////////////////////////////////////////////////////////////////////////////////////////////////////
+
+FAllocationItem* FShortLivingAllocs::FindRef(uint64 Address)
+{
+#if 1
+	FNode** NodePtr = AddressMap.Find(Address);
+	return NodePtr ? (*NodePtr)->Alloc : nullptr;
+#else
+	// Search linearly the list of allocations, backward, starting with most recent one.
+	// As the probability of finding a match for a "free" event is higher for the recent allocs,
+	// this could actually be faster than doing a O(log n) search in AddressMap.
+	// TODO: This needs to be tested for various MaxAllocCount values.
+	FNode* Node = LastAddedAllocNode;
+	while (Node != nullptr)
+	{
+		if (Node->Alloc->Address == Address)
+		{
+			return Node->Alloc;
+		}
+		Node = Node->Prev;
+	}
+	return nullptr;
+#endif
+}
+
+////////////////////////////////////////////////////////////////////////////////////////////////////
+
+FAllocationItem* FShortLivingAllocs::AddChecked(FAllocationItem* Alloc)
+{
+	if (FirstUnusedNode == nullptr)
+	{
+		// Collection is already full.
+		INSIGHTS_SLOW_CHECK(AllocCount == MaxAllocCount);
+		INSIGHTS_SLOW_CHECK(OldestAllocNode != nullptr);
+		INSIGHTS_SLOW_CHECK(LastAddedAllocNode != nullptr);
+
+		// Reuse the node of the oldest allocation for the new allocation.
+		FNode* NewNode = OldestAllocNode;
+
+		// Remove the oldest allocation.
+		FAllocationItem* RemovedAlloc = OldestAllocNode->Alloc;
+		AddressMap.Remove(RemovedAlloc->Address);
+		OldestAllocNode = OldestAllocNode->Next;
+		INSIGHTS_SLOW_CHECK(OldestAllocNode != nullptr);
+		OldestAllocNode->Prev = nullptr;
+
+		// Add the new node.
+		AddressMap.Add(Alloc->Address, NewNode);
+		NewNode->Alloc = Alloc;
+		NewNode->Next = nullptr;
+		NewNode->Prev = LastAddedAllocNode;
+		LastAddedAllocNode->Next = NewNode;
+		LastAddedAllocNode = NewNode;
+
+		return RemovedAlloc;
+	}
+	else
+	{
+		INSIGHTS_SLOW_CHECK(AllocCount < MaxAllocCount);
+		++AllocCount;
+
+		FNode* NewNode = FirstUnusedNode;
+		FirstUnusedNode = FirstUnusedNode->Next;
+
+		// Add the new node.
+		AddressMap.Add(Alloc->Address, NewNode);
+		NewNode->Alloc = Alloc;
+		NewNode->Next = nullptr;
+		NewNode->Prev = LastAddedAllocNode;
+		if (LastAddedAllocNode)
+		{
+			LastAddedAllocNode->Next = NewNode;
+		}
+		else
+		{
+			OldestAllocNode = NewNode;
+		}
+		LastAddedAllocNode = NewNode;
+
+		return nullptr;
+	}
+}
+
+////////////////////////////////////////////////////////////////////////////////////////////////////
+
+FAllocationItem* FShortLivingAllocs::Remove(uint64 Address)
+{
+	FNode* RemovedNode;
+	if (AddressMap.RemoveAndCopyValue(Address, RemovedNode))
+	{
+		INSIGHTS_SLOW_CHECK(AllocCount > 0);
+		--AllocCount;
+
+		FAllocationItem* RemovedAlloc = RemovedNode->Alloc;
+		INSIGHTS_SLOW_CHECK(RemovedAlloc->Address == Address);
+
+		// Remove node.
+		if (RemovedNode == OldestAllocNode)
+		{
+			OldestAllocNode = OldestAllocNode->Next;
+		}
+		if (RemovedNode == LastAddedAllocNode)
+		{
+			LastAddedAllocNode = LastAddedAllocNode->Prev;
+		}
+		if (RemovedNode->Prev)
+		{
+			RemovedNode->Prev->Next = RemovedNode->Next;
+		}
+		if (RemovedNode->Next)
+		{
+			RemovedNode->Next->Prev = RemovedNode->Prev;
+		}
+
+		// Add the removed node to the "unused" list.
+		RemovedNode->Next = FirstUnusedNode;
+		FirstUnusedNode = RemovedNode;
+
+		return RemovedAlloc;
+	}
+
+	return nullptr;
+}
+
+////////////////////////////////////////////////////////////////////////////////////////////////////
+
+void FShortLivingAllocs::Enumerate(TFunctionRef<void(const FAllocationItem& Alloc)> Callback) const
+{
+	FNode* Node = LastAddedAllocNode;
+	while (Node != nullptr)
+	{
+		Callback(*(Node->Alloc));
+		Node = Node->Prev;
+	}
+}
+
+////////////////////////////////////////////////////////////////////////////////////////////////////
+// FLiveAllocCollection
+////////////////////////////////////////////////////////////////////////////////////////////////////
+
+FLiveAllocCollection::FLiveAllocCollection()
+{
+}
+
+////////////////////////////////////////////////////////////////////////////////////////////////////
+
+FLiveAllocCollection::~FLiveAllocCollection()
+{
+	for (const auto& KV : LongLivingAllocs)
+	{
+		const FAllocationItem* Allocation = KV.Value;
+		delete Allocation;
+	}
+
+#if INSIGHTS_USE_LAST_ALLOC
+	if (LastAlloc)
+	{
+		delete LastAlloc;
+	}
+#endif
+}
+
+////////////////////////////////////////////////////////////////////////////////////////////////////
+
+FAllocationItem* FLiveAllocCollection::FindRef(uint64 Address)
+{
+#if INSIGHTS_USE_LAST_ALLOC
+	if (LastAlloc && LastAlloc->Address == Address)
+	{
+		return LastAlloc;
+	}
+#endif
+
+#if INSIGHTS_USE_SHORT_LIVING_ALLOCS
+	FAllocationItem* FoundShortLivingAlloc = ShortLivingAllocs.FindRef(Address);
+	if (FoundShortLivingAlloc)
+	{
+		INSIGHTS_SLOW_CHECK(FoundShortLivingAlloc->Address == Address);
+		return FoundShortLivingAlloc;
+	}
+#endif
+
+	FAllocationItem* FoundLongLivingAlloc = LongLivingAllocs.FindRef(Address);
+	if (FoundLongLivingAlloc)
+	{
+		INSIGHTS_SLOW_CHECK(FoundLongLivingAlloc->Address == Address);
+		return FoundLongLivingAlloc;
+	}
+
+	return nullptr;
+}
+
+////////////////////////////////////////////////////////////////////////////////////////////////////
+
+FAllocationItem* FLiveAllocCollection::AddNewChecked(uint64 Address)
+{
+	++TotalAllocCount;
+	if (TotalAllocCount > MaxAllocCount)
+	{
+		MaxAllocCount = TotalAllocCount;
+	}
+
+	FAllocationItem* NewAlloc = new FAllocationItem();
+	NewAlloc->Address = Address;
+
+	// The allocation that is further moved from "last" to "short living" to "long living".
+	FAllocationItem* Allocation = NewAlloc;
+
+#if INSIGHTS_USE_LAST_ALLOC
+	if (LastAlloc)
+	{
+		// We have a new "last allocation".
+		// The previous one will be moved to the short living allocations.
+		Allocation = LastAlloc;
+		LastAlloc = NewAlloc;
+	}
+	else
+	{
+		LastAlloc = NewAlloc;
+		return NewAlloc;
+	}
+#endif
+
+#if INSIGHTS_USE_SHORT_LIVING_ALLOCS
+	Allocation = ShortLivingAllocs.AddChecked(Allocation);
+	if (Allocation == nullptr)
+	{
+		return NewAlloc;
+	}
+#endif
+
+	LongLivingAllocs.Add(Allocation->Address, Allocation);
+	return NewAlloc;
+}
+
+////////////////////////////////////////////////////////////////////////////////////////////////////
+
+FAllocationItem* FLiveAllocCollection::Remove(uint64 Address)
+{
+#if INSIGHTS_USE_LAST_ALLOC
+	if (LastAlloc && LastAlloc->Address == Address)
+	{
+		FAllocationItem* RemovedAlloc = LastAlloc;
+		LastAlloc = nullptr;
+		INSIGHTS_SLOW_CHECK(TotalAllocCount > 0);
+		--TotalAllocCount;
+		return RemovedAlloc;
+	}
+#endif
+
+#if INSIGHTS_USE_SHORT_LIVING_ALLOCS
+	FAllocationItem* RemovedShortLivingAlloc = ShortLivingAllocs.Remove(Address);
+	if (RemovedShortLivingAlloc)
+	{
+		INSIGHTS_SLOW_CHECK(RemovedShortLivingAlloc->Address == Address);
+		INSIGHTS_SLOW_CHECK(TotalAllocCount > 0);
+		--TotalAllocCount;
+		return RemovedShortLivingAlloc;
+	}
+#endif
+
+	FAllocationItem* RemovedLongLivingAlloc;
+	if (LongLivingAllocs.RemoveAndCopyValue(Address, RemovedLongLivingAlloc))
+	{
+		INSIGHTS_SLOW_CHECK(RemovedLongLivingAlloc->Address == Address);
+		INSIGHTS_SLOW_CHECK(TotalAllocCount > 0);
+		--TotalAllocCount;
+		return RemovedLongLivingAlloc;
+	}
+
+	return nullptr;
+}
+
+////////////////////////////////////////////////////////////////////////////////////////////////////
+
+void FLiveAllocCollection::Enumerate(TFunctionRef<void(const FAllocationItem& Alloc)> Callback) const
+{
+#if INSIGHTS_USE_LAST_ALLOC
+	if (LastAlloc)
+	{
+		Callback(*LastAlloc);
+	}
+#endif
+
+#if INSIGHTS_USE_SHORT_LIVING_ALLOCS
+	ShortLivingAllocs.Enumerate(Callback);
+#endif
+
+	for (const auto& KV : LongLivingAllocs)
+	{
+		const FAllocationItem* Allocation = KV.Value;
+		Callback(*Allocation);
+	}
+}
+
 ////////////////////////////////////////////////////////////////////////////////////////////////////
 // FAllocationsProvider
 ////////////////////////////////////////////////////////////////////////////////////////////////////
@@ -381,17 +735,6 @@ void FAllocationsProvider::EditRemoveCore(double Time, uint64 Owner, uint64 Base
 
 ////////////////////////////////////////////////////////////////////////////////////////////////////
 
-#define MEMALLOC_DEBUG_WATCH 0
-
-#if MEMALLOC_DEBUG_WATCH
-static uint64 GWatchAddresses[] =
-{
-	0x0, // add here the addresses to watch
-};
-#endif // MEMALLOC_DEBUG_WATCH
-
-////////////////////////////////////////////////////////////////////////////////////////////////////
-
 void FAllocationsProvider::EditAlloc(double Time, uint64 Owner, uint64 Address, uint32 InSize, uint8 InAlignmentAndSizeLower, uint32 ThreadId, uint8 Tracker)
 {
 	Lock.WriteAccessCheck();
@@ -408,67 +751,67 @@ void FAllocationsProvider::EditAlloc(double Time, uint64 Owner, uint64 Address,
 
 	SbTree->SetTimeForEvent(EventIndex, Time);
 
+	AdvanceTimelines(Time);
+
 	const uint8 SizeLowerMask = ((1 << SizeShift) - 1);
 	const uint8 AlignmentMask = ~SizeLowerMask;
 	const uint64 Size = (static_cast<uint64>(InSize) << SizeShift) | static_cast<uint64>(InAlignmentAndSizeLower & SizeLowerMask);
 
 	const uint32 Tag = TagTracker.GetCurrentTag(ThreadId, Tracker);
 
-	FAllocationItem* AllocationPtr = LiveAllocs.Find(Address);
-	if (AllocationPtr)
+#if INSIGHTS_DEBUG_WATCH
+	for (int32 AddrIndex = 0; AddrIndex < UE_ARRAY_COUNT(GWatchAddresses); ++AddrIndex)
 	{
-		++AllocErrors;
-		if (AllocErrors <= MaxLogMessagesPerErrorType)
+		if (GWatchAddresses[AddrIndex] == Address)
 		{
-			UE_LOG(LogTraceServices, Error, TEXT("[MemAlloc] Invalid ALLOC event (Address=0x%llX, Size=%llu, Tag=%u, Time=%f)!"), Address, Size, Tag, Time);
+			UE_LOG(LogTraceServices, Warning, TEXT("[MemAlloc] Alloc 0x%llX : Size=%llu, Tag=%u, Time=%f"), Address, Size, Tag, Time);
+			break;
 		}
 	}
-	else
-	{
-#if MEMALLOC_DEBUG_WATCH
-		for (int32 AddrIndex = 0; AddrIndex < UE_ARRAY_COUNT(GWatchAddresses); ++AddrIndex)
-		{
-			if (GWatchAddresses[AddrIndex] == Address)
-			{
-				UE_LOG(LogTraceServices, Warning, TEXT("[MemAlloc] Alloc 0x%llX : Size=%llu, Tag=%u, Time=%f"), Address, Size, Tag, Time);
-				break;
-			}
-		}
-#endif // MEMALLOC_DEBUG_WATCH
-
-		AdvanceTimelines(Time);
+#endif // INSIGHTS_DEBUG_WATCH
 
-		FAllocationItem Allocation =
-		{
-			EventIndex,
-			(uint32)-1,
-			Time,
-			std::numeric_limits<double>::infinity(),
-			Owner,
-			Address,
-			FAllocationItem::PackSizeAndAlignment(Size, InAlignmentAndSizeLower & AlignmentMask),
-			nullptr,
-			Tag,
-			0, // Reserverd1
-		};
-		LiveAllocs.Add(Address, Allocation);
+#if INSIGHTS_VALIDATE_ALLOC_EVENTS
+	FAllocationItem* AllocationPtr = LiveAllocs.FindRef(Address);
+	if (!AllocationPtr)
+	{
+#else
+		FAllocationItem* AllocationPtr;
+#endif
 
-		//TODO: Cache the last added alloc, as there is ~10% chance to be freed in the next event.
-
-		const uint32 CurrentLiveAllocCount = static_cast<uint32>(LiveAllocs.Num());
-		if (CurrentLiveAllocCount > MaxLiveAllocCount)
-		{
-			MaxLiveAllocCount = CurrentLiveAllocCount;
-		}
+		AllocationPtr = LiveAllocs.AddNewChecked(Address);
+		INSIGHTS_SLOW_CHECK(Allocation.Address == Address)
+
+		FAllocationItem& Allocation = *AllocationPtr;
+		Allocation.StartEventIndex = EventIndex;
+		Allocation.EndEventIndex = (uint32)-1;
+		Allocation.StartTime = Time;
+		Allocation.EndTime = std::numeric_limits<double>::infinity();
+		Allocation.Owner = Owner;
+		//Allocation.Address = Address;
+		Allocation.SizeAndAlignment = FAllocationItem::PackSizeAndAlignment(Size, InAlignmentAndSizeLower & AlignmentMask);
+		Allocation.Callstack = nullptr;
+		Allocation.Tag = Tag;
+		Allocation.Reserved1 = 0;
 
 		UpdateHistogramByAllocSize(Size);
 
-		// Update stats for current timeline sample.
+		// Update stats for the current timeline sample.
 		TotalAllocatedMemory += Size;
 		SampleMaxTotalAllocatedMemory = FMath::Max(SampleMaxTotalAllocatedMemory, TotalAllocatedMemory);
 		SampleMaxLiveAllocations = FMath::Max(SampleMaxLiveAllocations, (uint32)LiveAllocs.Num());
 		++SampleAllocEvents;
+
+#if INSIGHTS_VALIDATE_ALLOC_EVENTS
+	}
+	else
+	{
+		++AllocErrors;
+		if (AllocErrors <= MaxLogMessagesPerErrorType)
+		{
+			UE_LOG(LogTraceServices, Error, TEXT("[MemAlloc] Invalid ALLOC event (Address=0x%llX, Size=%llu, Tag=%u, Time=%f)!"), Address, Size, Tag, Time);
+		}
 	}
+#endif
 
 	++AllocCount;
 	++EventIndex;
@@ -492,50 +835,49 @@ void FAllocationsProvider::EditFree(double Time, uint64 Address)
 
 	SbTree->SetTimeForEvent(EventIndex, Time);
 
-	FAllocationItem* AllocationPtr = LiveAllocs.Find(Address);
-	if (!AllocationPtr)
+	AdvanceTimelines(Time);
+
+#if INSIGHTS_DEBUG_WATCH
+	for (int32 AddrIndex = 0; AddrIndex < UE_ARRAY_COUNT(GWatchAddresses); ++AddrIndex)
 	{
-		++FreeErrors;
-		if (FreeErrors <= MaxLogMessagesPerErrorType)
+		if (GWatchAddresses[AddrIndex] == Address)
 		{
-			UE_LOG(LogTraceServices, Error, TEXT("[MemAlloc] Invalid FREE event (Address=0x%llX, Time=%f)!"), Address, Time);
+			UE_LOG(LogTraceServices, Warning, TEXT("[MemAlloc] Free 0x%llX : Time=%f"), Address, Time);
+			break;
 		}
 	}
-	else
-	{
-#if MEMALLOC_DEBUG_WATCH
-		for (int32 AddrIndex = 0; AddrIndex < UE_ARRAY_COUNT(GWatchAddresses); ++AddrIndex)
-		{
-			if (GWatchAddresses[AddrIndex] == Address)
-			{
-				UE_LOG(LogTraceServices, Warning, TEXT("[MemAlloc] Free 0x%llX : Time=%f"), Address, Time);
-				break;
-			}
-		}
-#endif // MEMALLOC_DEBUG_WATCH
-
-		AdvanceTimelines(Time);
+#endif // INSIGHTS_DEBUG_WATCH
 
+	FAllocationItem* AllocationPtr = LiveAllocs.Remove(Address); // we take ownership of AllocationPtr
+	if (AllocationPtr)
+	{
 		check(EventIndex > AllocationPtr->StartEventIndex);
 		AllocationPtr->EndEventIndex = EventIndex;
 		AllocationPtr->EndTime = Time;
 
 		const uint64 OldSize = AllocationPtr->GetSize();
 
-		SbTree->AddAlloc(AllocationPtr);
+		SbTree->AddAlloc(AllocationPtr); // SbTree takes ownership of AllocationPtr
 
 		uint32 EventDistance = AllocationPtr->EndEventIndex - AllocationPtr->StartEventIndex;
 		UpdateHistogramByEventDistance(EventDistance);
 
-		LiveAllocs.Remove(Address);
-
-		// Update stats for current timeline sample.
+		// Update stats for the current timeline sample.
 		TotalAllocatedMemory -= OldSize;
 		SampleMinTotalAllocatedMemory = FMath::Min(SampleMinTotalAllocatedMemory, TotalAllocatedMemory);
 		SampleMinLiveAllocations = FMath::Min(SampleMinLiveAllocations, (uint32)LiveAllocs.Num());
 		++SampleFreeEvents;
 	}
+	else
+	{
+		++FreeErrors;
+		if (FreeErrors <= MaxLogMessagesPerErrorType)
+		{
+			UE_LOG(LogTraceServices, Error, TEXT("[MemAlloc] Invalid FREE event (Address=0x%llX, Time=%f)!"), Address, Time);
+		}
+	}
 
+	++FreeCount;
 	++EventIndex;
 }
 
@@ -628,23 +970,27 @@ void FAllocationsProvider::AdvanceTimelines(double Time)
 
 ////////////////////////////////////////////////////////////////////////////////////////////////////
 
-void FAllocationsProvider::EditPushRealloc(uint32 ThreadId, uint8 Tracker, uint64 Ptr) 
-{ 
-	EditAccessCheck(); 
+void FAllocationsProvider::EditPushRealloc(uint32 ThreadId, uint8 Tracker, uint64 Ptr)
+{
+	EditAccessCheck();
 	int32 Tag(0); // If ptr is not found use "Untagged"
-	if (FAllocationItem* Alloc = LiveAllocs.Find(Ptr))
+	if (FAllocationItem* Alloc = LiveAllocs.FindRef(Ptr))
 	{
 		Tag = Alloc->Tag;
 	}
 	TagTracker.PushRealloc(ThreadId, Tracker, Tag);
 }
 
-void FAllocationsProvider::EditPopRealloc(uint32 ThreadId, uint8 Tracker) 
-{ 
-	EditAccessCheck(); 
-	TagTracker.PopRealloc(ThreadId, Tracker); 
+////////////////////////////////////////////////////////////////////////////////////////////////////
+
+void FAllocationsProvider::EditPopRealloc(uint32 ThreadId, uint8 Tracker)
+{
+	EditAccessCheck();
+	TagTracker.PopRealloc(ThreadId, Tracker);
 }
 
+////////////////////////////////////////////////////////////////////////////////////////////////////
+
 void FAllocationsProvider::EditOnAnalysisCompleted(double Time)
 {
 	Lock.WriteAccessCheck();
@@ -659,9 +1005,9 @@ void FAllocationsProvider::EditOnAnalysisCompleted(double Time)
 
 	// Add all live allocs to SbTree (with infinite end time).
 	uint64 LiveAllocsTotalSize = 0;
-	for (TPair<uint64, FAllocationItem>& KV : LiveAllocs)
+	LiveAllocs.Enumerate([](const FAllocationItem& Alloc)
 	{
-		FAllocationItem* AllocationPtr = &KV.Value;
+		FAllocationItem* AllocationPtr = const_cast<FAllocationItem*>(&Alloc);
 
 		LiveAllocsTotalSize += AllocationPtr->GetSize();
 
@@ -672,7 +1018,8 @@ void FAllocationsProvider::EditOnAnalysisCompleted(double Time)
 
 		uint32 EventDistance = AllocationPtr->EndEventIndex - AllocationPtr->StartEventIndex;
 		UpdateHistogramByEventDistance(EventDistance);
-	}
+	});
+	//TODO: LiveAllocs.RemoveAll();
 	check(TotalAllocatedMemory == LiveAllocsTotalSize);
 
 	if (bResetTimelineAtEnd)
@@ -997,3 +1344,9 @@ const IAllocationsProvider* ReadAllocationsProvider(const IAnalysisSession& Sess
 ////////////////////////////////////////////////////////////////////////////////////////////////////
 
 } // namespace TraceServices
+
+#undef INSIGHTS_SLOW_CHECK
+#undef INSIGHTS_DEBUG_WATCH
+#undef INSIGHTS_USE_SHORT_LIVING_ALLOCS
+#undef INSIGHTS_USE_LAST_ALLOC
+#undef INSIGHTS_VALIDATE_ALLOC_EVENTS
diff --git a/Engine/Source/Developer/TraceServices/Private/Model/AllocationsProvider.h b/Engine/Source/Developer/TraceServices/Private/Model/AllocationsProvider.h
--- a/Engine/Source/Developer/TraceServices/Private/Model/AllocationsProvider.h
+++ b/Engine/Source/Developer/TraceServices/Private/Model/AllocationsProvider.h
@@ -119,6 +119,80 @@ struct FAllocationCoreItem
 #endif
 ////////////////////////////////////////////////////////////////////////////////////////////////////
 
+class FShortLivingAllocs
+{
+private:
+	struct FNode
+	{
+		FAllocationItem* Alloc;
+		FNode* Next;
+		FNode* Prev;
+	};
+
+	static const int32 MaxAllocCount = 8 * 1024; // max number of short living allocations
+
+public:
+	FShortLivingAllocs();
+	~FShortLivingAllocs();
+
+	bool IsFull() const { return AllocCount == MaxAllocCount; }
+	int32 Num() const { return AllocCount; }
+
+	FORCEINLINE FAllocationItem* FindRef(uint64 Address);
+
+	// The collection keeps ownership of FAllocationItem* until Remove is called or until the oldest allocation is removed.
+	// Returns the removed oldest allocation if collection is already full; nullptr otherwise.
+	// The caller receives ownership of the removed oldest allocation, if a valid pointer is returned.
+	FORCEINLINE FAllocationItem* AddChecked(FAllocationItem* Alloc);
+
+	// The caller takes ownership of FAllocationItem*. Returns nullptr if Address is not found.
+	FORCEINLINE FAllocationItem* Remove(uint64 Address);
+
+	void Enumerate(TFunctionRef<void(const FAllocationItem& Alloc)> Callback) const;
+
+private:
+	TMap<uint64, FNode*> AddressMap; // map for short living allocations: Address -> FNode*
+	FNode* AllNodes = nullptr; // preallocated array of nodes
+	FNode* LastAddedAllocNode = nullptr; // the last added alloc; double linked list: Prev -> .. -> OldestAlloc
+	FNode* OldestAllocNode = nullptr; // the oldest alloc; double linked list: Next -> .. -> LastAddedAlloc
+	FNode* FirstUnusedNode = nullptr; // simple linked list with unused nodes (uses Next)
+	int32 AllocCount = 0;
+};
+
+////////////////////////////////////////////////////////////////////////////////////////////////////
+
+class FLiveAllocCollection
+{
+public:
+	FLiveAllocCollection();
+	~FLiveAllocCollection();
+
+	int32 Num() const { return TotalAllocCount; }
+	int32 PeakCount() const { return MaxAllocCount; }
+
+	FORCEINLINE FAllocationItem* FindRef(uint64 Address);
+
+	// The collection keeps ownership of FAllocationItem* until Remove is called.
+	// Returns the new added allocation.
+	FORCEINLINE FAllocationItem* AddNewChecked(uint64 Address);
+
+	// The caller takes ownership of FAllocationItem*.
+	// Returns nullptr if Address is not found.
+	FORCEINLINE FAllocationItem* Remove(uint64 Address);
+
+	void Enumerate(TFunctionRef<void(const FAllocationItem& Alloc)> Callback) const;
+
+private:
+	FAllocationItem* LastAlloc = nullptr; // last allocation
+	FShortLivingAllocs ShortLivingAllocs; // short living allocs
+	TMap<uint64, FAllocationItem*> LongLivingAllocs; // long living allocations
+
+	int32 TotalAllocCount = 0;
+	int32 MaxAllocCount = 0; // debug stats
+};
+
+////////////////////////////////////////////////////////////////////////////////////////////////////
+
 class FAllocationsProvider : public IAllocationsProvider
 {
 private:
@@ -158,7 +232,7 @@ public:
 
 	const FSbTree* GetSbTree() const { ReadAccessCheck(); return SbTree; }
 
-	const TMap<uint64, FAllocationItem>& GetLiveAllocs()  const { ReadAccessCheck(); return LiveAllocs; }
+	void EnumerateLiveAllocs(TFunctionRef<void(const FAllocationItem& Alloc)> Callback) const { ReadAccessCheck(); return LiveAllocs.Enumerate(Callback); }
 	uint32 GetNumLiveAllocs() const { ReadAccessCheck(); return LiveAllocs.Num(); }
 
 	bool HasReallocScope(uint32 ThreadId, uint8 Tracker) const { ReadAccessCheck(); return TagTracker.HasReallocScope(ThreadId, Tracker); }
@@ -212,8 +286,7 @@ private:
 	uint64 AllocCount = 0;
 	uint64 FreeCount = 0;
 
-	uint32 MaxLiveAllocCount = 0;
-	TMap<uint64, FAllocationItem> LiveAllocs;
+	FLiveAllocCollection LiveAllocs;
 
 	uint64 AllocErrors = 0;
 	uint64 FreeErrors = 0;
diff --git a/Engine/Source/Developer/TraceServices/Private/Model/AllocationsQuery.cpp b/Engine/Source/Developer/TraceServices/Private/Model/AllocationsQuery.cpp
--- a/Engine/Source/Developer/TraceServices/Private/Model/AllocationsQuery.cpp
+++ b/Engine/Source/Developer/TraceServices/Private/Model/AllocationsQuery.cpp
@@ -122,7 +122,6 @@ void FAllocationsQuery::Run()
 			TotalAllocationCount += NumLiveAllocs;
 
 			QueryCallstacks(LiveAllocsResult);
-
 			Results.Enqueue(LiveAllocsResult);
 		}
 		else
@@ -159,10 +158,10 @@ void FAllocationsQuery::Run()
 			const uint32 NumAllocs = static_cast<uint32>(CellResult->Items.Num());
 			if (NumAllocs != 0)
 			{
-				QueryCallstacks(CellResult);
-
 				UE_LOG(LogTraceServices, Log, TEXT("[MemAlloc] Enqueue %u allocs..."), NumAllocs);
 				TotalAllocationCount += NumAllocs;
+
+				QueryCallstacks(CellResult);
 				Results.Enqueue(CellResult);
 			}
 			else
@@ -191,8 +190,6 @@ void FAllocationsQuery::Run()
 
 void FAllocationsQuery::QueryLiveAllocs(TArray<const FAllocationItem*>& OutAllocs) const
 {
-	const TMap<uint64, FAllocationItem>& LiveAllocs = AllocationsProvider.GetLiveAllocs();
-
 	// TODO: Live allocs have the end time set to std::numeric_limits<double>::infinity(),
 	//       so the conditions below could be simplified to remove the end time checks.
 	//       For now, the code remains "unoptimized" for debugging purposes.
@@ -202,42 +199,39 @@ void FAllocationsQuery::QueryLiveAllocs(TArray<const FAllocationItem*>& OutAlloc
 	case IAllocationsProvider::EQueryRule::aAf: // active allocs at A
 	{
 		const double Time = Params.TimeA;
-		for (const auto& KV : LiveAllocs)
+		AllocationsProvider.EnumerateLiveAllocs([Time, &OutAllocs](const FAllocationItem& Alloc)
 		{
-			const FAllocationItem& Alloc = KV.Value;
 			if (Alloc.StartTime <= Time && Time <= Alloc.EndTime)
 			{
 				OutAllocs.Add(&Alloc);
 			}
-		}
+		});
 	}
 	break;
 
 	case IAllocationsProvider::EQueryRule::afA: // before
 	{
 		const double Time = Params.TimeA;
-		for (const auto& KV : LiveAllocs)
+		AllocationsProvider.EnumerateLiveAllocs([Time, &OutAllocs](const FAllocationItem& Alloc)
 		{
-			const FAllocationItem& Alloc = KV.Value;
 			if (Alloc.EndTime <= Time)
 			{
 				OutAllocs.Add(&Alloc);
 			}
-		}
+		});
 	}
 	break;
 
 	case IAllocationsProvider::EQueryRule::Aaf: // after
 	{
 		const double Time = Params.TimeA;
-		for (const auto& KV : LiveAllocs)
+		AllocationsProvider.EnumerateLiveAllocs([Time, &OutAllocs](const FAllocationItem& Alloc)
 		{
-			const FAllocationItem& Alloc = KV.Value;
 			if (Alloc.StartTime >= Time)
 			{
 				OutAllocs.Add(&Alloc);
 			}
-		}
+		});
 	}
 	break;
 
@@ -245,14 +239,13 @@ void FAllocationsQuery::QueryLiveAllocs(TArray<const FAllocationItem*>& OutAlloc
 	{
 		const double TimeA = Params.TimeA;
 		const double TimeB = Params.TimeB;
-		for (const auto& KV : LiveAllocs)
+		AllocationsProvider.EnumerateLiveAllocs([TimeA, TimeB, &OutAllocs](const FAllocationItem& Alloc)
 		{
-			const FAllocationItem& Alloc = KV.Value;
 			if (Alloc.StartTime <= TimeA && Alloc.EndTime >= TimeA && Alloc.EndTime <= TimeB)
 			{
 				OutAllocs.Add(&Alloc);
 			}
-		}
+		});
 	}
 	break;
 
@@ -260,14 +253,13 @@ void FAllocationsQuery::QueryLiveAllocs(TArray<const FAllocationItem*>& OutAlloc
 	{
 		const double TimeA = Params.TimeA;
 		const double TimeB = Params.TimeB;
-		for (const auto& KV : LiveAllocs)
+		AllocationsProvider.EnumerateLiveAllocs([TimeA, TimeB, &OutAllocs](const FAllocationItem& Alloc)
 		{
-			const FAllocationItem& Alloc = KV.Value;
 			if (Alloc.StartTime >= TimeA && Alloc.StartTime <= TimeB && Alloc.EndTime >= TimeB)
 			{
 				OutAllocs.Add(&Alloc);
 			}
-		}
+		});
 	}
 	break;
 
@@ -275,14 +267,13 @@ void FAllocationsQuery::QueryLiveAllocs(TArray<const FAllocationItem*>& OutAlloc
 	{
 		const double TimeA = Params.TimeA;
 		const double TimeB = Params.TimeB;
-		for (const auto& KV : LiveAllocs)
+		AllocationsProvider.EnumerateLiveAllocs([TimeA, TimeB, &OutAllocs](const FAllocationItem& Alloc)
 		{
-			const FAllocationItem& Alloc = KV.Value;
 			if (Alloc.EndTime >= TimeA && Alloc.EndTime <= TimeB)
 			{
 				OutAllocs.Add(&Alloc);
 			}
-		}
+		});
 	}
 	break;
 
@@ -290,14 +281,13 @@ void FAllocationsQuery::QueryLiveAllocs(TArray<const FAllocationItem*>& OutAlloc
 	{
 		const double TimeA = Params.TimeA;
 		const double TimeB = Params.TimeB;
-		for (const auto& KV : LiveAllocs)
+		AllocationsProvider.EnumerateLiveAllocs([TimeA, TimeB, &OutAllocs](const FAllocationItem& Alloc)
 		{
-			const FAllocationItem& Alloc = KV.Value;
 			if (Alloc.StartTime >= TimeA && Alloc.StartTime <= TimeB)
 			{
 				OutAllocs.Add(&Alloc);
 			}
-		}
+		});
 	}
 	break;
 
@@ -305,14 +295,13 @@ void FAllocationsQuery::QueryLiveAllocs(TArray<const FAllocationItem*>& OutAlloc
 	{
 		const double TimeA = Params.TimeA;
 		const double TimeB = Params.TimeB;
-		for (const auto& KV : LiveAllocs)
+		AllocationsProvider.EnumerateLiveAllocs([TimeA, TimeB, &OutAllocs](const FAllocationItem& Alloc)
 		{
-			const FAllocationItem& Alloc = KV.Value;
 			if (Alloc.StartTime >= TimeA && Alloc.EndTime <= TimeB)
 			{
 				OutAllocs.Add(&Alloc);
 			}
-		}
+		});
 	}
 	break;
 
@@ -320,14 +309,13 @@ void FAllocationsQuery::QueryLiveAllocs(TArray<const FAllocationItem*>& OutAlloc
 	{
 		const double TimeA = Params.TimeA;
 		const double TimeB = Params.TimeB;
-		for (const auto& KV : LiveAllocs)
+		AllocationsProvider.EnumerateLiveAllocs([TimeA, TimeB, &OutAllocs](const FAllocationItem& Alloc)
 		{
-			const FAllocationItem& Alloc = KV.Value;
 			if (Alloc.StartTime <= TimeA && Alloc.EndTime >= TimeB)
 			{
 				OutAllocs.Add(&Alloc);
 			}
-		}
+		});
 	}
 	break;
 
@@ -336,14 +324,13 @@ void FAllocationsQuery::QueryLiveAllocs(TArray<const FAllocationItem*>& OutAlloc
 		const double TimeA = Params.TimeA;
 		const double TimeB = Params.TimeB;
 		const double TimeC = Params.TimeC;
-		for (const auto& KV : LiveAllocs)
+		AllocationsProvider.EnumerateLiveAllocs([TimeA, TimeB, TimeC, &OutAllocs](const FAllocationItem& Alloc)
 		{
-			const FAllocationItem& Alloc = KV.Value;
 			if (Alloc.StartTime >= TimeA && Alloc.StartTime <= TimeB && Alloc.EndTime >= TimeC)
 			{
 				OutAllocs.Add(&Alloc);
 			}
-		}
+		});
 	}
 	break;
 
@@ -352,14 +339,13 @@ void FAllocationsQuery::QueryLiveAllocs(TArray<const FAllocationItem*>& OutAlloc
 		const double TimeA = Params.TimeA;
 		const double TimeB = Params.TimeB;
 		const double TimeC = Params.TimeC;
-		for (const auto& KV : LiveAllocs)
+		AllocationsProvider.EnumerateLiveAllocs([TimeA, TimeB, TimeC, &OutAllocs](const FAllocationItem& Alloc)
 		{
-			const FAllocationItem& Alloc = KV.Value;
 			if (Alloc.StartTime >= TimeA && Alloc.StartTime <= TimeB && Alloc.EndTime >= TimeB && Alloc.EndTime <= TimeC)
 			{
 				OutAllocs.Add(&Alloc);
 			}
-		}
+		});
 	}
 	break;
 
@@ -368,14 +354,13 @@ void FAllocationsQuery::QueryLiveAllocs(TArray<const FAllocationItem*>& OutAlloc
 		const double TimeA = Params.TimeA;
 		const double TimeB = Params.TimeB;
 		const double TimeC = Params.TimeC;
-		for (const auto& KV : LiveAllocs)
+		AllocationsProvider.EnumerateLiveAllocs([TimeA, TimeB, TimeC, &OutAllocs](const FAllocationItem& Alloc)
 		{
-			const FAllocationItem& Alloc = KV.Value;
 			if (Alloc.StartTime <= TimeA && Alloc.EndTime >= TimeB && Alloc.EndTime <= TimeC)
 			{
 				OutAllocs.Add(&Alloc);
 			}
-		}
+		});
 	}
 	break;
 
@@ -385,14 +370,13 @@ void FAllocationsQuery::QueryLiveAllocs(TArray<const FAllocationItem*>& OutAlloc
 		const double TimeB = Params.TimeB;
 		const double TimeC = Params.TimeC;
 		const double TimeD = Params.TimeD;
-		for (const auto& KV : LiveAllocs)
+		AllocationsProvider.EnumerateLiveAllocs([TimeA, TimeB, TimeC, TimeD, &OutAllocs](const FAllocationItem& Alloc)
 		{
-			const FAllocationItem& Alloc = KV.Value;
 			if (Alloc.StartTime >= TimeA && Alloc.StartTime <= TimeB && Alloc.EndTime >= TimeC && Alloc.EndTime <= TimeD)
 			{
 				OutAllocs.Add(&Alloc);
 			}
-		}
+		});
 	}
 	break;
 	}
@@ -404,11 +388,12 @@ void FAllocationsQuery::QueryCallstacks(FAllocationsImpl* Result) const
 {
 	for (auto Item : Result->Items)
 	{
-		// Callstacks could have been resolved in previous queries, 
+		// Callstacks could have been resolved in previous queries,
 		// check before querying again
 		if (!Item->Callstack)
 		{
 			Item->Callstack = CallstacksProvider.GetCallstack(Item->Owner);
+			check(Item->Callstack != nullptr);
 		}
 	}
 }
diff --git a/Engine/Source/Developer/TraceServices/Private/Model/SbTree.cpp b/Engine/Source/Developer/TraceServices/Private/Model/SbTree.cpp
--- a/Engine/Source/Developer/TraceServices/Private/Model/SbTree.cpp
+++ b/Engine/Source/Developer/TraceServices/Private/Model/SbTree.cpp
@@ -62,9 +62,19 @@ FSbTreeCell::FSbTreeCell(ILinearAllocator& InAllocator)
 
 ////////////////////////////////////////////////////////////////////////////////////////////////////
 
+FSbTreeCell::~FSbTreeCell()
+{
+	for (const FAllocationItem* Alloc : Allocs)
+	{
+		delete Alloc;
+	}
+}
+
+////////////////////////////////////////////////////////////////////////////////////////////////////
+
 void FSbTreeCell::AddAlloc(const FAllocationItem* Alloc)
 {
-	Allocs.Add(*Alloc);
+	Allocs.Add(Alloc);
 
 	if (Alloc->StartEventIndex < MinStartEventIndex)
 	{
@@ -94,11 +104,11 @@ void FSbTreeCell::Query(TArray<const FAllocationItem*>& OutAllocs, const IAlloca
 	case IAllocationsProvider::EQueryRule::aAf: // active allocs at A
 	{
 		const double Time = Params.TimeA;
-		for (const FAllocationItem& Alloc : Allocs)
+		for (const FAllocationItem* Alloc : Allocs)
 		{
-			if (Alloc.StartTime <= Time && Time <= Alloc.EndTime)
+			if (Alloc->StartTime <= Time && Time <= Alloc->EndTime)
 			{
-				OutAllocs.Add(&Alloc);
+				OutAllocs.Add(Alloc);
 			}
 		}
 	}
@@ -107,11 +117,11 @@ void FSbTreeCell::Query(TArray<const FAllocationItem*>& OutAllocs, const IAlloca
 	case IAllocationsProvider::EQueryRule::afA: // before
 	{
 		const double Time = Params.TimeA;
-		for (const FAllocationItem& Alloc : Allocs)
+		for (const FAllocationItem* Alloc : Allocs)
 		{
-			if (Alloc.EndTime <= Time)
+			if (Alloc->EndTime <= Time)
 			{
-				OutAllocs.Add(&Alloc);
+				OutAllocs.Add(Alloc);
 			}
 		}
 	}
@@ -120,11 +130,11 @@ void FSbTreeCell::Query(TArray<const FAllocationItem*>& OutAllocs, const IAlloca
 	case IAllocationsProvider::EQueryRule::Aaf: // after
 	{
 		const double Time = Params.TimeA;
-		for (const FAllocationItem& Alloc : Allocs)
+		for (const FAllocationItem* Alloc : Allocs)
 		{
-			if (Alloc.StartTime >= Time)
+			if (Alloc->StartTime >= Time)
 			{
-				OutAllocs.Add(&Alloc);
+				OutAllocs.Add(Alloc);
 			}
 		}
 	}
@@ -134,11 +144,11 @@ void FSbTreeCell::Query(TArray<const FAllocationItem*>& OutAllocs, const IAlloca
 	{
 		const double TimeA = Params.TimeA;
 		const double TimeB = Params.TimeB;
-		for (const FAllocationItem& Alloc : Allocs)
+		for (const FAllocationItem* Alloc : Allocs)
 		{
-			if (Alloc.StartTime <= TimeA && Alloc.EndTime >= TimeA && Alloc.EndTime <= TimeB)
+			if (Alloc->StartTime <= TimeA && Alloc->EndTime >= TimeA && Alloc->EndTime <= TimeB)
 			{
-				OutAllocs.Add(&Alloc);
+				OutAllocs.Add(Alloc);
 			}
 		}
 	}
@@ -148,11 +158,11 @@ void FSbTreeCell::Query(TArray<const FAllocationItem*>& OutAllocs, const IAlloca
 	{
 		const double TimeA = Params.TimeA;
 		const double TimeB = Params.TimeB;
-		for (const FAllocationItem& Alloc : Allocs)
+		for (const FAllocationItem* Alloc : Allocs)
 		{
-			if (Alloc.StartTime >= TimeA && Alloc.StartTime <= TimeB && Alloc.EndTime >= TimeB)
+			if (Alloc->StartTime >= TimeA && Alloc->StartTime <= TimeB && Alloc->EndTime >= TimeB)
 			{
-				OutAllocs.Add(&Alloc);
+				OutAllocs.Add(Alloc);
 			}
 		}
 	}
@@ -162,11 +172,11 @@ void FSbTreeCell::Query(TArray<const FAllocationItem*>& OutAllocs, const IAlloca
 	{
 		const double TimeA = Params.TimeA;
 		const double TimeB = Params.TimeB;
-		for (const FAllocationItem& Alloc : Allocs)
+		for (const FAllocationItem* Alloc : Allocs)
 		{
-			if (Alloc.EndTime >= TimeA && Alloc.EndTime <= TimeB)
+			if (Alloc->EndTime >= TimeA && Alloc->EndTime <= TimeB)
 			{
-				OutAllocs.Add(&Alloc);
+				OutAllocs.Add(Alloc);
 			}
 		}
 	}
@@ -176,11 +186,11 @@ void FSbTreeCell::Query(TArray<const FAllocationItem*>& OutAllocs, const IAlloca
 	{
 		const double TimeA = Params.TimeA;
 		const double TimeB = Params.TimeB;
-		for (const FAllocationItem& Alloc : Allocs)
+		for (const FAllocationItem* Alloc : Allocs)
 		{
-			if (Alloc.StartTime >= TimeA && Alloc.StartTime <= TimeB)
+			if (Alloc->StartTime >= TimeA && Alloc->StartTime <= TimeB)
 			{
-				OutAllocs.Add(&Alloc);
+				OutAllocs.Add(Alloc);
 			}
 		}
 	}
@@ -190,11 +200,11 @@ void FSbTreeCell::Query(TArray<const FAllocationItem*>& OutAllocs, const IAlloca
 	{
 		const double TimeA = Params.TimeA;
 		const double TimeB = Params.TimeB;
-		for (const FAllocationItem& Alloc : Allocs)
+		for (const FAllocationItem* Alloc : Allocs)
 		{
-			if (Alloc.StartTime >= TimeA && Alloc.EndTime <= TimeB)
+			if (Alloc->StartTime >= TimeA && Alloc->EndTime <= TimeB)
 			{
-				OutAllocs.Add(&Alloc);
+				OutAllocs.Add(Alloc);
 			}
 		}
 	}
@@ -204,11 +214,11 @@ void FSbTreeCell::Query(TArray<const FAllocationItem*>& OutAllocs, const IAlloca
 	{
 		const double TimeA = Params.TimeA;
 		const double TimeB = Params.TimeB;
-		for (const FAllocationItem& Alloc : Allocs)
+		for (const FAllocationItem* Alloc : Allocs)
 		{
-			if (Alloc.StartTime <= TimeA && Alloc.EndTime >= TimeB)
+			if (Alloc->StartTime <= TimeA && Alloc->EndTime >= TimeB)
 			{
-				OutAllocs.Add(&Alloc);
+				OutAllocs.Add(Alloc);
 			}
 		}
 	}
@@ -219,11 +229,11 @@ void FSbTreeCell::Query(TArray<const FAllocationItem*>& OutAllocs, const IAlloca
 		const double TimeA = Params.TimeA;
 		const double TimeB = Params.TimeB;
 		const double TimeC = Params.TimeC;
-		for (const FAllocationItem& Alloc : Allocs)
+		for (const FAllocationItem* Alloc : Allocs)
 		{
-			if (Alloc.StartTime >= TimeA && Alloc.StartTime <= TimeB && Alloc.EndTime >= TimeC)
+			if (Alloc->StartTime >= TimeA && Alloc->StartTime <= TimeB && Alloc->EndTime >= TimeC)
 			{
-				OutAllocs.Add(&Alloc);
+				OutAllocs.Add(Alloc);
 			}
 		}
 	}
@@ -234,11 +244,11 @@ void FSbTreeCell::Query(TArray<const FAllocationItem*>& OutAllocs, const IAlloca
 		const double TimeA = Params.TimeA;
 		const double TimeB = Params.TimeB;
 		const double TimeC = Params.TimeC;
-		for (const FAllocationItem& Alloc : Allocs)
+		for (const FAllocationItem* Alloc : Allocs)
 		{
-			if (Alloc.StartTime >= TimeA && Alloc.StartTime <= TimeB && Alloc.EndTime >= TimeB && Alloc.EndTime <= TimeC)
+			if (Alloc->StartTime >= TimeA && Alloc->StartTime <= TimeB && Alloc->EndTime >= TimeB && Alloc->EndTime <= TimeC)
 			{
-				OutAllocs.Add(&Alloc);
+				OutAllocs.Add(Alloc);
 			}
 		}
 	}
@@ -249,11 +259,11 @@ void FSbTreeCell::Query(TArray<const FAllocationItem*>& OutAllocs, const IAlloca
 		const double TimeA = Params.TimeA;
 		const double TimeB = Params.TimeB;
 		const double TimeC = Params.TimeC;
-		for (const FAllocationItem& Alloc : Allocs)
+		for (const FAllocationItem* Alloc : Allocs)
 		{
-			if (Alloc.StartTime <= TimeA && Alloc.EndTime >= TimeB && Alloc.EndTime <= TimeC)
+			if (Alloc->StartTime <= TimeA && Alloc->EndTime >= TimeB && Alloc->EndTime <= TimeC)
 			{
-				OutAllocs.Add(&Alloc);
+				OutAllocs.Add(Alloc);
 			}
 		}
 	}
@@ -265,11 +275,11 @@ void FSbTreeCell::Query(TArray<const FAllocationItem*>& OutAllocs, const IAlloca
 		const double TimeB = Params.TimeB;
 		const double TimeC = Params.TimeC;
 		const double TimeD = Params.TimeD;
-		for (const FAllocationItem& Alloc : Allocs)
+		for (const FAllocationItem* Alloc : Allocs)
 		{
-			if (Alloc.StartTime >= TimeA && Alloc.StartTime <= TimeB && Alloc.EndTime >= TimeC && Alloc.EndTime <= TimeD)
+			if (Alloc->StartTime >= TimeA && Alloc->StartTime <= TimeB && Alloc->EndTime >= TimeC && Alloc->EndTime <= TimeD)
 			{
-				OutAllocs.Add(&Alloc);
+				OutAllocs.Add(Alloc);
 			}
 		}
 	}
diff --git a/Engine/Source/Developer/TraceServices/Private/Model/SbTree.h b/Engine/Source/Developer/TraceServices/Private/Model/SbTree.h
--- a/Engine/Source/Developer/TraceServices/Private/Model/SbTree.h
+++ b/Engine/Source/Developer/TraceServices/Private/Model/SbTree.h
@@ -27,9 +27,11 @@ class FSbTreeCell
 {
 public:
 	FSbTreeCell(ILinearAllocator& InAllocator);
+	~FSbTreeCell();
 
 	uint32 GetAllocCount() const { return Allocs.Num(); }
-	void AddAlloc(const FAllocationItem* Alloc);
+
+	void AddAlloc(const FAllocationItem* Alloc); // takes ownership of memory allocated by Alloc pointer
 
 	uint32 GetMinStartEventIndex() const { return MinStartEventIndex; }
 	uint32 GetMaxEndEventIndex() const { return MaxEndEventIndex; }
@@ -43,7 +45,7 @@ private:
 	ILinearAllocator& Allocator;
 
 	//TODO: TPagedArray<FAllocationItem> Allocs;
-	TArray<FAllocationItem> Allocs;
+	TArray<const FAllocationItem*> Allocs;
 
 	uint32 MinStartEventIndex;
 	uint32 MaxEndEventIndex;
@@ -61,7 +63,7 @@ public:
 
 	void SetTimeForEvent(uint32 EventIndex, double Time);
 
-	void AddAlloc(const FAllocationItem* Alloc);
+	void AddAlloc(const FAllocationItem* Alloc); // takes ownership of memory allocated by Alloc pointer
 
 	uint32 GetColumnWidth() const { return 1 << ColumnShift; }
 	uint32 GetCurrentColumn() const { return CurrentColumn; }
