[Insights] Fixed false positive IncludeTool error.

diff --git a/Engine/Source/Developer/TraceServices/Private/Model/AllocationsProvider.cpp b/Engine/Source/Developer/TraceServices/Private/Model/AllocationsProvider.cpp
--- a/Engine/Source/Developer/TraceServices/Private/Model/AllocationsProvider.cpp
+++ b/Engine/Source/Developer/TraceServices/Private/Model/AllocationsProvider.cpp
@@ -768,16 +768,16 @@ void FAllocationsProvider::EditAlloc(double Time, uint64 Owner, uint64 Address,
 			break;
 		}
 	}
-#endif // INSIGHTS_DEBUG_WATCH
+#endif
 
 #if INSIGHTS_VALIDATE_ALLOC_EVENTS
 	FAllocationItem* AllocationPtr = LiveAllocs.FindRef(Address);
-	if (!AllocationPtr)
-	{
 #else
-		FAllocationItem* AllocationPtr;
+	FAllocationItem* AllocationPtr = nullptr;
 #endif
 
+	if (!AllocationPtr)
+	{
 		AllocationPtr = LiveAllocs.AddNewChecked(Address);
 		INSIGHTS_SLOW_CHECK(Allocation.Address == Address)
 
@@ -800,9 +800,8 @@ void FAllocationsProvider::EditAlloc(double Time, uint64 Owner, uint64 Address,
 		SampleMaxTotalAllocatedMemory = FMath::Max(SampleMaxTotalAllocatedMemory, TotalAllocatedMemory);
 		SampleMaxLiveAllocations = FMath::Max(SampleMaxLiveAllocations, (uint32)LiveAllocs.Num());
 		++SampleAllocEvents;
-
-#if INSIGHTS_VALIDATE_ALLOC_EVENTS
 	}
+#if INSIGHTS_VALIDATE_ALLOC_EVENTS
 	else
 	{
 		++AllocErrors;
