TasksInsights: fixed a minor bug

diff --git a/Engine/Source/Developer/TraceServices/Public/TraceServices/Model/TasksProfiler.h b/Engine/Source/Developer/TraceServices/Public/TraceServices/Model/TasksProfiler.h
--- a/Engine/Source/Developer/TraceServices/Public/TraceServices/Model/TasksProfiler.h
+++ b/Engine/Source/Developer/TraceServices/Public/TraceServices/Model/TasksProfiler.h
@@ -69,7 +69,7 @@ namespace TraceServices
 	{
 		TArray<TaskTrace::FId> Tasks;
 		double StartedTimestamp;
-		double FinishedTimestamp;
+		double FinishedTimestamp = FTaskInfo::InvalidTimestamp;
 	};
 
 	// query interface to tasks info
