Segfault fix. The moduleprovider is conditionally created depending on the platform that is hosting Unreal Insights and the trace's symbol format. As such it could be null.

#rnx

diff --git a/Engine/Source/Developer/TraceServices/Private/Analyzers/ModuleAnalysis.cpp b/Engine/Source/Developer/TraceServices/Private/Analyzers/ModuleAnalysis.cpp
--- a/Engine/Source/Developer/TraceServices/Private/Analyzers/ModuleAnalysis.cpp
+++ b/Engine/Source/Developer/TraceServices/Private/Analyzers/ModuleAnalysis.cpp
@@ -62,6 +62,7 @@ bool FModuleAnalyzer::OnEvent(uint16 RouteId, EStyle Style, const FOnEventContex
 			break;
 
 		case RouteId_ModuleLoad:
+			if (Provider != nullptr)
 			{
 				const uint64 Base = GetRealBaseAddress(Context.EventData.GetValue<uint32>("Base"));
 				const uint32 Size = Context.EventData.GetValue<uint32>("Size");
@@ -79,6 +80,7 @@ bool FModuleAnalyzer::OnEvent(uint16 RouteId, EStyle Style, const FOnEventContex
 			break;
 
 		case RouteId_ModuleUnload:
+			if (Provider != nullptr)
 			{
 				const uint32 Base = Context.EventData.GetValue<uint32>("Base");
 				if (Provider)
diff --git a/Engine/Source/Developer/TraceServices/Private/Model/CallstacksProvider.cpp b/Engine/Source/Developer/TraceServices/Private/Model/CallstacksProvider.cpp
--- a/Engine/Source/Developer/TraceServices/Private/Model/CallstacksProvider.cpp
+++ b/Engine/Source/Developer/TraceServices/Private/Model/CallstacksProvider.cpp
@@ -48,7 +48,8 @@ void FCallstacksProvider::AddCallstack(uint64 InCallstackId, const uint64* InFra
 #endif
 
 	// The module provider is created on the fly so we want to cache it 
-	// once it's available.
+	// once it's available. Note that the module provider is conditionally
+	// created so EditProvider() may return a null pointer.
 	if (!ModuleProvider)
 	{
 		ModuleProvider = Session.EditProvider<IModuleProvider>(FName("ModuleProvider"));
