[Insights] RAD symbols resolver implementation

Implements a resolver using RAD symbol parsing library.

diff --git a/Engine/Source/Developer/TraceServices/Private/Analyzers/ModuleAnalysis.cpp b/Engine/Source/Developer/TraceServices/Private/Analyzers/ModuleAnalysis.cpp
--- a/Engine/Source/Developer/TraceServices/Private/Analyzers/ModuleAnalysis.cpp
+++ b/Engine/Source/Developer/TraceServices/Private/Analyzers/ModuleAnalysis.cpp
@@ -40,6 +40,7 @@ void FModuleAnalyzer::OnAnalysisEnd()
 	}
 }
 
+
 ////////////////////////////////////////////////////////////////////////////////
 bool FModuleAnalyzer::OnEvent(uint16 RouteId, EStyle Style, const FOnEventContext& Context)
 {
@@ -64,16 +65,18 @@ bool FModuleAnalyzer::OnEvent(uint16 RouteId, EStyle Style, const FOnEventContex
 		case RouteId_ModuleLoad:
 			if (Provider != nullptr)
 			{
-				const uint64 Base = GetRealBaseAddress(Context.EventData.GetValue<uint32>("Base"));
-				const uint32 Size = Context.EventData.GetValue<uint32>("Size");
-
 				//todo: Use string store
 				FStringView ModuleName;
 				if (Context.EventData.GetString("Name", ModuleName))
 				{
+					const uint64 Base = GetRealBaseAddress(Context.EventData.GetValue<uint32>("Base"));
+					const uint32 Size = Context.EventData.GetValue<uint32>("Size");
+					const uint32 Age = Context.EventData.GetValue<uint32>("Age", 0);
+					const IAnalyzer::TArrayReader<uint8>& ImageId = Context.EventData.GetArray<uint8>("ImageId");
+
 					if (Provider)
 					{
-						Provider->OnModuleLoad(ModuleName, Base, Size);
+						Provider->OnModuleLoad(ModuleName, Base, Size, ImageId.GetData(), ImageId.Num());
 					}
 				}
 			}
diff --git a/Engine/Source/Developer/TraceServices/Private/Model/DbgHelpResolver.cpp b/Engine/Source/Developer/TraceServices/Private/Model/DbgHelpResolver.cpp
new file mode 100644
--- /dev/null
+++ b/Engine/Source/Developer/TraceServices/Private/Model/DbgHelpResolver.cpp
@@ -0,0 +1,283 @@
+// Copyright Epic Games, Inc. All Rights Reserved.
+
+#if PLATFORM_WINDOWS
+
+#include "DbgHelpResolver.h"
+#include "Algo/Sort.h"
+#include "Containers/Queue.h"
+#include "Containers/StringView.h"
+#include "HAL/PlatformProcess.h"
+#include "HAL/RunnableThread.h"
+#include "Logging/LogMacros.h"
+#include "Misc/PathViews.h"
+#include "Misc/ScopeLock.h"
+#include "Misc/StringBuilder.h"
+#include "TraceServices/Model/AnalysisSession.h"
+#include <atomic>
+
+#include "Windows/AllowWindowsPlatformTypes.h"
+#include <dbghelp.h>
+
+/////////////////////////////////////////////////////////////////////
+DEFINE_LOG_CATEGORY_STATIC(LogDbgHelp, Log, All);
+
+/////////////////////////////////////////////////////////////////////
+namespace TraceServices {
+
+/////////////////////////////////////////////////////////////////////
+
+static const TCHAR* GUnknownModuleTextDbgHelp = TEXT("Unknown");
+
+/////////////////////////////////////////////////////////////////////
+FDbgHelpResolver::FDbgHelpResolver(IAnalysisSession& InSession)
+	: bRunWorkerThread(false)
+	, bDrainThenStop(false)
+	, bInitialized(false)
+	, Session(InSession)
+{
+	bInitialized = SetupSyms();
+
+	if (bInitialized)
+	{
+		// Start the worker thread
+		bRunWorkerThread = true;
+		Thread = FRunnableThread::Create(this, TEXT("DbgHelpWorker"), 0, TPri_Normal);
+	}
+}
+
+/////////////////////////////////////////////////////////////////////
+FDbgHelpResolver::~FDbgHelpResolver()
+{
+	bRunWorkerThread = false;
+	if (Thread)
+	{
+		Thread->WaitForCompletion();
+	}
+}
+
+/////////////////////////////////////////////////////////////////////
+void FDbgHelpResolver::QueueModuleLoad(const FStringView& ModulePath, uint64 Base, uint32 Size, const uint8* ImageId, uint32 ImageIdSize)
+{
+	FScopeLock _(&ModulesCs);
+
+	const FStringView ModuleName = FPathViews::GetCleanFilename(ModulePath);
+
+	// Add module and sort list according to base address
+	const int32 Index = LoadedModules.Add(FModuleEntry{ Base, Size, Session.StoreString(ModuleName), Session.StoreString(ModulePath), false });
+
+	// Queue up module to have symbols loaded
+	LoadSymbolsQueue.Enqueue(FQueuedModule{ Base, Size, LoadedModules[Index].Path });
+
+	// Sort list according to base address
+	Algo::Sort(LoadedModules, [](const FModuleEntry& Lhs, const FModuleEntry& Rhs) { return Lhs.Base < Rhs.Base; });
+
+	++ModulesDiscovered;
+}
+
+/////////////////////////////////////////////////////////////////////
+void FDbgHelpResolver::QueueSymbolResolve(uint64 Address, FResolvedSymbol* Symbol)
+{
+	++SymbolsDiscovered;
+	ResolveQueue.Enqueue(FQueuedAddress{Address, Symbol});
+}
+
+/////////////////////////////////////////////////////////////////////
+void FDbgHelpResolver::GetStats(IModuleProvider::FStats* OutStats) const
+{
+	OutStats->ModulesDiscovered = ModulesDiscovered.load();
+	OutStats->ModulesFailed = ModulesFailed.load();
+	OutStats->ModulesLoaded = ModulesLoaded.load();
+	OutStats->SymbolsDiscovered = SymbolsDiscovered.load();
+	OutStats->SymbolsFailed = SymbolsFailed.load();
+	OutStats->SymbolsResolved = SymbolsResolved.load();
+}
+
+/////////////////////////////////////////////////////////////////////
+void FDbgHelpResolver::OnAnalysisComplete()
+{
+	// At this point no more module loads or symbol requests will be queued,
+	// we drain the current queue, then release resources and file locks.
+	bDrainThenStop = true;
+}
+
+
+/////////////////////////////////////////////////////////////////////
+bool FDbgHelpResolver::SetupSyms() 
+{
+	// Create a unique handle
+	static UPTRINT BaseHandle = 0x493;
+	Handle = ++BaseHandle;
+
+	// Load DbgHelp interface
+	ULONG SymOpts = 0;
+	SymOpts |= SYMOPT_LOAD_LINES;
+	SymOpts |= SYMOPT_OMAP_FIND_NEAREST;
+	SymOpts |= SYMOPT_DEFERRED_LOADS;
+	SymOpts |= SYMOPT_EXACT_SYMBOLS;
+	SymOpts |= SYMOPT_IGNORE_NT_SYMPATH;
+	SymOpts |= SYMOPT_UNDNAME;
+
+	SymSetOptions(SymOpts);
+	return SymInitialize((HANDLE)Handle, NULL, FALSE);
+}
+
+
+/////////////////////////////////////////////////////////////////////
+void FDbgHelpResolver::FreeSyms() const
+{
+	// This release file locks on debug files
+	SymCleanup((HANDLE)Handle);
+}
+
+
+/////////////////////////////////////////////////////////////////////
+uint32 FDbgHelpResolver::Run()
+{
+	while (bRunWorkerThread)
+	{
+		// Prioritize queued module loads
+		while (!LoadSymbolsQueue.IsEmpty() && bRunWorkerThread)
+		{
+			FQueuedModule Item;
+			if (LoadSymbolsQueue.Dequeue(Item))
+			{
+				LoadModuleSymbols(Item.Base, Item.Size, Item.ImagePath);
+			}
+		}
+
+		// Resolve one symbol at a time to give way for modules
+		while (!ResolveQueue.IsEmpty() && LoadSymbolsQueue.IsEmpty() && bRunWorkerThread)
+		{
+			FQueuedAddress Item;
+			if (ResolveQueue.Dequeue(Item))
+			{
+				ResolveSymbol(Item.Address, Item.Target);
+			}
+		}
+
+		if (bDrainThenStop && ResolveQueue.IsEmpty() && LoadSymbolsQueue.IsEmpty())
+		{
+			bRunWorkerThread = false;
+		}
+		
+		// ...and breathe...
+		FPlatformProcess::Sleep(0.2f);
+	}
+
+	// We don't need the syms library anymore
+	FreeSyms();
+	
+	return 0;
+}
+
+/////////////////////////////////////////////////////////////////////
+void FDbgHelpResolver::Stop()
+{
+	bRunWorkerThread = false; 
+}
+
+/////////////////////////////////////////////////////////////////////
+void FDbgHelpResolver::UpdateResolvedSymbol(FResolvedSymbol* Symbol, ESymbolQueryResult Result, const TCHAR* Name, const TCHAR* FileAndLine)
+{
+	Symbol->Name = Name;
+	Symbol->FileAndLine = FileAndLine;
+	Symbol->Result.store(Result, std::memory_order_release);
+}
+
+/////////////////////////////////////////////////////////////////////
+void FDbgHelpResolver::ResolveSymbol(uint64 Address, FResolvedSymbol* Target)
+{
+	check(Target);
+
+	uint8 InfoBuffer[sizeof(SYMBOL_INFO) + (MaxNameLen * sizeof(char) + 1)];
+	SYMBOL_INFO* Info = (SYMBOL_INFO*)InfoBuffer;
+	Info->SizeOfStruct = sizeof(SYMBOL_INFO);
+	Info->MaxNameLen = MaxNameLen;
+
+	// Find and build the symbol name
+	if (!SymFromAddr((HANDLE)Handle, Address, NULL, Info))
+	{
+		++SymbolsFailed;
+		UpdateResolvedSymbol(Target, ESymbolQueryResult::NotFound, GUnknownModuleTextDbgHelp, GUnknownModuleTextDbgHelp);
+		return;
+	}
+
+	TStringBuilder<256> SymbolName;
+	SymbolName << ANSI_TO_TCHAR(Info->Name);
+	const TCHAR* SymbolNameStr = Session.StoreString(FStringView(SymbolName));
+
+	// Find the source file and line
+	DWORD  dwDisplacement;
+	IMAGEHLP_LINE Line;
+	Line.SizeOfStruct = sizeof(IMAGEHLP_LINE);
+
+	if (!SymGetLineFromAddr((HANDLE)Handle, Address, &dwDisplacement, &Line))
+	{
+		++SymbolsFailed;
+		UpdateResolvedSymbol(Target, ESymbolQueryResult::NotFound, SymbolNameStr, GUnknownModuleTextDbgHelp);
+		return;
+	}
+
+	TStringBuilder<256> FileAndLine;
+	FileAndLine << ANSI_TO_TCHAR(Line.FileName) << TEXT(" (") << uint32(Line.LineNumber) << TEXT(")");
+	const TCHAR* FileAndLineStr = Session.StoreString(FStringView(FileAndLine));
+	
+	++SymbolsResolved;
+	UpdateResolvedSymbol(Target, ESymbolQueryResult::OK, SymbolNameStr, FileAndLineStr);
+}
+
+/////////////////////////////////////////////////////////////////////
+bool FDbgHelpResolver::LoadModuleSymbols(uint64 Base, uint64 Size, const TCHAR* Path)
+{
+	// Attempt to load symbols
+	const DWORD64 LoadedBaseAddress = SymLoadModuleEx((HANDLE)Handle, NULL, TCHAR_TO_ANSI(Path), NULL, Base, Size, NULL, 0);
+	const bool bSymbolsLoaded = Base == LoadedBaseAddress;
+
+	if (!bSymbolsLoaded)
+	{
+		UE_LOG(LogDbgHelp, Warning, TEXT("Unable to load symbols for %s at %p"), Path, Base);
+		++ModulesFailed;
+	}
+	else
+	{
+		UE_LOG(LogDbgHelp, Display, TEXT("Loaded symbols for %s at %p."), Path, Base);
+		++ModulesLoaded;
+	}
+
+	// Update the module entry with the result
+	FScopeLock _(&ModulesCs);
+	const int32 EntryIdx = Algo::BinarySearchBy(LoadedModules, Base, [](const FModuleEntry& Entry) { return Entry.Base; });
+	check(EntryIdx >= 0 && EntryIdx < LoadedModules.Num());
+	LoadedModules[EntryIdx].bSymbolsLoaded = bSymbolsLoaded;
+	
+	return bSymbolsLoaded;
+}
+
+
+/////////////////////////////////////////////////////////////////////
+const TCHAR* FDbgHelpResolver::GetModuleNameForAddress(uint64 Address) const
+{
+	const int32 EntryIdx = Algo::LowerBoundBy(LoadedModules, Address, [](const FModuleEntry& Entry) { return Entry.Base; }) - 1;
+	if (EntryIdx < 0 || EntryIdx >= LoadedModules.Num())
+	{
+		return GUnknownModuleTextDbgHelp;
+	}
+
+	const FModuleEntry& Entry = LoadedModules[EntryIdx];
+	if (Address > (Entry.Base + Entry.Size))
+	{
+		return GUnknownModuleTextDbgHelp;
+	}
+
+	return Entry.Name;
+}
+
+
+/////////////////////////////////////////////////////////////////////
+#include "Windows/HideWindowsPlatformTypes.h"
+
+/////////////////////////////////////////////////////////////////////
+
+} // namespace TraceServices
+
+#endif // PLATFORM_WINDOWS
diff --git a/Engine/Source/Developer/TraceServices/Private/Model/DbgHelpResolver.h b/Engine/Source/Developer/TraceServices/Private/Model/DbgHelpResolver.h
new file mode 100644
--- /dev/null
+++ b/Engine/Source/Developer/TraceServices/Private/Model/DbgHelpResolver.h
@@ -0,0 +1,89 @@
+// Copyright Epic Games, Inc. All Rights Reserved.
+#pragma once
+#if PLATFORM_WINDOWS
+
+#include "Containers/Queue.h"
+#include "HAL/CriticalSection.h"
+#include "Misc/PathViews.h"
+#include "Misc/StringBuilder.h"
+#include "TraceServices/Model/AnalysisSession.h"
+#include "TraceServices/Model/Modules.h"
+#include <atomic>
+
+/////////////////////////////////////////////////////////////////////
+namespace TraceServices {
+
+/////////////////////////////////////////////////////////////////////
+class FDbgHelpResolver : public FRunnable
+{
+public:
+	FDbgHelpResolver(IAnalysisSession& InSession);
+	~FDbgHelpResolver();
+	void QueueModuleLoad(const FStringView& ModulePath, uint64 Base, uint32 Size, const uint8* ImageId, uint32 ImageIdSize);
+	void QueueSymbolResolve(uint64 Address, FResolvedSymbol* Symbol);
+	void GetStats(IModuleProvider::FStats* OutStats) const;
+	void OnAnalysisComplete();
+
+private:	
+	struct FModuleEntry
+	{
+		uint64 Base;
+		uint32 Size;
+		const TCHAR* Name;
+		const TCHAR* Path;
+		bool bSymbolsLoaded;
+	};
+
+	struct FQueuedAddress
+	{
+		uint64 Address;
+		FResolvedSymbol* Target;
+	};
+
+	struct FQueuedModule
+	{
+		uint64 Base;
+		uint64 Size;
+		const TCHAR* ImagePath;
+	};
+
+	enum : uint32 {
+		MaxNameLen = 512,
+	};
+
+
+	bool SetupSyms();
+	void FreeSyms() const;
+	virtual uint32 Run() override;
+	virtual void Stop() override;
+	static void UpdateResolvedSymbol(FResolvedSymbol* Symbol, ESymbolQueryResult Result, const TCHAR* Name, const TCHAR* FileAndLine);
+
+	void ResolveSymbol(uint64 Address, FResolvedSymbol* Target);
+	bool LoadModuleSymbols(uint64 Base, uint64 Size, const TCHAR* Path);
+	const TCHAR* GetModuleNameForAddress(uint64 Address) const;
+
+	FCriticalSection ModulesCs;
+	TArray<FModuleEntry> LoadedModules;
+	TQueue<FQueuedModule, EQueueMode::Mpsc> LoadSymbolsQueue;
+	TQueue<FQueuedAddress, EQueueMode::Mpsc> ResolveQueue;
+
+	std::atomic<uint32> ModulesDiscovered;
+	std::atomic<uint32> ModulesFailed;
+	std::atomic<uint32> ModulesLoaded;
+	std::atomic<uint32> SymbolsDiscovered;
+	std::atomic<uint32> SymbolsFailed;
+	std::atomic<uint32> SymbolsResolved;
+
+	bool bRunWorkerThread;
+	bool bDrainThenStop;
+	bool bInitialized;
+	UPTRINT Handle;
+	IAnalysisSession& Session;
+	FRunnableThread* Thread;
+};
+
+/////////////////////////////////////////////////////////////////////
+
+} // namespace TraceServices
+
+#endif // PLATFORM_WINDOWS
diff --git a/Engine/Source/Developer/TraceServices/Private/Model/ModuleProvider.cpp b/Engine/Source/Developer/TraceServices/Private/Model/ModuleProvider.cpp
--- a/Engine/Source/Developer/TraceServices/Private/Model/ModuleProvider.cpp
+++ b/Engine/Source/Developer/TraceServices/Private/Model/ModuleProvider.cpp
@@ -4,12 +4,15 @@
 #include "Containers/Map.h"
 #include "TraceServices/Model/AnalysisSession.h"
 #include "TraceServices/Model/Diagnostics.h"
-#include "HAL/Runnable.h"
 #include "HAL/RunnableThread.h"
 #include <atomic>
 
+// When set use the RAD Syms library, otherwise fall back to DbgHelp (on Windows)
+#define USE_SYMS_LIB 1
+
 // Symbol files implementations
-#include "PbdSymbols.inl"
+#include "SymslibResolver.h"
+#include "DbgHelpResolver.h"
 
 namespace TraceServices {
 
@@ -49,7 +52,7 @@ public:
 	
 	
 	//Analysis interface
-	void						OnModuleLoad(const FStringView& Module, uint64 Base, uint32 Size) override;
+	void						OnModuleLoad(const FStringView& Module, uint64 Base, uint32 Size, const uint8* Checksum, uint32 ChecksumSize) override;
 	void 						OnModuleUnload(uint64 Base) override;
 	void						OnAnalysisComplete() override;
 
@@ -90,7 +93,7 @@ const FResolvedSymbol* TModuleProvider<SymbolProvider>::GetSymbol(uint64 Address
 		}
 	}
 
-	FResolvedSymbol* ResolvedSymbol = nullptr;
+	FResolvedSymbol* ResolvedSymbol;
 	{
 		// Add a pending entry to our cache
 		FWriteScopeLock _(SymbolsLock);
@@ -118,7 +121,7 @@ void TModuleProvider<SymbolResolver>::GetStats(FStats* OutStats) const
 
 /////////////////////////////////////////////////////////////////////
 template<typename SymbolProvider>
-void TModuleProvider<SymbolProvider>::OnModuleLoad(const FStringView& Module, uint64 Base, uint32 Size)
+void TModuleProvider<SymbolProvider>::OnModuleLoad(const FStringView& Module, uint64 Base, uint32 Size, const uint8* ImageId, uint32 ImageIdSize)
 {
 	if (Module.Len() == 0)
 	{
@@ -128,13 +131,12 @@ void TModuleProvider<SymbolProvider>::OnModuleLoad(const FStringView& Module, ui
 	const TCHAR* Path = Session.StoreString(Module);
 	const TCHAR* Name = Session.StoreString(Module);
 
-	FModuleEntry* NewEntry = nullptr;
 	{
 		FWriteScopeLock _(ModulesLock);
-		NewEntry = &Modules.EmplaceBack(Name, Path, Base, Size, false);
+		Modules.EmplaceBack(Name, Path, Base, Size, false);
 	}
 
-	Resolver.QueueModuleLoad(Module, Base, Size);
+	Resolver.QueueModuleLoad(Module, Base, Size, ImageId, ImageIdSize);
 }
 
 /////////////////////////////////////////////////////////////////////
@@ -154,13 +156,18 @@ void TModuleProvider<SymbolResolver>::OnAnalysisComplete()
 /////////////////////////////////////////////////////////////////////
 IModuleAnalysisProvider* CreateModuleProvider(IAnalysisSession& InSession, const FAnsiStringView& InSymbolFormat)
 {
+	IModuleAnalysisProvider* Provider(nullptr);
 #if PLATFORM_WINDOWS
-	if (InSymbolFormat.Equals("pdb"))
+	if (!Provider && InSymbolFormat.Equals("pdb"))
+	{
+		Provider = new TModuleProvider<FSymslibResolver>(InSession);
+	}
+	if (!Provider && InSymbolFormat.Equals("pdb"))
 	{
-		return new TModuleProvider<FPdbSymbols>(InSession);
+		Provider = new TModuleProvider<FDbgHelpResolver>(InSession);
 	}
-#endif
-	return nullptr;
+#endif //PLATFORM_WINDOWS
+	return Provider;
 }
 
 } // namespace TraceServices
diff --git a/Engine/Source/Developer/TraceServices/Private/Model/ModuleProvider.h b/Engine/Source/Developer/TraceServices/Private/Model/ModuleProvider.h
--- a/Engine/Source/Developer/TraceServices/Private/Model/ModuleProvider.h
+++ b/Engine/Source/Developer/TraceServices/Private/Model/ModuleProvider.h
@@ -14,7 +14,7 @@ class IAnalysisSession;
 class IModuleAnalysisProvider : public IModuleProvider
 {
 public:
-	virtual void	OnModuleLoad(const FStringView& Module, uint64 Base, uint32 Size) = 0;
+	virtual void	OnModuleLoad(const FStringView& Module, uint64 Base, uint32 Size, const uint8* ImageId, uint32 ImageIdSize) = 0;
 	virtual void	OnModuleUnload(uint64 Base) = 0;
 	virtual void	OnAnalysisComplete() = 0;
 };
diff --git a/Engine/Source/Developer/TraceServices/Private/Model/PbdSymbols.inl b/Engine/Source/Developer/TraceServices/Private/Model/PbdSymbols.inl
deleted file mode 100644
diff --git a/Engine/Source/Developer/TraceServices/Private/Model/SymslibResolver.cpp b/Engine/Source/Developer/TraceServices/Private/Model/SymslibResolver.cpp
new file mode 100644
--- /dev/null
+++ b/Engine/Source/Developer/TraceServices/Private/Model/SymslibResolver.cpp
@@ -0,0 +1,379 @@
+// Copyright Epic Games, Inc. All Rights Reserved.
+
+#if PLATFORM_WINDOWS
+#include "SymslibResolver.h"
+#include "Algo/Sort.h"
+#include "Async/MappedFileHandle.h"
+#include "Async/ParallelFor.h"
+#include "Containers/StringView.h"
+#include "HAL/PlatformFileManager.h"
+#include "Logging/LogMacros.h"
+#include "Misc/PathViews.h"
+#include "Misc/ScopeLock.h"
+#include "Misc/ScopeRWLock.h"
+#include "Misc/StringBuilder.h"
+#include "TraceServices/Model/AnalysisSession.h"
+#include <atomic>
+
+#include "symslib.h"
+
+/////////////////////////////////////////////////////////////////////
+DEFINE_LOG_CATEGORY_STATIC(LogSymslib, Log, All);
+
+/////////////////////////////////////////////////////////////////////
+namespace TraceServices {
+
+/////////////////////////////////////////////////////////////////////
+
+static const TCHAR* GUnknownModuleTextSymsLib = TEXT("Unknown");
+
+/////////////////////////////////////////////////////////////////////
+FSymslibResolver::FSymslibResolver(IAnalysisSession& InSession)
+	: Modules(InSession.GetLinearAllocator(), 128)
+	, Session(InSession)
+{
+}
+
+FSymslibResolver::~FSymslibResolver()
+{
+}
+
+void FSymslibResolver::QueueModuleLoad(const FStringView& ModulePath, uint64 Base, uint32 Size, const uint8* ImageId, uint32 ImageIdSize)
+{
+	FWriteScopeLock _(ModulesLock);
+	const FStringView ModuleName = FPathViews::GetCleanFilename(ModulePath);
+	
+	// Add module
+	FModuleEntry* Entry = &Modules.PushBack();
+	Entry->Base = Base;
+	Entry->Size = Size;
+	Entry->Name = Session.StoreString(ModuleName); //Allows safe sharing to UI
+	Entry->Path = Session.StoreString(ModulePath);
+	Entry->Status = EModuleStatus::Pending;
+	Entry->ImageId = TArrayView<const uint8>(ImageId, ImageIdSize);
+
+	++ModulesDiscovered;
+	
+	// Queue up module to have symbols loaded
+	++TasksInFlight;
+	FFunctionGraphTask::CreateAndDispatchWhenReady([this, Entry] { LoadModuleTracked(Entry); --TasksInFlight;});
+
+	// Sort list according to base address
+	SortedModules.Add(Entry);
+	Algo::Sort(SortedModules, [](const FModuleEntry* Lhs, const FModuleEntry* Rhs) { return Lhs->Base < Rhs->Base; });
+}
+
+void FSymslibResolver::QueueSymbolResolve(uint64 Address, FResolvedSymbol* Symbol)
+{
+	FScopeLock _(&SymbolsQueueLock);
+	MaybeDispatchQueuedAddresses();
+	++SymbolsDiscovered;
+	ResolveQueue.Add(FQueuedAddress{Address, Symbol});
+}
+
+void FSymslibResolver::OnAnalysisComplete()
+{
+	// Dispatch any remaining requests
+	{
+		FScopeLock _(&SymbolsQueueLock);
+		DispatchQueuedAddresses();
+		ResolveQueue.Reset();
+	}
+
+	// Wait for outstanding batches to complete
+	uint32 OutstandingTasks = TasksInFlight.load(std::memory_order_acquire);
+	UE_LOG(LogSymslib, Display, TEXT("Waiting for %d outstanding tasks..."), OutstandingTasks);
+	do
+	{
+		OutstandingTasks = TasksInFlight.load(std::memory_order_acquire);
+		FPlatformProcess::Sleep(0.0);
+	} while(OutstandingTasks > 0);
+
+	// Release file handles. No tasks can be in running at this point.
+	for (auto File : FileHandlesAndRegions)
+	{
+		delete File.Get<1>();
+		delete File.Get<0>();
+	}
+}
+
+void FSymslibResolver::GetStats(IModuleProvider::FStats* OutStats) const
+{
+	OutStats->ModulesDiscovered = ModulesDiscovered.load();
+	OutStats->ModulesFailed = ModulesFailed.load();
+	OutStats->ModulesLoaded = ModulesLoaded.load();
+	OutStats->SymbolsDiscovered = SymbolsDiscovered.load();
+	OutStats->SymbolsFailed = SymbolsFailed.load();
+	OutStats->SymbolsResolved = SymbolsResolved.load();
+}
+
+
+/**
+ * Checks if there are no modules in flight and that the queue has reached 
+ * the threshold for dispatching. Note that is up to the caller to synchronize.
+ */
+void FSymslibResolver::MaybeDispatchQueuedAddresses()
+{
+	const bool bModulesInFlight = (ModulesDiscovered.load() - ModulesFailed.load() - ModulesLoaded.load()) > 0;
+	if (!bModulesInFlight && (ResolveQueue.Num() >= QueuedAddressLength))
+	{
+		DispatchQueuedAddresses();
+		ResolveQueue.Reset();
+	}
+}
+
+/**
+ * Dispatches the currently queued addresses to be resolved. Note that is up to the caller to synchronize.
+ */
+void FSymslibResolver::DispatchQueuedAddresses()
+{
+	//todo: This creates two copies of resolve queue!
+	TArray<FQueuedAddress> WorkingSet(ResolveQueue);
+	++TasksInFlight;
+	FFunctionGraphTask::CreateAndDispatchWhenReady([this, WorkingSet] {
+		ParallelFor(WorkingSet.Num(), [this, &WorkingSet](uint32 Index) {
+			const FQueuedAddress& ToResolve = WorkingSet[Index];
+			ResolveSymbolTracked(ToResolve.Address, ToResolve.Target);
+		});
+		--TasksInFlight;
+	});
+}
+
+FSymslibResolver::FModuleEntry* FSymslibResolver::GetModuleForAddress(uint64 Address)
+{
+	FReadScopeLock _(ModulesLock);
+	const int32 EntryIdx = Algo::LowerBoundBy(SortedModules, Address, [](const FModuleEntry* Entry) { return Entry->Base; }) - 1;
+	if (EntryIdx < 0 || EntryIdx >= SortedModules.Num())
+	{
+		return nullptr;
+	}
+	return SortedModules[EntryIdx];
+}
+
+void FSymslibResolver::UpdateResolvedSymbol(FResolvedSymbol* Symbol, ESymbolQueryResult Result, const TCHAR* Name, const TCHAR* FileAndLine)
+{
+	Symbol->Name = Name;
+	Symbol->FileAndLine = FileAndLine;
+	Symbol->Result.store(Result, std::memory_order_release);
+}
+
+void FSymslibResolver::LoadModuleTracked(FModuleEntry* Module)
+{
+	const EModuleStatus Status = LoadModule(Module);
+	if (Status == EModuleStatus::Loaded)
+	{
+		++ModulesLoaded;
+	}
+	else
+	{
+		++ModulesFailed;
+		Module->Status.store(Status);
+	}
+
+	//  Make the status visible
+	Module->Status.store(Status);
+
+	// If this is the last module in flight dispatch pending queries
+	{
+		FScopeLock _(&SymbolsQueueLock);
+		MaybeDispatchQueuedAddresses();
+	}
+}
+
+FSymslibResolver::EModuleStatus FSymslibResolver::LoadModule(FModuleEntry* Module)
+{
+	IPlatformFile* PlatformFile = &FPlatformFileManager::Get().GetPlatformFile();
+
+	// Initialize the instance
+	Module->Instance = syms_init_ex(Module->Base);
+	SymsInstance* ModuleInstance = Module->Instance;
+	
+	if (!ModuleInstance)
+	{
+		UE_LOG(LogSymslib, Error, TEXT("Failed to initialize '%s'"), Module->Name);
+		return EModuleStatus::Failed;
+	}
+
+	if (!PlatformFile->FileExists(Module->Path))
+	{
+		UE_LOG(LogSymslib, Error, TEXT("File '%s' does not exist"), Module->Path);
+	}
+
+	// Map the image file to memory
+	IMappedFileHandle* ImageFileHandle = PlatformFile->OpenMapped(Module->Path);
+	if (ImageFileHandle == nullptr)
+	{
+		return EModuleStatus::Failed;
+	}
+	IMappedFileRegion* ImageFileRegion = ImageFileHandle->MapRegion(0, ImageFileHandle->GetFileSize());
+
+	// Track all open handles
+	TrackFileHandlesAndRegions(ImageFileHandle , ImageFileRegion);
+
+	// Load the image
+	SymsErrorCode Result = syms_load_image(ModuleInstance, (void*)ImageFileRegion->GetMappedPtr(), ImageFileRegion->GetMappedSize(), 0);
+
+	if (SYMS_RESULT_FAIL(Result))
+	{
+		UE_LOG(LogSymslib, Warning, TEXT("Failed to load '%s'"), Module->Path);
+		return EModuleStatus::Failed;
+	}
+
+	// Iterate and map all the debug files
+	const FStringView ModuleBasePath = FPathViews::GetPath(Module->Path);
+	TStringBuilder<256> DebugFilePath;
+	TArray<SymsFile> Files;
+
+	SymsDebugFileIter DebugFiles;
+	if (syms_debug_file_iter_init(&DebugFiles, ModuleInstance))
+	{
+		SymsString Path;
+		while (syms_debug_file_iter_next(&DebugFiles, &Path))
+		{
+			// Map the debug file
+			DebugFilePath.Reset();
+			FPathViews::Append(DebugFilePath, ModuleBasePath);
+			FPathViews::Append(DebugFilePath, ANSI_TO_TCHAR(Path.data));
+			
+			if (PlatformFile->FileExists(*DebugFilePath))
+			{
+				IMappedFileHandle* FileHandle = PlatformFile->OpenMapped(*DebugFilePath); 
+				IMappedFileRegion* FileRegion = FileHandle->MapRegion(0, FileHandle->GetFileSize());
+
+				Files.Push(SymsFile{
+					Path.data,
+					(void*)FileRegion->GetMappedPtr(),
+					(uint64)FileRegion->GetMappedSize()
+				});
+
+				TrackFileHandlesAndRegions(FileHandle, FileRegion);
+			}
+		}
+	}
+	
+	// Prepare to load the debug files
+	Result = syms_load_debug_info_ex(ModuleInstance, Files.GetData(), Files.Num(), SYMS_LOAD_DEBUG_INFO_FLAGS_DEFER_BUILD_MODULE);
+
+	if (SYMS_RESULT_FAIL(Result))
+	{
+		UE_LOG(LogSymslib, Warning, TEXT("Could not read debug files for '%s'"), Module->Path);
+		return EModuleStatus::Failed;
+	}
+
+	// Split up the work of actually loading debug info
+	const uint32 Count = syms_get_module_build_count(ModuleInstance);
+	// Each work item needs some memory
+	TArray<SymsArena*> Arenas;
+	for (uint32 i = 0; i < Count; ++i) { Arenas.Push(syms_borrow_memory(ModuleInstance)); }
+	
+	ParallelFor(Count, [ModuleInstance, Arenas](uint32 Index) {
+		syms_build_module(ModuleInstance, (SymsModID)Index, Arenas[Index]); 
+	});
+
+	// Check age of debug data
+	if (!Module->ImageId.IsEmpty() && ModuleInstance->img.type == SYMS_IMAGE_NT)
+	{
+		// For Pdbs checksum is a 16 byte guid and 8 byte unsigned integer for age
+		static_assert(sizeof(FGuid) == 16, "Expected 16 byte FGuid");
+		check(Module->ImageId.Num() == 20);
+		FGuid* ModuleGuid = (FGuid*) Module->ImageId.GetData();
+		uint32* ModuleAge = ((uint32*) Module->ImageId.GetData()) + 4;
+		SymsDebugInfoPdb* PdbInfo = (SymsDebugInfoPdb*)ModuleInstance->debug_info.impl;
+		FGuid* PdbGuid = (FGuid*)&PdbInfo->context.auth_guid;
+		const uint32 PdbAge = PdbInfo->context.age; 
+		if (*ModuleGuid != *PdbGuid || *ModuleAge != PdbAge)
+		{
+			UE_LOG(LogSymslib, Warning, TEXT("Symbols for '%s' does not match traced binary."), Module->Name);
+			return EModuleStatus::VersionMismatch;
+		}
+	}
+
+	const uint32 SymbolCount = syms_get_proc_count(ModuleInstance); // get symbol count
+	const uint32 LineCount = syms_get_line_count(ModuleInstance); // get line count
+	UE_LOG(LogSymslib, Display, TEXT("Loaded symbols for '%s', %u symbols and %u lines."), Module->Name, SymbolCount, LineCount);
+
+	return EModuleStatus::Loaded;
+}
+
+void FSymslibResolver::ResolveSymbolTracked(uint64 Address, FResolvedSymbol* Target)
+{
+	if (!ResolveSymbol(Address, Target))
+	{
+		++SymbolsFailed;
+	}
+	else
+	{
+		++SymbolsResolved;
+	}
+}
+
+bool FSymslibResolver::ResolveSymbol(uint64 Address, FResolvedSymbol* Target)
+{
+	FModuleEntry* Module = GetModuleForAddress(Address);
+	if (!Module)
+	{
+		UpdateResolvedSymbol(Target, ESymbolQueryResult::NotFound, GUnknownModuleTextSymsLib, GUnknownModuleTextSymsLib);
+		return false;
+	}
+
+	EModuleStatus Status = Module->Status.load();
+	while (Status == EModuleStatus::Pending)
+	{
+		//todo: yield
+		Status = Module->Status.load();
+	}
+
+	switch (Status)
+	{
+	case EModuleStatus::Failed:
+		UpdateResolvedSymbol(Target, ESymbolQueryResult::NotLoaded, Module->Name, GUnknownModuleTextSymsLib);
+		return false;
+	case EModuleStatus::VersionMismatch:
+		UpdateResolvedSymbol(Target, ESymbolQueryResult::Mismatch, Module->Name, GUnknownModuleTextSymsLib);
+		return false;
+	default:
+		break;
+	}
+
+	SymsProc Proc;
+	SymsSourceFileMap File;
+
+	// Find procedure and source file for address
+	if (!syms_proc_from_va(Module->Instance, Address, &Proc) || !syms_va_to_src(Module->Instance, Address, &File))
+	{
+		UpdateResolvedSymbol(Target, ESymbolQueryResult::NotFound, Module->Name, GUnknownModuleTextSymsLib); 
+		return false;
+	}
+
+	// Read the symbol name
+	ANSICHAR ProcedureName[1024];
+	syms_read_strref(Module->Instance, &Proc.name_ref, ProcedureName, sizeof(ProcedureName));
+
+	// Read the source location
+	ANSICHAR SourceFile[1024];
+	syms_read_strref(Module->Instance, &File.file.name, SourceFile, sizeof(SourceFile));
+
+	TStringBuilder<1024> SourceLocation;
+	SourceLocation << ANSI_TO_TCHAR(SourceFile) << TEXT("(") << File.line.ln << TEXT(")");
+
+	// Store the strings and update the target data
+	UpdateResolvedSymbol(
+		Target, 
+		ESymbolQueryResult::OK,
+		Session.StoreString(ANSI_TO_TCHAR(ProcedureName)), 
+		Session.StoreString(SourceLocation)
+	);
+
+	return true;
+}
+
+void FSymslibResolver::TrackFileHandlesAndRegions(IMappedFileHandle* FileHandle, IMappedFileRegion* FileRegion)
+{
+	FileHandlesAndRegions.Add(TTuple<IMappedFileHandle*, IMappedFileRegion*>(FileHandle, FileRegion));
+}
+
+/////////////////////////////////////////////////////////////////////
+
+} // namespace TraceServices
+
+#endif // PLATFORM_WINDOWS
diff --git a/Engine/Source/Developer/TraceServices/Private/Model/SymslibResolver.h b/Engine/Source/Developer/TraceServices/Private/Model/SymslibResolver.h
new file mode 100644
--- /dev/null
+++ b/Engine/Source/Developer/TraceServices/Private/Model/SymslibResolver.h
@@ -0,0 +1,106 @@
+// Copyright Epic Games, Inc. All Rights Reserved.
+#pragma once
+#if PLATFORM_WINDOWS
+
+#include "Async/MappedFileHandle.h"
+#include "Common/PagedArray.h"
+#include "HAL/CriticalSection.h"
+#include "TraceServices/Model/AnalysisSession.h"
+#include "TraceServices/Model/Modules.h"
+#include <atomic>
+
+#include "symslib.h"
+
+/////////////////////////////////////////////////////////////////////
+namespace TraceServices {
+
+/////////////////////////////////////////////////////////////////////
+class FSymslibResolver
+{
+public:
+	FSymslibResolver(IAnalysisSession& InSession);
+	~FSymslibResolver();
+	void QueueModuleLoad(const FStringView& ModulePath, uint64 Base, uint32 Size, const uint8* ImageId, uint32 ImageIdSize);
+	void QueueSymbolResolve(uint64 Address, FResolvedSymbol* Symbol);
+	void OnAnalysisComplete();
+	void GetStats(IModuleProvider::FStats* OutStats) const;
+
+private:
+
+	enum class EModuleStatus : uint8
+	{
+		Pending,
+		Loaded,
+		VersionMismatch,
+		Failed
+	};
+
+	struct FModuleEntry
+	{
+		uint64 Base;
+		uint32 Size;
+		const TCHAR* Name;
+		const TCHAR* Path;
+		SymsInstance* Instance;
+		std::atomic<EModuleStatus> Status;
+		TArray<uint8> ImageId;
+	};
+
+	struct FQueuedAddress
+	{
+		uint64 Address;
+		FResolvedSymbol* Target;
+	};
+
+	struct FQueuedModule
+	{
+		uint64 Base;
+		uint64 Size;
+		const TCHAR* ImagePath;
+	};
+
+	enum : uint32 {
+		MaxNameLen = 512,
+		QueuedAddressLength = 256
+	};
+
+	/**
+	 * Checks if there are no modules in flight and that the queue has reached 
+	 * the threshold for dispatching. Note that is up to the caller to synchronize.
+	 */
+	void MaybeDispatchQueuedAddresses();
+	/**
+	 * Dispatches the currently queued addresses to be resolved. Note that is up to the caller to synchronize.
+	 */
+	void DispatchQueuedAddresses();
+	FModuleEntry* GetModuleForAddress(uint64 Address);
+	static void UpdateResolvedSymbol(FResolvedSymbol* Symbol, ESymbolQueryResult Result, const TCHAR* Name, const TCHAR* FileAndLine);
+	void LoadModuleTracked(FModuleEntry* Module);
+	EModuleStatus LoadModule(FModuleEntry* Module);
+	void ResolveSymbolTracked(uint64 Address, FResolvedSymbol* Target);
+	bool ResolveSymbol(uint64 Address, FResolvedSymbol* Target);
+	void TrackFileHandlesAndRegions(IMappedFileHandle* FileHandle, IMappedFileRegion* FileRegion);
+	
+	FRWLock ModulesLock;
+	TPagedArray<FModuleEntry> Modules;
+	TArray<FModuleEntry*> SortedModules;
+	FCriticalSection SymbolsQueueLock;
+	TArray<FQueuedAddress, TInlineAllocator<QueuedAddressLength>> ResolveQueue;
+	std::atomic<uint32> TasksInFlight;
+	TArray<TTuple<IMappedFileHandle*, IMappedFileRegion*>> FileHandlesAndRegions;
+	
+	std::atomic<uint32> ModulesDiscovered;
+	std::atomic<uint32> ModulesFailed;
+	std::atomic<uint32> ModulesLoaded;
+	std::atomic<uint32> SymbolsDiscovered;
+	std::atomic<uint32> SymbolsFailed;
+	std::atomic<uint32> SymbolsResolved;
+
+	IAnalysisSession& Session;
+};
+
+/////////////////////////////////////////////////////////////////////
+
+} // namespace TraceServices
+
+#endif // PLATFORM_WINDOWS
diff --git a/Engine/Source/Developer/TraceServices/TraceServices.Build.cs b/Engine/Source/Developer/TraceServices/TraceServices.Build.cs
--- a/Engine/Source/Developer/TraceServices/TraceServices.Build.cs
+++ b/Engine/Source/Developer/TraceServices/TraceServices.Build.cs
@@ -10,7 +10,7 @@ public class TraceServices : ModuleRules
 			new string[] {
 				"Core",
 				"TraceAnalysis",
+				"SymsLib"
 			});
-
 	}
 }
