Clean up code after suggestions.

#rnx

diff --git a/Engine/Source/Developer/TraceServices/Private/Analyzers/ModuleAnalysis.cpp b/Engine/Source/Developer/TraceServices/Private/Analyzers/ModuleAnalysis.cpp
--- a/Engine/Source/Developer/TraceServices/Private/Analyzers/ModuleAnalysis.cpp
+++ b/Engine/Source/Developer/TraceServices/Private/Analyzers/ModuleAnalysis.cpp
@@ -3,7 +3,6 @@
 #include "ModuleAnalysis.h"
 #include "TraceServices/Model/AnalysisSession.h"
 #include "Model/ModuleProvider.h"
-#include "UObject/NameTypes.h"
 
 namespace TraceServices {
 
@@ -19,6 +18,7 @@ enum Routes
 FModuleAnalyzer::FModuleAnalyzer(IAnalysisSession& InSession)
 	: Session(InSession)
 	, Provider(nullptr)
+	, ModuleBaseShift(0)
 {
 }
 
@@ -71,13 +71,9 @@ bool FModuleAnalyzer::OnEvent(uint16 RouteId, EStyle Style, const FOnEventContex
 				{
 					const uint64 Base = GetRealBaseAddress(Context.EventData.GetValue<uint32>("Base"));
 					const uint32 Size = Context.EventData.GetValue<uint32>("Size");
-					const uint32 Age = Context.EventData.GetValue<uint32>("Age", 0);
-					const IAnalyzer::TArrayReader<uint8>& ImageId = Context.EventData.GetArray<uint8>("ImageId");
+					const TArrayReader<uint8>& ImageId = Context.EventData.GetArray<uint8>("ImageId");
 
-					if (Provider)
-					{
-						Provider->OnModuleLoad(ModuleName, Base, Size, ImageId.GetData(), ImageId.Num());
-					}
+					Provider->OnModuleLoad(ModuleName, Base, Size, ImageId.GetData(), ImageId.Num());
 				}
 			}
 			break;
@@ -86,10 +82,7 @@ bool FModuleAnalyzer::OnEvent(uint16 RouteId, EStyle Style, const FOnEventContex
 			if (Provider != nullptr)
 			{
 				const uint32 Base = Context.EventData.GetValue<uint32>("Base");
-				if (Provider)
-				{
-					Provider->OnModuleUnload(Base);
-				}
+				Provider->OnModuleUnload(Base);
 			}
 			break;
 	}
@@ -97,7 +90,7 @@ bool FModuleAnalyzer::OnEvent(uint16 RouteId, EStyle Style, const FOnEventContex
 }
 
 ////////////////////////////////////////////////////////////////////////////////
-uint64 FModuleAnalyzer::GetRealBaseAddress(uint32 EventBase)
+uint64 FModuleAnalyzer::GetRealBaseAddress(uint32 EventBase) const
 {
 	// Unshift base address according to capture alignment (e.g. Windows 4k aligned)
 	return uint64(EventBase) << ModuleBaseShift;
diff --git a/Engine/Source/Developer/TraceServices/Private/Analyzers/ModuleAnalysis.h b/Engine/Source/Developer/TraceServices/Private/Analyzers/ModuleAnalysis.h
--- a/Engine/Source/Developer/TraceServices/Private/Analyzers/ModuleAnalysis.h
+++ b/Engine/Source/Developer/TraceServices/Private/Analyzers/ModuleAnalysis.h
@@ -19,7 +19,7 @@ public:
 	virtual bool 	OnEvent(uint16 RouteId, EStyle Style, const FOnEventContext& Context) override;
 
 private:
-	uint64			GetRealBaseAddress(uint32 EventBase);
+	uint64			GetRealBaseAddress(uint32 EventBase) const;
 
 	IAnalysisSession& Session;
 	IModuleAnalysisProvider* Provider;
