No-unity compile fix.

#rnx

diff --git a/Engine/Source/Developer/TraceServices/Private/Model/DbgHelpResolver.h b/Engine/Source/Developer/TraceServices/Private/Model/DbgHelpResolver.h
--- a/Engine/Source/Developer/TraceServices/Private/Model/DbgHelpResolver.h
+++ b/Engine/Source/Developer/TraceServices/Private/Model/DbgHelpResolver.h
@@ -2,8 +2,10 @@
 #pragma once
 #if PLATFORM_WINDOWS
 
+#include "Containers/Array.h"
 #include "Containers/Queue.h"
 #include "HAL/CriticalSection.h"
+#include "HAL/Runnable.h"
 #include "Misc/PathViews.h"
 #include "Misc/StringBuilder.h"
 #include "TraceServices/Model/AnalysisSession.h"
