Fix Non Unity Errors

diff --git a/Engine/Source/Developer/TraceServices/Private/Analyzers/ModuleAnalysis.cpp b/Engine/Source/Developer/TraceServices/Private/Analyzers/ModuleAnalysis.cpp
--- a/Engine/Source/Developer/TraceServices/Private/Analyzers/ModuleAnalysis.cpp
+++ b/Engine/Source/Developer/TraceServices/Private/Analyzers/ModuleAnalysis.cpp
@@ -3,6 +3,7 @@
 #include "ModuleAnalysis.h"
 #include "TraceServices/Model/AnalysisSession.h"
 #include "Model/ModuleProvider.h"
+#include "UObject/NameTypes.h"
 
 namespace TraceServices {
 
diff --git a/Engine/Source/Developer/TraceServices/Private/Model/SymslibResolver.cpp b/Engine/Source/Developer/TraceServices/Private/Model/SymslibResolver.cpp
--- a/Engine/Source/Developer/TraceServices/Private/Model/SymslibResolver.cpp
+++ b/Engine/Source/Developer/TraceServices/Private/Model/SymslibResolver.cpp
@@ -6,6 +6,7 @@
 #include "Async/MappedFileHandle.h"
 #include "Async/ParallelFor.h"
 #include "Containers/StringView.h"
+#include "GenericPlatform/GenericPlatformFile.h"
 #include "HAL/PlatformFileManager.h"
 #include "Logging/LogMacros.h"
 #include "Misc/PathViews.h"
