[Insights]
  - Fixed/added the compile time macros to enable Symslib / DbgHelp resolvers.
  - Fixed crash in Symslib resolver when analysis completes (thread safety for array with mapped file regions and handles).
  - Attempt to fix a crash in Symslib resolver when loading a module (thread safety).
  - Minor code style / coding standards fixes.

diff --git a/Engine/Source/Developer/TraceServices/Private/Model/ModuleProvider.cpp b/Engine/Source/Developer/TraceServices/Private/Model/ModuleProvider.cpp
--- a/Engine/Source/Developer/TraceServices/Private/Model/ModuleProvider.cpp
+++ b/Engine/Source/Developer/TraceServices/Private/Model/ModuleProvider.cpp
@@ -7,12 +7,18 @@
 #include "HAL/RunnableThread.h"
 #include <atomic>
 
-// When set use the RAD Syms library, otherwise fall back to DbgHelp (on Windows)
-#define USE_SYMS_LIB 1
+// Choose which symbol resolver to use.
+// If both are enabled and the RAD Syms library fails to initialize, it will fall back to DbgHelp (on Windows).
+#define USE_SYMSLIB 1
+#define USE_DBGHELP 1
 
 // Symbol files implementations
+#if USE_SYMSLIB
 #include "SymslibResolver.h"
+#endif
+#if USE_DBGHELP
 #include "DbgHelpResolver.h"
+#endif
 
 namespace TraceServices {
 
@@ -49,8 +55,7 @@ public:
 	//Query interface
 	const FResolvedSymbol*		GetSymbol(uint64 Address) override;
 	void						GetStats(FStats* OutStats) const override;
-	
-	
+
 	//Analysis interface
 	void						OnModuleLoad(const FStringView& Module, uint64 Base, uint32 Size, const uint8* Checksum, uint32 ChecksumSize) override;
 	void 						OnModuleUnload(uint64 Base) override;
@@ -158,16 +163,23 @@ IModuleAnalysisProvider* CreateModuleProvider(IAnalysisSession& InSession, const
 {
 	IModuleAnalysisProvider* Provider(nullptr);
 #if PLATFORM_WINDOWS
+#if USE_SYMSLIB
 	if (!Provider && InSymbolFormat.Equals("pdb"))
 	{
 		Provider = new TModuleProvider<FSymslibResolver>(InSession);
 	}
+#endif // USE_SYMSLIB
+#if USE_DBGHELP
 	if (!Provider && InSymbolFormat.Equals("pdb"))
 	{
 		Provider = new TModuleProvider<FDbgHelpResolver>(InSession);
 	}
+#endif // USE_DBGHELP
 #endif //PLATFORM_WINDOWS
 	return Provider;
 }
 
 } // namespace TraceServices
+
+#undef USE_SYMSLIB
+#undef USE_DBGHELP
diff --git a/Engine/Source/Developer/TraceServices/Private/Model/SymslibResolver.cpp b/Engine/Source/Developer/TraceServices/Private/Model/SymslibResolver.cpp
--- a/Engine/Source/Developer/TraceServices/Private/Model/SymslibResolver.cpp
+++ b/Engine/Source/Developer/TraceServices/Private/Model/SymslibResolver.cpp
@@ -1,6 +1,7 @@
 // Copyright Epic Games, Inc. All Rights Reserved.
 
 #if PLATFORM_WINDOWS
+
 #include "SymslibResolver.h"
 #include "Algo/Sort.h"
 #include "Async/MappedFileHandle.h"
@@ -43,7 +44,7 @@ void FSymslibResolver::QueueModuleLoad(const FStringView& ModulePath, uint64 Bas
 {
 	FWriteScopeLock _(ModulesLock);
 	const FStringView ModuleName = FPathViews::GetCleanFilename(ModulePath);
-	
+
 	// Add module
 	FModuleEntry* Entry = &Modules.PushBack();
 	Entry->Base = Base;
@@ -54,7 +55,7 @@ void FSymslibResolver::QueueModuleLoad(const FStringView& ModulePath, uint64 Bas
 	Entry->ImageId = TArrayView<const uint8>(ImageId, ImageIdSize);
 
 	++ModulesDiscovered;
-	
+
 	// Queue up module to have symbols loaded
 	++TasksInFlight;
 	FFunctionGraphTask::CreateAndDispatchWhenReady([this, Entry] { LoadModuleTracked(Entry); --TasksInFlight;});
@@ -74,27 +75,31 @@ void FSymslibResolver::QueueSymbolResolve(uint64 Address, FResolvedSymbol* Symbo
 
 void FSymslibResolver::OnAnalysisComplete()
 {
-	// Dispatch any remaining requests
+	// Dispatch any remaining requests.
 	{
 		FScopeLock _(&SymbolsQueueLock);
 		DispatchQueuedAddresses();
 		ResolveQueue.Reset();
 	}
 
-	// Wait for outstanding batches to complete
+	// Wait for outstanding batches to complete.
 	uint32 OutstandingTasks = TasksInFlight.load(std::memory_order_acquire);
 	UE_LOG(LogSymslib, Display, TEXT("Waiting for %d outstanding tasks..."), OutstandingTasks);
 	do
 	{
 		OutstandingTasks = TasksInFlight.load(std::memory_order_acquire);
 		FPlatformProcess::Sleep(0.0);
-	} while(OutstandingTasks > 0);
+	} while (OutstandingTasks > 0);
 
-	// Release file handles. No tasks can be in running at this point.
-	for (auto File : FileHandlesAndRegions)
+	// No tasks can be in running at this point.
 	{
-		delete File.Get<1>();
-		delete File.Get<0>();
+		FReadScopeLock _(FileHandlesAndRegionsLock);
+		for (FMappedFileAndRegion& File : FileHandlesAndRegions)
+		{
+			// Release file region before releasing the file handle.
+			delete File.Region;
+			delete File.Handle;
+		}
 	}
 }
 
@@ -108,9 +113,8 @@ void FSymslibResolver::GetStats(IModuleProvider::FStats* OutStats) const
 	OutStats->SymbolsResolved = SymbolsResolved.load();
 }
 
-
 /**
- * Checks if there are no modules in flight and that the queue has reached 
+ * Checks if there are no modules in flight and that the queue has reached
  * the threshold for dispatching. Note that is up to the caller to synchronize.
  */
 void FSymslibResolver::MaybeDispatchQueuedAddresses()
@@ -188,7 +192,7 @@ FSymslibResolver::EModuleStatus FSymslibResolver::LoadModule(FModuleEntry* Modul
 	// Initialize the instance
 	Module->Instance = syms_init_ex(Module->Base);
 	SymsInstance* ModuleInstance = Module->Instance;
-	
+
 	if (!ModuleInstance)
 	{
 		UE_LOG(LogSymslib, Error, TEXT("Failed to initialize '%s'"), Module->Name);
@@ -198,18 +202,21 @@ FSymslibResolver::EModuleStatus FSymslibResolver::LoadModule(FModuleEntry* Modul
 	if (!PlatformFile->FileExists(Module->Path))
 	{
 		UE_LOG(LogSymslib, Error, TEXT("File '%s' does not exist"), Module->Path);
+		return EModuleStatus::Failed;
 	}
 
 	// Map the image file to memory
 	IMappedFileHandle* ImageFileHandle = PlatformFile->OpenMapped(Module->Path);
 	if (ImageFileHandle == nullptr)
 	{
+		UE_LOG(LogSymslib, Error, TEXT("Failed to open '%s'"), Module->Path);
 		return EModuleStatus::Failed;
 	}
 	IMappedFileRegion* ImageFileRegion = ImageFileHandle->MapRegion(0, ImageFileHandle->GetFileSize());
+	check(ImageFileRegion != nullptr);
 
 	// Track all open handles
-	TrackFileHandlesAndRegions(ImageFileHandle , ImageFileRegion);
+	TrackFileHandlesAndRegions(ImageFileHandle, ImageFileRegion);
 
 	// Load the image
 	SymsErrorCode Result = syms_load_image(ModuleInstance, (void*)ImageFileRegion->GetMappedPtr(), ImageFileRegion->GetMappedSize(), 0);
@@ -235,23 +242,31 @@ FSymslibResolver::EModuleStatus FSymslibResolver::LoadModule(FModuleEntry* Modul
 			DebugFilePath.Reset();
 			FPathViews::Append(DebugFilePath, ModuleBasePath);
 			FPathViews::Append(DebugFilePath, ANSI_TO_TCHAR(Path.data));
-			
+
 			if (PlatformFile->FileExists(*DebugFilePath))
 			{
-				IMappedFileHandle* FileHandle = PlatformFile->OpenMapped(*DebugFilePath); 
-				IMappedFileRegion* FileRegion = FileHandle->MapRegion(0, FileHandle->GetFileSize());
-
-				Files.Push(SymsFile{
-					Path.data,
-					(void*)FileRegion->GetMappedPtr(),
-					(uint64)FileRegion->GetMappedSize()
-				});
-
-				TrackFileHandlesAndRegions(FileHandle, FileRegion);
+				IMappedFileHandle* FileHandle = PlatformFile->OpenMapped(*DebugFilePath);
+				if (FileHandle != nullptr)
+				{
+					IMappedFileRegion* FileRegion = FileHandle->MapRegion(0, FileHandle->GetFileSize());
+					check(FileRegion != nullptr);
+
+					Files.Push(SymsFile{
+						Path.data,
+						(void*)FileRegion->GetMappedPtr(),
+						(uint64)FileRegion->GetMappedSize()
+					});
+
+					TrackFileHandlesAndRegions(FileHandle, FileRegion);
+				}
+				else
+				{
+					UE_LOG(LogSymslib, Warning, TEXT("Could not open the debug file '%s'"), *DebugFilePath);
+				}
 			}
 		}
 	}
-	
+
 	// Prepare to load the debug files
 	Result = syms_load_debug_info_ex(ModuleInstance, Files.GetData(), Files.Num(), SYMS_LOAD_DEBUG_INFO_FLAGS_DEFER_BUILD_MODULE);
 
@@ -265,10 +280,13 @@ FSymslibResolver::EModuleStatus FSymslibResolver::LoadModule(FModuleEntry* Modul
 	const uint32 Count = syms_get_module_build_count(ModuleInstance);
 	// Each work item needs some memory
 	TArray<SymsArena*> Arenas;
-	for (uint32 i = 0; i < Count; ++i) { Arenas.Push(syms_borrow_memory(ModuleInstance)); }
-	
-	ParallelFor(Count, [ModuleInstance, Arenas](uint32 Index) {
-		syms_build_module(ModuleInstance, (SymsModID)Index, Arenas[Index]); 
+	for (uint32 Index = 0; Index < Count; ++Index)
+	{
+		Arenas.Add(syms_borrow_memory(ModuleInstance));
+	}
+	ParallelFor(Count, [ModuleInstance, &Arenas](uint32 Index)
+	{
+		syms_build_module(ModuleInstance, (SymsModID)Index, Arenas[Index]);
 	});
 
 	// Check age of debug data
@@ -281,7 +299,7 @@ FSymslibResolver::EModuleStatus FSymslibResolver::LoadModule(FModuleEntry* Modul
 		uint32* ModuleAge = ((uint32*) Module->ImageId.GetData()) + 4;
 		SymsDebugInfoPdb* PdbInfo = (SymsDebugInfoPdb*)ModuleInstance->debug_info.impl;
 		FGuid* PdbGuid = (FGuid*)&PdbInfo->context.auth_guid;
-		const uint32 PdbAge = PdbInfo->context.age; 
+		const uint32 PdbAge = PdbInfo->context.age;
 		if (*ModuleGuid != *PdbGuid || *ModuleAge != PdbAge)
 		{
 			UE_LOG(LogSymslib, Warning, TEXT("Symbols for '%s' does not match traced binary."), Module->Name);
@@ -342,26 +360,32 @@ bool FSymslibResolver::ResolveSymbol(uint64 Address, FResolvedSymbol* Target)
 	// Find procedure and source file for address
 	if (!syms_proc_from_va(Module->Instance, Address, &Proc) || !syms_va_to_src(Module->Instance, Address, &File))
 	{
-		UpdateResolvedSymbol(Target, ESymbolQueryResult::NotFound, Module->Name, GUnknownModuleTextSymsLib); 
+		UpdateResolvedSymbol(Target, ESymbolQueryResult::NotFound, Module->Name, GUnknownModuleTextSymsLib);
 		return false;
 	}
 
+	constexpr uint32 MaxStringSize = 1024;
+
 	// Read the symbol name
-	ANSICHAR ProcedureName[1024];
-	syms_read_strref(Module->Instance, &Proc.name_ref, ProcedureName, sizeof(ProcedureName));
+	ANSICHAR ProcedureName[MaxStringSize];
+	ProcedureName[MaxStringSize - 1] = 0;
+	syms_read_strref(Module->Instance, &Proc.name_ref, ProcedureName, sizeof(ProcedureName) - 1);
+	check(ProcedureName[MaxStringSize - 1] == 0);
 
 	// Read the source location
-	ANSICHAR SourceFile[1024];
-	syms_read_strref(Module->Instance, &File.file.name, SourceFile, sizeof(SourceFile));
+	ANSICHAR SourceFile[MaxStringSize];
+	SourceFile[MaxStringSize - 1] = 0;
+	syms_read_strref(Module->Instance, &File.file.name, SourceFile, sizeof(SourceFile) - 1);
+	check(SourceFile[MaxStringSize - 1] == 0);
 
 	TStringBuilder<1024> SourceLocation;
 	SourceLocation << ANSI_TO_TCHAR(SourceFile) << TEXT("(") << File.line.ln << TEXT(")");
 
 	// Store the strings and update the target data
 	UpdateResolvedSymbol(
-		Target, 
+		Target,
 		ESymbolQueryResult::OK,
-		Session.StoreString(ANSI_TO_TCHAR(ProcedureName)), 
+		Session.StoreString(ANSI_TO_TCHAR(ProcedureName)),
 		Session.StoreString(SourceLocation)
 	);
 
@@ -370,7 +394,8 @@ bool FSymslibResolver::ResolveSymbol(uint64 Address, FResolvedSymbol* Target)
 
 void FSymslibResolver::TrackFileHandlesAndRegions(IMappedFileHandle* FileHandle, IMappedFileRegion* FileRegion)
 {
-	FileHandlesAndRegions.Add(TTuple<IMappedFileHandle*, IMappedFileRegion*>(FileHandle, FileRegion));
+	FWriteScopeLock _(FileHandlesAndRegionsLock);
+	FileHandlesAndRegions.Add(FMappedFileAndRegion{ FileHandle, FileRegion });
 }
 
 /////////////////////////////////////////////////////////////////////
diff --git a/Engine/Source/Developer/TraceServices/Private/Model/SymslibResolver.h b/Engine/Source/Developer/TraceServices/Private/Model/SymslibResolver.h
--- a/Engine/Source/Developer/TraceServices/Private/Model/SymslibResolver.h
+++ b/Engine/Source/Developer/TraceServices/Private/Model/SymslibResolver.h
@@ -1,5 +1,6 @@
 // Copyright Epic Games, Inc. All Rights Reserved.
 #pragma once
+
 #if PLATFORM_WINDOWS
 
 #include "Async/MappedFileHandle.h"
@@ -26,7 +27,6 @@ public:
 	void GetStats(IModuleProvider::FStats* OutStats) const;
 
 private:
-
 	enum class EModuleStatus : uint8
 	{
 		Pending,
@@ -65,7 +65,7 @@ private:
 	};
 
 	/**
-	 * Checks if there are no modules in flight and that the queue has reached 
+	 * Checks if there are no modules in flight and that the queue has reached
 	 * the threshold for dispatching. Note that is up to the caller to synchronize.
 	 */
 	void MaybeDispatchQueuedAddresses();
@@ -80,15 +80,22 @@ private:
 	void ResolveSymbolTracked(uint64 Address, FResolvedSymbol* Target);
 	bool ResolveSymbol(uint64 Address, FResolvedSymbol* Target);
 	void TrackFileHandlesAndRegions(IMappedFileHandle* FileHandle, IMappedFileRegion* FileRegion);
-	
+
 	FRWLock ModulesLock;
 	TPagedArray<FModuleEntry> Modules;
 	TArray<FModuleEntry*> SortedModules;
 	FCriticalSection SymbolsQueueLock;
 	TArray<FQueuedAddress, TInlineAllocator<QueuedAddressLength>> ResolveQueue;
 	std::atomic<uint32> TasksInFlight;
-	TArray<TTuple<IMappedFileHandle*, IMappedFileRegion*>> FileHandlesAndRegions;
-	
+
+	struct FMappedFileAndRegion
+	{
+		IMappedFileHandle* Handle;
+		IMappedFileRegion* Region;
+	};
+	TArray<FMappedFileAndRegion> FileHandlesAndRegions;
+	FRWLock FileHandlesAndRegionsLock;
+
 	std::atomic<uint32> ModulesDiscovered;
 	std::atomic<uint32> ModulesFailed;
 	std::atomic<uint32> ModulesLoaded;
