TasksInsights: improved `FTasksProvider::EnumerateTasks` perf by exploiting the fact that `ExecutionThreads` are sorted by `Task.StartedTimestamp`

diff --git a/Engine/Source/Developer/TraceServices/Private/Model/TasksProfiler.cpp b/Engine/Source/Developer/TraceServices/Private/Model/TasksProfiler.cpp
--- a/Engine/Source/Developer/TraceServices/Private/Model/TasksProfiler.cpp
+++ b/Engine/Source/Developer/TraceServices/Private/Model/TasksProfiler.cpp
@@ -439,16 +439,41 @@ namespace TraceServices
 
 	void FTasksProvider::EnumerateTasks(double StartTime, double EndTime, TaskCallback Callback) const
 	{
-		// A naive implementation for now
-		for (const FTaskInfo& Task : Tasks)
+		for (const TPair<uint32, TArray64<TaskTrace::FId>>& KeyValue : ExecutionThreads)
 		{
-			if (Task.CreatedTimestamp <= EndTime && Task.FinishedTimestamp >= StartTime)
+			const TArray64<TaskTrace::FId>& Thread = KeyValue.Value;
+			
+			// find the first task with `StartedTimestamp <= StartTime`
+			int64 TaskIndex = Algo::LowerBound(Thread, StartTime, 
+				[this](TaskTrace::FId TaskId, double Timestamp)
+				{
+					const FTaskInfo* Task = TryGetTask(TaskId);
+					return Task != nullptr && Task->StartedTimestamp <= Timestamp;
+				}
+			);
+
+			// check if there's a previous task whose execution overlaps `StartTime`
+			if (TaskIndex != 0)
 			{
-				if (Callback(Task) == ETaskEnumerationResult::Stop)
+				const FTaskInfo* Task = TryGetTask(Thread[TaskIndex - 1]);
+				if (Task != nullptr && Task->FinishedTimestamp > StartTime)
 				{
-					break;
+					--TaskIndex;
 				}
 			}
+
+			if (TaskIndex == Thread.Num())
+			{
+				continue; // all tasks on this thread are before StartTime so nothing to do here, go to the next thread
+			}
+
+			// report all tasks whose execution overlaps [StartTime, EndTime]
+			const FTaskInfo* Task = TryGetTask(Thread[TaskIndex]);
+			while (Task != nullptr && Task->StartedTimestamp <= EndTime && Callback(*Task) != ETaskEnumerationResult::Stop)
+			{
+				++TaskIndex;
+				Task = TryGetTask(Thread[TaskIndex]);
+			}
 		}
 	}
 }
