[Insights] Reduce stalls when analysing symbols

When parsing debug files and resolving symbols the RAD symbol resolver went very wide. This had two negative effects on performance:

* Module loading and symbol resolve tasks were scheduled together with Slate tasks, causing stuttering in rendering. With this change we move these tasks to the background threads since they are not time critical and to give preference to Slate rendering.
* Symbol resolves are very cheap compared to creating a task. With this change we batch resolve work in a fixed number of tasks.
* Resolve symbol tasks were competing for the string store lock. The ModuleProvider already handles symbol deduplication so each symbol string should already be unique. This change add a task local string allocator which does not have be thread-safe and can take advantage of batching by allocating memory in blocks.

#rnx

diff --git a/Engine/Source/Developer/TraceServices/Private/Model/SymslibResolver.cpp b/Engine/Source/Developer/TraceServices/Private/Model/SymslibResolver.cpp
--- a/Engine/Source/Developer/TraceServices/Private/Model/SymslibResolver.cpp
+++ b/Engine/Source/Developer/TraceServices/Private/Model/SymslibResolver.cpp
@@ -29,6 +29,49 @@ namespace TraceServices {
 
 static const TCHAR* GUnknownModuleTextSymsLib = TEXT("Unknown");
 
+/////////////////////////////////////////////////////////////////////
+
+class FSymbolStringAllocator
+{
+public:
+	FSymbolStringAllocator(ILinearAllocator& InAllocator, uint32 InBlockSize)
+		: Allocator(InAllocator)
+		, BlockSize(InBlockSize)
+	{}
+
+	const TCHAR* Store(const TCHAR* InString)
+	{
+		return Store(FStringView(InString));
+	}
+
+	const TCHAR* Store(const FStringView InString)
+	{
+		const uint32 StringSize = InString.Len() + 1;
+		check(StringSize <= BlockSize);
+		if (StringSize > BlockRemaining)
+		{
+			Block = (TCHAR*) Allocator.Allocate(BlockSize * sizeof(TCHAR));
+			BlockRemaining = BlockSize;
+			++BlockUsed;
+		}
+		const uint32 CopiedSize = InString.CopyString(Block, BlockRemaining - 1, 0);
+		check(StringSize - 1 == CopiedSize);
+		Block[StringSize] = TEXT('\0');
+		BlockRemaining -= StringSize;
+		const TCHAR* OutString = Block;
+		Block += StringSize;
+		return OutString;
+	}
+
+private:
+	friend class FSymslibResolver;
+	ILinearAllocator& Allocator;
+	TCHAR* Block = nullptr;
+	uint32 BlockSize;
+	uint32 BlockRemaining = 0;
+	uint32 BlockUsed = 0;
+};
+	
 /////////////////////////////////////////////////////////////////////
 FSymslibResolver::FSymslibResolver(IAnalysisSession& InSession)
 	: Modules(InSession.GetLinearAllocator(), 128)
@@ -56,9 +99,15 @@ void FSymslibResolver::QueueModuleLoad(const FStringView& ModulePath, uint64 Bas
 
 	++ModulesDiscovered;
 
-	// Queue up module to have symbols loaded
+	// Queue up module to have symbols loaded. Run as background task as to not
+	// interfere with Slate.
 	++TasksInFlight;
-	FFunctionGraphTask::CreateAndDispatchWhenReady([this, Entry] { LoadModuleTracked(Entry); --TasksInFlight;});
+	FFunctionGraphTask::CreateAndDispatchWhenReady(
+		[this, Entry]
+		{
+			LoadModuleTracked(Entry);
+			--TasksInFlight;
+		}, TStatId{}, nullptr, ENamedThreads::AnyBackgroundThreadNormalTask);
 
 	// Sort list according to base address
 	SortedModules.Add(Entry);
@@ -101,6 +150,8 @@ void FSymslibResolver::OnAnalysisComplete()
 			delete File.Handle;
 		}
 	}
+
+	UE_LOG(LogSymslib, Display, TEXT("Allocated %.02f Mb of strings, %.02f Mb wasted."), SymbolBytesAllocated / float(1024*1024), SymbolBytesWasted / float(1024*1024));
 }
 
 void FSymslibResolver::GetStats(IModuleProvider::FStats* OutStats) const
@@ -132,16 +183,35 @@ void FSymslibResolver::MaybeDispatchQueuedAddresses()
  */
 void FSymslibResolver::DispatchQueuedAddresses()
 {
-	//todo: This creates two copies of resolve queue!
 	TArray<FQueuedAddress> WorkingSet(ResolveQueue);
-	++TasksInFlight;
-	FFunctionGraphTask::CreateAndDispatchWhenReady([this, WorkingSet] {
-		ParallelFor(WorkingSet.Num(), [this, &WorkingSet](uint32 Index) {
-			const FQueuedAddress& ToResolve = WorkingSet[Index];
-			ResolveSymbolTracked(ToResolve.Address, ToResolve.Target);
-		});
+	TasksInFlight += SymbolTasksInParallel;
+	const uint32 Stride = (WorkingSet.Num() + SymbolTasksInParallel - 1) / SymbolTasksInParallel;
+
+	// Use background priority in order to not not interfere with Slate
+	ParallelFor(SymbolTasksInParallel, [this, &WorkingSet, Stride](uint32 Index) {
+		const uint32 StartIndex = Index * Stride;
+		const uint32 EndIndex = FMath::Min(StartIndex + Stride, (uint32)WorkingSet.Num());
+		TArrayView<FQueuedAddress> QueuedWork(&WorkingSet[StartIndex], EndIndex - StartIndex);
+		ResolveSymbols(QueuedWork);
 		--TasksInFlight;
-	});
+	}, EParallelForFlags::BackgroundPriority);
+}
+
+void FSymslibResolver::ResolveSymbols(TArrayView<FQueuedAddress>& QueuedWork)
+{
+	// Create a local string allocator. We don't use the session string store due to contention when going wide. Since
+	// the ModuleProvider already de-duplicates symbols we do not need this feature from the string store.
+	FSymbolStringAllocator StringAllocator(Session.GetLinearAllocator(), (2 << 12) / sizeof(TCHAR) );
+	for (const FQueuedAddress& ToResolve : QueuedWork)
+	{
+		ResolveSymbolTracked(ToResolve.Address, ToResolve.Target, StringAllocator);
+	}
+	UE_LOG(LogSymslib, VeryVerbose, TEXT("String allocator used: %.02f kb, wasted: %.02f kb using %d blocks"),
+		((StringAllocator.BlockUsed * StringAllocator.BlockSize - StringAllocator.BlockRemaining) * sizeof(TCHAR)) / 1024.0f,
+		(StringAllocator.BlockRemaining * sizeof(TCHAR)) / 1024.0f,
+		StringAllocator.BlockUsed);
+	SymbolBytesAllocated.fetch_add(StringAllocator.BlockUsed * StringAllocator.BlockSize * sizeof(TCHAR));
+	SymbolBytesWasted.fetch_add(StringAllocator.BlockRemaining * sizeof(TCHAR));
 }
 
 FSymslibResolver::FModuleEntry* FSymslibResolver::GetModuleForAddress(uint64 Address)
@@ -314,9 +384,9 @@ FSymslibResolver::EModuleStatus FSymslibResolver::LoadModule(FModuleEntry* Modul
 	return EModuleStatus::Loaded;
 }
 
-void FSymslibResolver::ResolveSymbolTracked(uint64 Address, FResolvedSymbol* Target)
+void FSymslibResolver::ResolveSymbolTracked(uint64 Address, FResolvedSymbol* Target, FSymbolStringAllocator& StringAllocator)
 {
-	if (!ResolveSymbol(Address, Target))
+	if (!ResolveSymbol(Address, Target, StringAllocator))
 	{
 		++SymbolsFailed;
 	}
@@ -326,7 +396,7 @@ void FSymslibResolver::ResolveSymbolTracked(uint64 Address, FResolvedSymbol* Tar
 	}
 }
 
-bool FSymslibResolver::ResolveSymbol(uint64 Address, FResolvedSymbol* Target)
+bool FSymslibResolver::ResolveSymbol(uint64 Address, FResolvedSymbol* Target, FSymbolStringAllocator& StringAllocator)
 {
 	FModuleEntry* Module = GetModuleForAddress(Address);
 	if (!Module)
@@ -367,10 +437,10 @@ bool FSymslibResolver::ResolveSymbol(uint64 Address, FResolvedSymbol* Target)
 	constexpr uint32 MaxStringSize = 1024;
 
 	// Read the symbol name
-	ANSICHAR ProcedureName[MaxStringSize];
-	ProcedureName[MaxStringSize - 1] = 0;
-	syms_read_strref(Module->Instance, &Proc.name_ref, ProcedureName, sizeof(ProcedureName) - 1);
-	check(ProcedureName[MaxStringSize - 1] == 0);
+	ANSICHAR SymbolName[MaxStringSize];
+	SymbolName[MaxStringSize - 1] = 0;
+	syms_read_strref(Module->Instance, &Proc.name_ref, SymbolName, sizeof(SymbolName) - 1);
+	check(SymbolName[MaxStringSize - 1] == 0);
 
 	// Read the source location
 	ANSICHAR SourceFile[MaxStringSize];
@@ -378,15 +448,18 @@ bool FSymslibResolver::ResolveSymbol(uint64 Address, FResolvedSymbol* Target)
 	syms_read_strref(Module->Instance, &File.file.name, SourceFile, sizeof(SourceFile) - 1);
 	check(SourceFile[MaxStringSize - 1] == 0);
 
-	TStringBuilder<1024> SourceLocation;
-	SourceLocation << ANSI_TO_TCHAR(SourceFile) << TEXT("(") << File.line.ln << TEXT(")");
+	TStringBuilder<1024> FileAndLine;
+	FileAndLine << ANSI_TO_TCHAR(SourceFile) << TEXT("(") << File.line.ln << TEXT(")");
+
+	const TCHAR* SymbolNamePersistent =  StringAllocator.Store(ANSI_TO_TCHAR(SymbolName));
+	const TCHAR* FileAndLinePersistent = StringAllocator.Store(FileAndLine);
 
 	// Store the strings and update the target data
 	UpdateResolvedSymbol(
 		Target,
 		ESymbolQueryResult::OK,
-		Session.StoreString(ANSI_TO_TCHAR(ProcedureName)),
-		Session.StoreString(SourceLocation)
+		SymbolNamePersistent,
+		FileAndLinePersistent
 	);
 
 	return true;
diff --git a/Engine/Source/Developer/TraceServices/Private/Model/SymslibResolver.h b/Engine/Source/Developer/TraceServices/Private/Model/SymslibResolver.h
--- a/Engine/Source/Developer/TraceServices/Private/Model/SymslibResolver.h
+++ b/Engine/Source/Developer/TraceServices/Private/Model/SymslibResolver.h
@@ -61,7 +61,8 @@ private:
 
 	enum : uint32 {
 		MaxNameLen = 512,
-		QueuedAddressLength = 256
+		QueuedAddressLength = 2048,
+		SymbolTasksInParallel = 8
 	};
 
 	/**
@@ -73,12 +74,13 @@ private:
 	 * Dispatches the currently queued addresses to be resolved. Note that is up to the caller to synchronize.
 	 */
 	void DispatchQueuedAddresses();
+	void ResolveSymbols(TArrayView<FQueuedAddress>& QueuedWork);
 	FModuleEntry* GetModuleForAddress(uint64 Address);
 	static void UpdateResolvedSymbol(FResolvedSymbol* Symbol, ESymbolQueryResult Result, const TCHAR* Name, const TCHAR* FileAndLine);
 	void LoadModuleTracked(FModuleEntry* Module);
 	EModuleStatus LoadModule(FModuleEntry* Module);
-	void ResolveSymbolTracked(uint64 Address, FResolvedSymbol* Target);
-	bool ResolveSymbol(uint64 Address, FResolvedSymbol* Target);
+	void ResolveSymbolTracked(uint64 Address, FResolvedSymbol* Target, class FSymbolStringAllocator& StringAllocator);
+	bool ResolveSymbol(uint64 Address, FResolvedSymbol* Target, class FSymbolStringAllocator& StringAllocator);
 	void TrackFileHandlesAndRegions(IMappedFileHandle* FileHandle, IMappedFileRegion* FileRegion);
 
 	FRWLock ModulesLock;
@@ -102,6 +104,8 @@ private:
 	std::atomic<uint32> SymbolsDiscovered;
 	std::atomic<uint32> SymbolsFailed;
 	std::atomic<uint32> SymbolsResolved;
+	std::atomic<uint64> SymbolBytesAllocated;
+	std::atomic<uint64> SymbolBytesWasted;
 
 	IAnalysisSession& Session;
 };
