[Insights] Reduce memory footprint of SymslibResolver

Reduces the memory footprint of the resolver by releasing all memory associated with each module once analysis is completed.

#rnx

diff --git a/Engine/Source/Developer/TraceServices/Private/Model/SymslibResolver.cpp b/Engine/Source/Developer/TraceServices/Private/Model/SymslibResolver.cpp
--- a/Engine/Source/Developer/TraceServices/Private/Model/SymslibResolver.cpp
+++ b/Engine/Source/Developer/TraceServices/Private/Model/SymslibResolver.cpp
@@ -152,6 +152,12 @@ void FSymslibResolver::OnAnalysisComplete()
 	}
 
 	UE_LOG(LogSymslib, Display, TEXT("Allocated %.02f Mb of strings, %.02f Mb wasted."), SymbolBytesAllocated / float(1024*1024), SymbolBytesWasted / float(1024*1024));
+	// Release memory used by syms library
+	for (uint32 ModuleIndex = 0; ModuleIndex < Modules.Num(); ++ModuleIndex)
+	{
+		FModuleEntry& Module = Modules[ModuleIndex];
+		syms_quit(Module.Instance);
+	}
 }
 
 void FSymslibResolver::GetStats(IModuleProvider::FStats* OutStats) const
