[Insights] Fix off by one error with string store

#rnx

diff --git a/Engine/Source/Developer/TraceServices/Private/Model/SymslibResolver.cpp b/Engine/Source/Developer/TraceServices/Private/Model/SymslibResolver.cpp
--- a/Engine/Source/Developer/TraceServices/Private/Model/SymslibResolver.cpp
+++ b/Engine/Source/Developer/TraceServices/Private/Model/SymslibResolver.cpp
@@ -55,8 +55,8 @@ public:
 			++BlockUsed;
 		}
 		const uint32 CopiedSize = InString.CopyString(Block, BlockRemaining - 1, 0);
-		check(StringSize - 1 == CopiedSize);
-		Block[StringSize] = TEXT('\0');
+		check(StringSize == CopiedSize + 1);
+		Block[StringSize - 1] = TEXT('\0');
 		BlockRemaining -= StringSize;
 		const TCHAR* OutString = Block;
 		Block += StringSize;
