[Insights] Improve callstack display when frame cannot be displayed

Previously the name was only displayed when results was successful. With this change we use the additional information in the status code to display module name and reason why the symbol name could not be resolved.

diff --git a/Engine/Source/Developer/TraceInsights/Private/Insights/MemoryProfiler/ViewModels/CallstackFormatting.h b/Engine/Source/Developer/TraceInsights/Private/Insights/MemoryProfiler/ViewModels/CallstackFormatting.h
new file mode 100644
--- /dev/null
+++ b/Engine/Source/Developer/TraceInsights/Private/Insights/MemoryProfiler/ViewModels/CallstackFormatting.h
@@ -0,0 +1,48 @@
+﻿// Copyright Epic Games, Inc. All Rights Reserved.
+#pragma once
+#include "TraceServices/Model/Callstack.h"
+#include "TraceServices/Model/Modules.h"
+
+/////////////////////////////////////////////////////////////////////
+enum EStackFrameFormatFlags : uint8
+{
+	Module			= 1 << 0,
+	FileAndLine		= 1 << 1
+};
+
+/////////////////////////////////////////////////////////////////////
+inline void FormatStackFrame(const TraceServices::FStackFrame& Frame, FStringBuilderBase& OutString, uint8 FormatFlags)
+{
+	using namespace TraceServices;
+	const ESymbolQueryResult Result = Frame.Symbol->GetResult();
+	switch(Result)
+	{
+		case ESymbolQueryResult::OK:
+			if (FormatFlags & (uint8)EStackFrameFormatFlags::Module)
+			{
+				OutString.Appendf(TEXT("%s!"), Frame.Symbol->Module);
+			}
+			OutString.Append(Frame.Symbol->Name);
+			if (FormatFlags & (uint8)EStackFrameFormatFlags::FileAndLine)
+			{
+				OutString.Appendf(TEXT(" %s(%d)"), Frame.Symbol->File, Frame.Symbol->Line);
+			}
+			break;
+		case ESymbolQueryResult::Mismatch:
+		case ESymbolQueryResult::NotFound:
+		case ESymbolQueryResult::NotLoaded:
+			if (FormatFlags & (uint8)EStackFrameFormatFlags::Module)
+			{
+				OutString.Appendf(TEXT("%s!0x%08x (%s)"), Frame.Symbol->Module, Frame.Addr, QueryResultToString((Result)));
+			}
+			else
+			{
+				OutString.Appendf(TEXT("0x%08x (%s)"), Frame.Symbol->Name, Frame.Addr, QueryResultToString((Result)));
+			}
+			break;
+		case ESymbolQueryResult::Pending:
+		default:
+			OutString.Append(QueryResultToString(Result));
+			break;
+	}
+}
diff --git a/Engine/Source/Developer/TraceInsights/Private/Insights/MemoryProfiler/ViewModels/MemAllocGroupingByCallstack.cpp b/Engine/Source/Developer/TraceInsights/Private/Insights/MemoryProfiler/ViewModels/MemAllocGroupingByCallstack.cpp
--- a/Engine/Source/Developer/TraceInsights/Private/Insights/MemoryProfiler/ViewModels/MemAllocGroupingByCallstack.cpp
+++ b/Engine/Source/Developer/TraceInsights/Private/Insights/MemoryProfiler/ViewModels/MemAllocGroupingByCallstack.cpp
@@ -2,6 +2,7 @@
 
 #include "MemAllocGroupingByCallstack.h"
 #include "TraceServices/Model/Callstack.h"
+#include "CallstackFormatting.h"
 
 // Insights
 #include "Insights/MemoryProfiler/ViewModels/MemAllocNode.h"
@@ -176,20 +177,9 @@ FName FMemAllocGroupingByCallstack::GetGroupName(const TraceServices::FStackFram
 
 FText FMemAllocGroupingByCallstack::GetGroupTooltip(const TraceServices::FStackFrame* Frame) const
 {
-	const TraceServices::ESymbolQueryResult Result = Frame->Symbol->GetResult();
-	if (Result == TraceServices::ESymbolQueryResult::OK)
-	{
-		return FText::Format(LOCTEXT("CallstackFrameTooltipFmt1", "Callstack Frame\n\t{0}\n\t{1}\n\t{2}"),
-			FText::FromString(FString::Printf(TEXT("0x%X"), Frame->Addr)),
-			FText::FromString(FString(Frame->Symbol->Name)),
-			FText::FromString(FString(Frame->Symbol->FileAndLine)));
-	}
-	else
-	{
-		return FText::Format(LOCTEXT("CallstackFrameTooltipFmt2", "Callstack Frame\n\t{0}\n\t{1}"),
-			FText::FromString(FString::Printf(TEXT("0x%X"), Frame->Addr)),
-			FText::FromString(FString(TraceServices::QueryResultToString(Result))));
-	}
+	TStringBuilder<1024> String;
+	FormatStackFrame(*Frame, String, EStackFrameFormatFlags::Module | EStackFrameFormatFlags::FileAndLine);
+	return FText::FromString(FString(String));
 }
 
 ////////////////////////////////////////////////////////////////////////////////////////////////////
diff --git a/Engine/Source/Developer/TraceInsights/Private/Insights/MemoryProfiler/ViewModels/MemAllocTable.cpp b/Engine/Source/Developer/TraceInsights/Private/Insights/MemoryProfiler/ViewModels/MemAllocTable.cpp
--- a/Engine/Source/Developer/TraceInsights/Private/Insights/MemoryProfiler/ViewModels/MemAllocTable.cpp
+++ b/Engine/Source/Developer/TraceInsights/Private/Insights/MemoryProfiler/ViewModels/MemAllocTable.cpp
@@ -1,6 +1,7 @@
 // Copyright Epic Games, Inc. All Rights Reserved.
 
 #include "MemAllocTable.h"
+#include "CallstackFormatting.h"
 
 #include "Widgets/SBoxPanel.h"
 #include "Widgets/SToolTip.h"
@@ -520,9 +521,7 @@ void FMemAllocTable::AddDefaultColumns()
 					const FMemoryAlloc* Alloc = MemAllocNode.GetMemAlloc();
 					if (Alloc)
 					{
-						static const TCHAR* NotAvailable = TEXT("N/A");
-						const TCHAR* Value = NotAvailable;
-
+						static FString NotAvailable = TEXT("Unknown callstack");
 						const TraceServices::FCallstack* Callstack = Alloc->GetCallstack();
 
 						if (Callstack)
@@ -531,18 +530,12 @@ void FMemAllocTable::AddDefaultColumns()
 							const TraceServices::FStackFrame* Frame = Callstack->Frame(FMath::Min(2u, Callstack->Num() - 1));
 							check(Frame != nullptr);
 
-							const TraceServices::ESymbolQueryResult Result = Frame->Symbol->GetResult();
-							if (Result == TraceServices::ESymbolQueryResult::OK)
-							{
-								Value = Frame->Symbol->Name;
-							}
-							else
-							{
-								Value = TraceServices::QueryResultToString(Result);
-							}
+							TStringBuilder<1024> Str;
+							FormatStackFrame(*Frame, Str, EStackFrameFormatFlags::Module);
+							return FTableCellValue(FText::FromString(FString(Str)));
 						}
 
-						return FTableCellValue(Value);
+						return FTableCellValue(FText::FromString(NotAvailable));
 					}
 				}
 
@@ -552,7 +545,7 @@ void FMemAllocTable::AddDefaultColumns()
 		TSharedRef<ITableCellValueGetter> Getter = MakeShared<FFunctionValueGetter>();
 		Column.SetValueGetter(Getter);
 
-		class FunctionValueFormatter : public FCStringValueFormatterAsText
+		class FunctionValueFormatter : public FTextValueFormatter
 		{
 		public:
 			virtual TSharedPtr<IToolTip> GetCustomTooltip(const FTableColumn& Column, const FBaseTreeNode& Node) const override
diff --git a/Engine/Source/Developer/TraceInsights/Private/Insights/MemoryProfiler/ViewModels/MemoryAlloc.cpp b/Engine/Source/Developer/TraceInsights/Private/Insights/MemoryProfiler/ViewModels/MemoryAlloc.cpp
--- a/Engine/Source/Developer/TraceInsights/Private/Insights/MemoryProfiler/ViewModels/MemoryAlloc.cpp
+++ b/Engine/Source/Developer/TraceInsights/Private/Insights/MemoryProfiler/ViewModels/MemoryAlloc.cpp
@@ -1,6 +1,7 @@
 // Copyright Epic Games, Inc. All Rights Reserved.
 
 #include "MemoryAlloc.h"
+#include "CallstackFormatting.h"
 
 namespace Insights
 {
@@ -27,28 +28,22 @@ FMemoryAlloc::~FMemoryAlloc()
 
 FText FMemoryAlloc::GetFullCallstack() const
 {
-	FString Tooltip;
+	TStringBuilder<1024> Tooltip;
 	if (Callstack)
 	{
 		check(Callstack->Num() > 0);
-		for (uint32 Index = 0; Index < Callstack->Num(); ++Index)
+		const uint32 FramesNum = Callstack->Num();
+		for (uint32 Index = 0; Index < FramesNum; ++Index)
 		{
 			const TraceServices::FStackFrame* Frame = Callstack->Frame(Index);
 			check(Frame != nullptr);
 
-			const TraceServices::ESymbolQueryResult Result = Frame->Symbol->GetResult();
-			if (Result == TraceServices::ESymbolQueryResult::OK)
+			FormatStackFrame(*Frame, Tooltip, EStackFrameFormatFlags::Module | EStackFrameFormatFlags::FileAndLine);
+			if (Index != (FramesNum - 1))
 			{
-				Tooltip.Appendf(TEXT("%s\n"), Frame->Symbol->Name);
-			}
-			else
-			{
-				Tooltip.Appendf(TEXT("%s\n"), QueryResultToString(Result));
+				Tooltip << TEXT("\n");
 			}
 		}
-
-		// Remove the last "\n"
-		Tooltip.RemoveAt(Tooltip.Len() - 1, 1, false);
 	}
 
 	if (Tooltip.Len() == 0)
@@ -56,7 +51,7 @@ FText FMemoryAlloc::GetFullCallstack() const
 		Tooltip.Append(TEXT("N/A"));
 	}
 
-	return FText::FromString(Tooltip);
+	return FText::FromString(FString(Tooltip));
 }
 
 } // namespace Insights
diff --git a/Engine/Source/Developer/TraceServices/Private/Model/CallstacksProvider.cpp b/Engine/Source/Developer/TraceServices/Private/Model/CallstacksProvider.cpp
--- a/Engine/Source/Developer/TraceServices/Private/Model/CallstacksProvider.cpp
+++ b/Engine/Source/Developer/TraceServices/Private/Model/CallstacksProvider.cpp
@@ -12,7 +12,7 @@ namespace TraceServices
 {
 
 /////////////////////////////////////////////////////////////////////
-static const FResolvedSymbol GNeverResolveSymbol(ESymbolQueryResult::NotLoaded, nullptr, nullptr);
+static const FResolvedSymbol GNeverResolveSymbol(ESymbolQueryResult::NotLoaded, nullptr, nullptr, nullptr, 0);
 
 /////////////////////////////////////////////////////////////////////
 #ifdef TRACE_CALLSTACK_STATS
diff --git a/Engine/Source/Developer/TraceServices/Private/Model/DbgHelpResolver.cpp b/Engine/Source/Developer/TraceServices/Private/Model/DbgHelpResolver.cpp
--- a/Engine/Source/Developer/TraceServices/Private/Model/DbgHelpResolver.cpp
+++ b/Engine/Source/Developer/TraceServices/Private/Model/DbgHelpResolver.cpp
@@ -177,10 +177,12 @@ void FDbgHelpResolver::Stop()
 }
 
 /////////////////////////////////////////////////////////////////////
-void FDbgHelpResolver::UpdateResolvedSymbol(FResolvedSymbol* Symbol, ESymbolQueryResult Result, const TCHAR* Name, const TCHAR* FileAndLine)
+void FDbgHelpResolver::UpdateResolvedSymbol(FResolvedSymbol* Symbol, ESymbolQueryResult Result, const TCHAR* Module, const TCHAR* Name, const TCHAR* File, uint16 Line)
 {
+	Symbol->Module = Module;
 	Symbol->Name = Name;
-	Symbol->FileAndLine = FileAndLine;
+	Symbol->File = File;
+	Symbol->Line = Line;
 	Symbol->Result.store(Result, std::memory_order_release);
 }
 
@@ -189,6 +191,15 @@ void FDbgHelpResolver::ResolveSymbol(uint64 Address, FResolvedSymbol* Target)
 {
 	check(Target);
 
+	const FModuleEntry* Module = GetModuleForAddress(Address);
+	if (!Module)
+	{
+		++SymbolsFailed;
+		UpdateResolvedSymbol(Target, ESymbolQueryResult::NotLoaded, GUnknownModuleTextDbgHelp,
+		                     GUnknownModuleTextDbgHelp, GUnknownModuleTextDbgHelp, 0);
+		return;
+	}
+	
 	uint8 InfoBuffer[sizeof(SYMBOL_INFO) + (MaxNameLen * sizeof(char) + 1)];
 	SYMBOL_INFO* Info = (SYMBOL_INFO*)InfoBuffer;
 	Info->SizeOfStruct = sizeof(SYMBOL_INFO);
@@ -198,14 +209,13 @@ void FDbgHelpResolver::ResolveSymbol(uint64 Address, FResolvedSymbol* Target)
 	if (!SymFromAddr((HANDLE)Handle, Address, NULL, Info))
 	{
 		++SymbolsFailed;
-		UpdateResolvedSymbol(Target, ESymbolQueryResult::NotFound, GUnknownModuleTextDbgHelp, GUnknownModuleTextDbgHelp);
+		UpdateResolvedSymbol(Target, ESymbolQueryResult::NotFound, Module->Name, GUnknownModuleTextDbgHelp,
+		                     GUnknownModuleTextDbgHelp, 0);
 		return;
 	}
 
-	TStringBuilder<256> SymbolName;
-	SymbolName << ANSI_TO_TCHAR(Info->Name);
-	const TCHAR* SymbolNameStr = Session.StoreString(FStringView(SymbolName));
-
+	const TCHAR* SymbolNameStr = Session.StoreString(ANSI_TO_TCHAR(Info->Name));
+	
 	// Find the source file and line
 	DWORD  dwDisplacement;
 	IMAGEHLP_LINE Line;
@@ -214,16 +224,14 @@ void FDbgHelpResolver::ResolveSymbol(uint64 Address, FResolvedSymbol* Target)
 	if (!SymGetLineFromAddr((HANDLE)Handle, Address, &dwDisplacement, &Line))
 	{
 		++SymbolsFailed;
-		UpdateResolvedSymbol(Target, ESymbolQueryResult::NotFound, SymbolNameStr, GUnknownModuleTextDbgHelp);
+		UpdateResolvedSymbol(Target, ESymbolQueryResult::OK, Module->Name, SymbolNameStr, GUnknownModuleTextDbgHelp, 0);
 		return;
 	}
-
-	TStringBuilder<256> FileAndLine;
-	FileAndLine << ANSI_TO_TCHAR(Line.FileName) << TEXT(" (") << uint32(Line.LineNumber) << TEXT(")");
-	const TCHAR* FileAndLineStr = Session.StoreString(FStringView(FileAndLine));
+	
+	const TCHAR* SymbolFileStr = Session.StoreString(ANSI_TO_TCHAR(Line.FileName));
 	
 	++SymbolsResolved;
-	UpdateResolvedSymbol(Target, ESymbolQueryResult::OK, SymbolNameStr, FileAndLineStr);
+	UpdateResolvedSymbol(Target, ESymbolQueryResult::OK, Module->Name, SymbolNameStr, SymbolFileStr, Line.LineNumber);
 }
 
 /////////////////////////////////////////////////////////////////////
@@ -255,21 +263,15 @@ bool FDbgHelpResolver::LoadModuleSymbols(uint64 Base, uint64 Size, const TCHAR*
 
 
 /////////////////////////////////////////////////////////////////////
-const TCHAR* FDbgHelpResolver::GetModuleNameForAddress(uint64 Address) const
+const FDbgHelpResolver::FModuleEntry* FDbgHelpResolver::GetModuleForAddress(uint64 Address) const
 {
 	const int32 EntryIdx = Algo::LowerBoundBy(LoadedModules, Address, [](const FModuleEntry& Entry) { return Entry.Base; }) - 1;
 	if (EntryIdx < 0 || EntryIdx >= LoadedModules.Num())
 	{
-		return GUnknownModuleTextDbgHelp;
-	}
-
-	const FModuleEntry& Entry = LoadedModules[EntryIdx];
-	if (Address > (Entry.Base + Entry.Size))
-	{
-		return GUnknownModuleTextDbgHelp;
+		return nullptr;
 	}
 
-	return Entry.Name;
+	return &LoadedModules[EntryIdx];
 }
 
 
diff --git a/Engine/Source/Developer/TraceServices/Private/Model/DbgHelpResolver.h b/Engine/Source/Developer/TraceServices/Private/Model/DbgHelpResolver.h
--- a/Engine/Source/Developer/TraceServices/Private/Model/DbgHelpResolver.h
+++ b/Engine/Source/Developer/TraceServices/Private/Model/DbgHelpResolver.h
@@ -58,11 +58,11 @@ private:
 	void FreeSyms() const;
 	virtual uint32 Run() override;
 	virtual void Stop() override;
-	static void UpdateResolvedSymbol(FResolvedSymbol* Symbol, ESymbolQueryResult Result, const TCHAR* Name, const TCHAR* FileAndLine);
+	static void UpdateResolvedSymbol(FResolvedSymbol* Symbol, ESymbolQueryResult Result, const TCHAR* Module, const TCHAR* Name, const TCHAR* File, uint16 Line);
 
 	void ResolveSymbol(uint64 Address, FResolvedSymbol* Target);
 	bool LoadModuleSymbols(uint64 Base, uint64 Size, const TCHAR* Path);
-	const TCHAR* GetModuleNameForAddress(uint64 Address) const;
+	const FModuleEntry* GetModuleForAddress(uint64 Address) const;
 
 	FCriticalSection ModulesCs;
 	TArray<FModuleEntry> LoadedModules;
diff --git a/Engine/Source/Developer/TraceServices/Private/Model/ModuleProvider.cpp b/Engine/Source/Developer/TraceServices/Private/Model/ModuleProvider.cpp
--- a/Engine/Source/Developer/TraceServices/Private/Model/ModuleProvider.cpp
+++ b/Engine/Source/Developer/TraceServices/Private/Model/ModuleProvider.cpp
@@ -106,7 +106,7 @@ const FResolvedSymbol* TModuleProvider<SymbolProvider>::GetSymbol(uint64 Address
 		{
 			return SymbolCache[Address];
 		}
-		ResolvedSymbol = &Symbols.EmplaceBack(ESymbolQueryResult::Pending, nullptr, nullptr);
+		ResolvedSymbol = &Symbols.EmplaceBack(ESymbolQueryResult::Pending, nullptr, nullptr, nullptr, 0);
 		SymbolCache.Add(Address, ResolvedSymbol);
 	}
 
diff --git a/Engine/Source/Developer/TraceServices/Private/Model/SymslibResolver.cpp b/Engine/Source/Developer/TraceServices/Private/Model/SymslibResolver.cpp
--- a/Engine/Source/Developer/TraceServices/Private/Model/SymslibResolver.cpp
+++ b/Engine/Source/Developer/TraceServices/Private/Model/SymslibResolver.cpp
@@ -231,10 +231,12 @@ FSymslibResolver::FModuleEntry* FSymslibResolver::GetModuleForAddress(uint64 Add
 	return SortedModules[EntryIdx];
 }
 
-void FSymslibResolver::UpdateResolvedSymbol(FResolvedSymbol* Symbol, ESymbolQueryResult Result, const TCHAR* Name, const TCHAR* FileAndLine)
+void FSymslibResolver::UpdateResolvedSymbol(FResolvedSymbol* Symbol, ESymbolQueryResult Result, const TCHAR* Module, const TCHAR* Name, const TCHAR* File, uint16 Line)
 {
+	Symbol->Module = Module;
 	Symbol->Name = Name;
-	Symbol->FileAndLine = FileAndLine;
+	Symbol->File = File;
+	Symbol->Line = Line;
 	Symbol->Result.store(Result, std::memory_order_release);
 }
 
@@ -277,7 +279,7 @@ FSymslibResolver::EModuleStatus FSymslibResolver::LoadModule(FModuleEntry* Modul
 
 	if (!PlatformFile->FileExists(Module->Path))
 	{
-		UE_LOG(LogSymslib, Error, TEXT("File '%s' does not exist"), Module->Path);
+		UE_LOG(LogSymslib, Warning, TEXT("File '%s' does not exist"), Module->Path);
 		return EModuleStatus::Failed;
 	}
 
@@ -285,7 +287,7 @@ FSymslibResolver::EModuleStatus FSymslibResolver::LoadModule(FModuleEntry* Modul
 	IMappedFileHandle* ImageFileHandle = PlatformFile->OpenMapped(Module->Path);
 	if (ImageFileHandle == nullptr)
 	{
-		UE_LOG(LogSymslib, Error, TEXT("Failed to open '%s'"), Module->Path);
+		UE_LOG(LogSymslib, Warning, TEXT("Failed to open '%s'"), Module->Path);
 		return EModuleStatus::Failed;
 	}
 	IMappedFileRegion* ImageFileRegion = ImageFileHandle->MapRegion(0, ImageFileHandle->GetFileSize());
@@ -348,7 +350,7 @@ FSymslibResolver::EModuleStatus FSymslibResolver::LoadModule(FModuleEntry* Modul
 
 	if (SYMS_RESULT_FAIL(Result))
 	{
-		UE_LOG(LogSymslib, Warning, TEXT("Could not read debug files for '%s'"), Module->Path);
+		UE_LOG(LogSymslib, Display, TEXT("No debug information for '%s'"), Module->Path);
 		return EModuleStatus::Failed;
 	}
 
@@ -385,7 +387,7 @@ FSymslibResolver::EModuleStatus FSymslibResolver::LoadModule(FModuleEntry* Modul
 
 	const uint32 SymbolCount = syms_get_proc_count(ModuleInstance); // get symbol count
 	const uint32 LineCount = syms_get_line_count(ModuleInstance); // get line count
-	UE_LOG(LogSymslib, Display, TEXT("Loaded symbols for '%s', %u symbols and %u lines."), Module->Name, SymbolCount, LineCount);
+	UE_LOG(LogSymslib, Display, TEXT("Loaded symbols for '%s' at base 0x%016x, %u symbols and %u lines."), Module->Name, Module->Base, SymbolCount, LineCount);
 
 	return EModuleStatus::Loaded;
 }
@@ -407,7 +409,8 @@ bool FSymslibResolver::ResolveSymbol(uint64 Address, FResolvedSymbol* Target, FS
 	FModuleEntry* Module = GetModuleForAddress(Address);
 	if (!Module)
 	{
-		UpdateResolvedSymbol(Target, ESymbolQueryResult::NotFound, GUnknownModuleTextSymsLib, GUnknownModuleTextSymsLib);
+		UE_LOG(LogSymslib, Warning, TEXT("No module mapped to address 0x%016x."), Address);
+		UpdateResolvedSymbol(Target, ESymbolQueryResult::NotLoaded, GUnknownModuleTextSymsLib, GUnknownModuleTextSymsLib, GUnknownModuleTextSymsLib, 0);
 		return false;
 	}
 
@@ -421,10 +424,10 @@ bool FSymslibResolver::ResolveSymbol(uint64 Address, FResolvedSymbol* Target, FS
 	switch (Status)
 	{
 	case EModuleStatus::Failed:
-		UpdateResolvedSymbol(Target, ESymbolQueryResult::NotLoaded, Module->Name, GUnknownModuleTextSymsLib);
+		UpdateResolvedSymbol(Target, ESymbolQueryResult::NotLoaded, Module->Name, GUnknownModuleTextSymsLib, GUnknownModuleTextSymsLib, 0);
 		return false;
 	case EModuleStatus::VersionMismatch:
-		UpdateResolvedSymbol(Target, ESymbolQueryResult::Mismatch, Module->Name, GUnknownModuleTextSymsLib);
+		UpdateResolvedSymbol(Target, ESymbolQueryResult::Mismatch, Module->Name, GUnknownModuleTextSymsLib, GUnknownModuleTextSymsLib, 0);
 		return false;
 	default:
 		break;
@@ -436,7 +439,7 @@ bool FSymslibResolver::ResolveSymbol(uint64 Address, FResolvedSymbol* Target, FS
 	// Find procedure and source file for address
 	if (!syms_proc_from_va(Module->Instance, Address, &Proc) || !syms_va_to_src(Module->Instance, Address, &File))
 	{
-		UpdateResolvedSymbol(Target, ESymbolQueryResult::NotFound, Module->Name, GUnknownModuleTextSymsLib);
+		UpdateResolvedSymbol(Target, ESymbolQueryResult::NotFound, Module->Name, GUnknownModuleTextSymsLib, GUnknownModuleTextSymsLib, 0); 
 		return false;
 	}
 
@@ -454,18 +457,17 @@ bool FSymslibResolver::ResolveSymbol(uint64 Address, FResolvedSymbol* Target, FS
 	syms_read_strref(Module->Instance, &File.file.name, SourceFile, sizeof(SourceFile) - 1);
 	check(SourceFile[MaxStringSize - 1] == 0);
 
-	TStringBuilder<1024> FileAndLine;
-	FileAndLine << ANSI_TO_TCHAR(SourceFile) << TEXT("(") << File.line.ln << TEXT(")");
-
 	const TCHAR* SymbolNamePersistent =  StringAllocator.Store(ANSI_TO_TCHAR(SymbolName));
-	const TCHAR* FileAndLinePersistent = StringAllocator.Store(FileAndLine);
+	const TCHAR* SourceFilePersistent = StringAllocator.Store(ANSI_TO_TCHAR(SourceFile));
 
 	// Store the strings and update the target data
 	UpdateResolvedSymbol(
 		Target,
 		ESymbolQueryResult::OK,
+		Module->Name,
 		SymbolNamePersistent,
-		FileAndLinePersistent
+		SourceFilePersistent,
+		File.line.ln
 	);
 
 	return true;
diff --git a/Engine/Source/Developer/TraceServices/Private/Model/SymslibResolver.h b/Engine/Source/Developer/TraceServices/Private/Model/SymslibResolver.h
--- a/Engine/Source/Developer/TraceServices/Private/Model/SymslibResolver.h
+++ b/Engine/Source/Developer/TraceServices/Private/Model/SymslibResolver.h
@@ -76,7 +76,7 @@ private:
 	void DispatchQueuedAddresses();
 	void ResolveSymbols(TArrayView<FQueuedAddress>& QueuedWork);
 	FModuleEntry* GetModuleForAddress(uint64 Address);
-	static void UpdateResolvedSymbol(FResolvedSymbol* Symbol, ESymbolQueryResult Result, const TCHAR* Name, const TCHAR* FileAndLine);
+	static void UpdateResolvedSymbol(FResolvedSymbol* Symbol, ESymbolQueryResult Result, const TCHAR* Module, const TCHAR* Name, const TCHAR* File, uint16 Line);
 	void LoadModuleTracked(FModuleEntry* Module);
 	EModuleStatus LoadModule(FModuleEntry* Module);
 	void ResolveSymbolTracked(uint64 Address, FResolvedSymbol* Target, class FSymbolStringAllocator& StringAllocator);
diff --git a/Engine/Source/Developer/TraceServices/Public/TraceServices/Model/Modules.h b/Engine/Source/Developer/TraceServices/Public/TraceServices/Model/Modules.h
--- a/Engine/Source/Developer/TraceServices/Public/TraceServices/Model/Modules.h
+++ b/Engine/Source/Developer/TraceServices/Public/TraceServices/Model/Modules.h
@@ -32,7 +32,7 @@ inline const TCHAR* QueryResultToString(ESymbolQueryResult Result)
 		TEXT("Pending..."),
 		TEXT("Ok"),
 		TEXT("Not loaded"),
-		TEXT("Mismatch"),
+		TEXT("Version mismatch"),
 		TEXT("Not found")
 	};
 	static_assert(UE_ARRAY_COUNT(DisplayStrings) == (uint8) ESymbolQueryResult::StatusNum, "Missing QueryResult");
@@ -48,18 +48,22 @@ inline const TCHAR* QueryResultToString(ESymbolQueryResult Result)
 struct FResolvedSymbol
 {
 	std::atomic<ESymbolQueryResult> Result;
+	const TCHAR* Module;
 	const TCHAR* Name;
-	const TCHAR* FileAndLine;
-
+	const TCHAR* File;
+	uint16 Line;
+	
 	inline ESymbolQueryResult GetResult() const
 	{
 		return Result.load(std::memory_order_acquire);
 	}
 
-	FResolvedSymbol(ESymbolQueryResult InResult, const TCHAR* InName, const TCHAR* InFileAndLine) 
+	FResolvedSymbol(ESymbolQueryResult InResult, const TCHAR* InModule, const TCHAR* InName, const TCHAR* InFile, uint16 InLine) 
 		: Result(InResult)
+		, Module(InModule)
 		, Name(InName)
-		, FileAndLine(InFileAndLine)
+		, File(InFile)
+		, Line(InLine)
 	{}
 };
 
