[Insights] Stop analysis processing rather than closing data stream

When exiting Unreal Insights the data stream was closed in order for the analysis thread to finish and exit. Instead we trigger the stop event and wait for the analysis thread to exit. DataStreams are closed in the destructor.

#rnx

diff --git a/Engine/Source/Developer/TraceServices/Private/AnalysisService.cpp b/Engine/Source/Developer/TraceServices/Private/AnalysisService.cpp
--- a/Engine/Source/Developer/TraceServices/Private/AnalysisService.cpp
+++ b/Engine/Source/Developer/TraceServices/Private/AnalysisService.cpp
@@ -112,7 +112,7 @@ void FAnalysisSession::Start()
 
 void FAnalysisSession::Stop(bool bAndWait) const
 {
-	DataStream->Close();
+	Processor.Stop();
 	if (bAndWait)
 	{
 		Wait();
