[Insights] Make SymsLib cleanup task based

Move SymsLib cleanup to a task. In the destructor we need to wait for the task to complete. Add support to cancel tasks in order to handle cases where a lot of tasks are in flight on shutdown (for example when shutting down early in analysis).

#rnx

diff --git a/Engine/Source/Developer/TraceServices/Private/Model/SymslibResolver.cpp b/Engine/Source/Developer/TraceServices/Private/Model/SymslibResolver.cpp
--- a/Engine/Source/Developer/TraceServices/Private/Model/SymslibResolver.cpp
+++ b/Engine/Source/Developer/TraceServices/Private/Model/SymslibResolver.cpp
@@ -75,12 +75,16 @@ private:
 /////////////////////////////////////////////////////////////////////
 FSymslibResolver::FSymslibResolver(IAnalysisSession& InSession)
 	: Modules(InSession.GetLinearAllocator(), 128)
+	, CancelTasks(false)
 	, Session(InSession)
 {
 }
 
 FSymslibResolver::~FSymslibResolver()
 {
+	CancelTasks = true;
+	// Wait for cleanup task to finish
+	CleanupTask->Wait();
 }
 
 void FSymslibResolver::QueueModuleLoad(const FStringView& ModulePath, uint64 Base, uint32 Size, const uint8* ImageId, uint32 ImageIdSize)
@@ -95,6 +99,7 @@ void FSymslibResolver::QueueModuleLoad(const FStringView& ModulePath, uint64 Bas
 	Entry->Name = Session.StoreString(ModuleName); //Allows safe sharing to UI
 	Entry->Path = Session.StoreString(ModulePath);
 	Entry->Status = EModuleStatus::Pending;
+	Entry->Instance = nullptr;
 	Entry->ImageId = TArrayView<const uint8>(ImageId, ImageIdSize);
 
 	++ModulesDiscovered;
@@ -124,40 +129,53 @@ void FSymslibResolver::QueueSymbolResolve(uint64 Address, FResolvedSymbol* Symbo
 
 void FSymslibResolver::OnAnalysisComplete()
 {
-	// Dispatch any remaining requests.
+	CleanupTask = FFunctionGraphTask::CreateAndDispatchWhenReady([this]
 	{
-		FScopeLock _(&SymbolsQueueLock);
-		DispatchQueuedAddresses();
-		ResolveQueue.Reset();
-	}
+		// Dispatch any remaining requests.
+		{
+			FScopeLock _(&SymbolsQueueLock);
+			DispatchQueuedAddresses();
+			ResolveQueue.Reset();
+		}
 
-	// Wait for outstanding batches to complete.
-	uint32 OutstandingTasks = TasksInFlight.load(std::memory_order_acquire);
-	UE_LOG(LogSymslib, Display, TEXT("Waiting for %d outstanding tasks..."), OutstandingTasks);
-	do
-	{
-		OutstandingTasks = TasksInFlight.load(std::memory_order_acquire);
-		FPlatformProcess::Sleep(0.0);
-	} while (OutstandingTasks > 0);
+		// Wait for outstanding batches to complete.
+		uint32 OutstandingTasks = TasksInFlight.load(std::memory_order_acquire);
+		UE_LOG(LogSymslib, Display, TEXT("Waiting for %d outstanding tasks..."), OutstandingTasks);
+		do
+		{
+			OutstandingTasks = TasksInFlight.load(std::memory_order_acquire);
+			FPlatformProcess::Sleep(0.0);
+		}
+		while (OutstandingTasks > 0);
 
-	// No tasks can be in running at this point.
-	{
-		FReadScopeLock _(FileHandlesAndRegionsLock);
-		for (FMappedFileAndRegion& File : FileHandlesAndRegions)
+		// No tasks can be in running at this point.
 		{
-			// Release file region before releasing the file handle.
-			delete File.Region;
-			delete File.Handle;
+			FReadScopeLock _(FileHandlesAndRegionsLock);
+			for (FMappedFileAndRegion& File : FileHandlesAndRegions)
+			{
+				// Release file region before releasing the file handle.
+				delete File.Region;
+				delete File.Handle;
+			}
 		}
-	}
 
-	UE_LOG(LogSymslib, Display, TEXT("Allocated %.02f Mb of strings, %.02f Mb wasted."), SymbolBytesAllocated / float(1024*1024), SymbolBytesWasted / float(1024*1024));
-	// Release memory used by syms library
-	for (uint32 ModuleIndex = 0; ModuleIndex < Modules.Num(); ++ModuleIndex)
-	{
-		FModuleEntry& Module = Modules[ModuleIndex];
-		syms_quit(Module.Instance);
-	}
+		// Release memory used by syms library
+		{
+			FReadScopeLock _(ModulesLock);
+			for (uint32 ModuleIndex = 0; ModuleIndex < Modules.Num(); ++ModuleIndex)
+			{
+				FModuleEntry& Module = Modules[ModuleIndex];
+				if (Module.Instance)
+				{
+					syms_quit(Module.Instance);
+				}
+			}
+		}
+		
+		UE_LOG(LogSymslib, Display, TEXT("Allocated %.02f Mb of strings, %.02f Mb wasted."),
+		       SymbolBytesAllocated / float(1024*1024), SymbolBytesWasted / float(1024*1024));
+		
+	});
 }
 
 void FSymslibResolver::GetStats(IModuleProvider::FStats* OutStats) const
@@ -189,10 +207,15 @@ void FSymslibResolver::MaybeDispatchQueuedAddresses()
  */
 void FSymslibResolver::DispatchQueuedAddresses()
 {
+	if (ResolveQueue.IsEmpty())
+	{
+		return;
+	}
+	
 	TArray<FQueuedAddress> WorkingSet(ResolveQueue);
 	TasksInFlight += SymbolTasksInParallel;
 	const uint32 Stride = (WorkingSet.Num() + SymbolTasksInParallel - 1) / SymbolTasksInParallel;
-
+	
 	// Use background priority in order to not not interfere with Slate
 	ParallelFor(SymbolTasksInParallel, [this, &WorkingSet, Stride](uint32 Index) {
 		const uint32 StartIndex = Index * Stride;
@@ -210,6 +233,10 @@ void FSymslibResolver::ResolveSymbols(TArrayView<FQueuedAddress>& QueuedWork)
 	FSymbolStringAllocator StringAllocator(Session.GetLinearAllocator(), (2 << 12) / sizeof(TCHAR) );
 	for (const FQueuedAddress& ToResolve : QueuedWork)
 	{
+		if (CancelTasks.load(std::memory_order_relaxed))
+		{
+			break;
+		}
 		ResolveSymbolTracked(ToResolve.Address, ToResolve.Target, StringAllocator);
 	}
 	UE_LOG(LogSymslib, VeryVerbose, TEXT("String allocator used: %.02f kb, wasted: %.02f kb using %d blocks"),
@@ -242,6 +269,11 @@ void FSymslibResolver::UpdateResolvedSymbol(FResolvedSymbol* Symbol, ESymbolQuer
 
 void FSymslibResolver::LoadModuleTracked(FModuleEntry* Module)
 {
+	if (CancelTasks.load(std::memory_order_relaxed))
+	{
+		return;
+	}
+	
 	const EModuleStatus Status = LoadModule(Module);
 	if (Status == EModuleStatus::Loaded)
 	{
diff --git a/Engine/Source/Developer/TraceServices/Private/Model/SymslibResolver.h b/Engine/Source/Developer/TraceServices/Private/Model/SymslibResolver.h
--- a/Engine/Source/Developer/TraceServices/Private/Model/SymslibResolver.h
+++ b/Engine/Source/Developer/TraceServices/Private/Model/SymslibResolver.h
@@ -89,7 +89,9 @@ private:
 	FCriticalSection SymbolsQueueLock;
 	TArray<FQueuedAddress, TInlineAllocator<QueuedAddressLength>> ResolveQueue;
 	std::atomic<uint32> TasksInFlight;
-
+	FGraphEventRef CleanupTask;
+	std::atomic<bool> CancelTasks;
+	
 	struct FMappedFileAndRegion
 	{
 		IMappedFileHandle* Handle;
