Fix TaskProfiler crash

diff --git a/Engine/Source/Developer/TraceServices/Private/Model/TasksProfiler.cpp b/Engine/Source/Developer/TraceServices/Private/Model/TasksProfiler.cpp
--- a/Engine/Source/Developer/TraceServices/Private/Model/TasksProfiler.cpp
+++ b/Engine/Source/Developer/TraceServices/Private/Model/TasksProfiler.cpp
@@ -469,7 +469,7 @@ namespace TraceServices
 
 			// report all tasks whose execution overlaps [StartTime, EndTime]
 			const FTaskInfo* Task = TryGetTask(Thread[TaskIndex]);
-			while (Task != nullptr && Task->StartedTimestamp <= EndTime && Callback(*Task) != ETaskEnumerationResult::Stop)
+			while (Task != nullptr && Task->StartedTimestamp <= EndTime && Callback(*Task) != ETaskEnumerationResult::Stop && TaskIndex < Thread.Num() - 1)
 			{
 				++TaskIndex;
 				Task = TryGetTask(Thread[TaskIndex]);
