Unity build fix

#rnx

diff --git a/Engine/Source/Developer/TraceServices/Private/Model/SymslibResolver.h b/Engine/Source/Developer/TraceServices/Private/Model/SymslibResolver.h
--- a/Engine/Source/Developer/TraceServices/Private/Model/SymslibResolver.h
+++ b/Engine/Source/Developer/TraceServices/Private/Model/SymslibResolver.h
@@ -4,6 +4,7 @@
 #if PLATFORM_WINDOWS
 
 #include "Async/MappedFileHandle.h"
+#include "Async/TaskGraphInterfaces.h"
 #include "Common/PagedArray.h"
 #include "HAL/CriticalSection.h"
 #include "TraceServices/Model/AnalysisSession.h"
