Fix Editor hang when attempting to load a trace file while the analysis is reading from a trace stream where no events are received because all channels are disabled.

diff --git a/Engine/Source/Developer/TraceServices/Private/AnalysisService.cpp b/Engine/Source/Developer/TraceServices/Private/AnalysisService.cpp
--- a/Engine/Source/Developer/TraceServices/Private/AnalysisService.cpp
+++ b/Engine/Source/Developer/TraceServices/Private/AnalysisService.cpp
@@ -112,6 +112,7 @@ void FAnalysisSession::Start()
 
 void FAnalysisSession::Stop(bool bAndWait) const
 {
+	DataStream->Close();
 	Processor.Stop();
 	if (bAndWait)
 	{
