[Insights]
  - Added new trace events to track the ReOpenFile API. Updated the PlatformFileTraceAnalysis to process the respective new trace events. This fixes the issue where activities were displayed in Unreal Insights as I/O events with "Unknown" file (UE-114506).
  - Updated warning messages and debug code in PlatformFileTraceAnalysis.
  - Added handle tracking in Insights for I/O Activity events. File handle and read/write handles are now displayed in tooltip of an I/O Activity event.
  - Added "ReOpen" as a separate I/O activity type.

diff --git a/Engine/Source/Developer/TraceInsights/Private/Insights/ViewModels/FileActivityTimingTrack.cpp b/Engine/Source/Developer/TraceInsights/Private/Insights/ViewModels/FileActivityTimingTrack.cpp
--- a/Engine/Source/Developer/TraceInsights/Private/Insights/ViewModels/FileActivityTimingTrack.cpp
+++ b/Engine/Source/Developer/TraceInsights/Private/Insights/ViewModels/FileActivityTimingTrack.cpp
@@ -33,13 +33,15 @@
 const TCHAR* GetFileActivityTypeName(TraceServices::EFileActivityType Type)
 {
 	static_assert(TraceServices::FileActivityType_Open == 0, "TraceServices::EFileActivityType enum has changed!?");
-	static_assert(TraceServices::FileActivityType_Close == 1, "TraceServices::EFileActivityType enum has changed!?");
-	static_assert(TraceServices::FileActivityType_Read == 2, "TraceServices::EFileActivityType enum has changed!?");
-	static_assert(TraceServices::FileActivityType_Write == 3, "TraceServices::EFileActivityType enum has changed!?");
-	static_assert(TraceServices::FileActivityType_Count == 4, "TraceServices::EFileActivityType enum has changed!?");
+	static_assert(TraceServices::FileActivityType_ReOpen == 1, "TraceServices::EFileActivityType enum has changed!?");
+	static_assert(TraceServices::FileActivityType_Close == 2, "TraceServices::EFileActivityType enum has changed!?");
+	static_assert(TraceServices::FileActivityType_Read == 3, "TraceServices::EFileActivityType enum has changed!?");
+	static_assert(TraceServices::FileActivityType_Write == 4, "TraceServices::EFileActivityType enum has changed!?");
+	static_assert(TraceServices::FileActivityType_Count == 5, "TraceServices::EFileActivityType enum has changed!?");
 	static const TCHAR* GFileActivityTypeNames[] =
 	{
 		TEXT("Open"),
+		TEXT("ReOpen"),
 		TEXT("Close"),
 		TEXT("Read"),
 		TEXT("Write"),
@@ -56,6 +58,7 @@ uint32 GetFileActivityTypeColor(TraceServices::EFileActivityType Type)
 	static const uint32 GFileActivityTypeColors[] =
 	{
 		0xFFCCAA33, // open
+		0xFFBB9922, // reopen
 		0xFF33AACC, // close
 		0xFF33AA33, // read
 		0xFFDD33CC, // write
@@ -204,7 +207,7 @@ void FFileActivitySharedState::Tick(Insights::ITimingViewSession& InSession, con
 					uint32 LocalDepth = MAX_uint32;
 					for (int32 i = 0; i < ConcurrentEvents.Num(); ++i)
 					{
-						if (EventStartTime > ConcurrentEvents[i])
+						if (EventStartTime >= ConcurrentEvents[i])
 						{
 							LocalDepth = i;
 							ConcurrentEvents[i] = EventEndTime;
@@ -220,7 +223,7 @@ void FFileActivitySharedState::Tick(Insights::ITimingViewSession& InSession, con
 					}
 
 					uint32 Type = ((uint32)FileActivity->ActivityType & 0x0F) | (FileActivity->Failed ? 0x80 : 0);
-					AllIoEvents.Add(FIoTimingEvent{ EventStartTime, EventEndTime, LocalDepth, Type, FileActivity->Offset, FileActivity->Size, FileActivity->ActualSize, ActivityIndex });
+					AllIoEvents.Add(FIoTimingEvent{ EventStartTime, EventEndTime, LocalDepth, Type, FileActivity->Offset, FileActivity->Size, FileActivity->ActualSize, ActivityIndex, FileActivity->FileHandle, FileActivity->ReadWriteHandle });
 					return TraceServices::EEventEnumerate::Continue;
 				});
 
@@ -609,7 +612,26 @@ void FFileActivityTimingTrack::InitTooltip(FTooltipDrawState& InOutTooltip, cons
 			TypeLinearColor.B *= 2.0f;
 			InOutTooltip.AddTitle(TypeStr, TypeLinearColor);
 
-			InOutTooltip.AddTitle(SharedState.FileActivities[InEvent.FileActivityIndex]->Path);
+			if (ensure(InEvent.FileActivityIndex >= 0 && InEvent.FileActivityIndex < SharedState.FileActivities.Num()))
+			{
+				const TSharedPtr<FFileActivitySharedState::FIoFileActivity>& ActivityPtr = SharedState.FileActivities[InEvent.FileActivityIndex];
+				check(ActivityPtr.IsValid());
+				FFileActivitySharedState::FIoFileActivity& Activity = *ActivityPtr;
+
+				InOutTooltip.AddTitle(Activity.Path);
+			}
+
+			if (InEvent.FileHandle != uint64(-1))
+			{
+				const FString Value = FString::Printf(TEXT("0x%X"), InEvent.FileHandle);
+				InOutTooltip.AddNameValueTextLine(TEXT("File Handle:"), Value);
+			}
+
+			if (InEvent.ReadWriteHandle != uint64(-1))
+			{
+				const FString Value = FString::Printf(TEXT("0x%X"), InEvent.ReadWriteHandle);
+				InOutTooltip.AddNameValueTextLine(TEXT("Read/Write Handle:"), Value);
+			}
 
 			const double Duration = InEvent.EndTime - InEvent.StartTime;
 			InOutTooltip.AddNameValueTextLine(TEXT("Duration:"), TimeUtils::FormatTimeAuto(Duration));
@@ -646,48 +668,68 @@ bool FFileActivityTimingTrack::FindIoTimingEvent(const FTimingEventSearchParamet
 		// Search...
 		[this](TTimingEventSearch<FFileActivitySharedState::FIoTimingEvent>::FContext& InContext)
 		{
+			const TArray<FFileActivitySharedState::FIoTimingEvent>& Events = SharedState.GetAllEvents();
+
 			if (bIgnoreDuration)
 			{
-				for (const FFileActivitySharedState::FIoTimingEvent& Event : SharedState.GetAllEvents())
+				// Events are sorted by start time.
+				// Find the first event with StartTime >= searched StartTime.
+				int32 StartIndex = Algo::LowerBoundBy(Events, InContext.GetParameters().StartTime,
+					[](const FFileActivitySharedState::FIoTimingEvent& Event) { return Event.StartTime; });
+
+				for (int32 Index = StartIndex; Index < Events.Num(); ++Index)
 				{
+					const FFileActivitySharedState::FIoTimingEvent& Event = Events[Index];
+
 					if (bShowOnlyErrors && ((Event.Type & 0xF0) == 0))
 					{
 						continue;
 					}
 
-					if (Event.StartTime < InContext.GetParameters().StartTime)
-					{
-						continue;
-					}
+					ensure(Event.StartTime >= InContext.GetParameters().StartTime);
 
-					if (!InContext.ShouldContinueSearching() || Event.StartTime > InContext.GetParameters().EndTime)
+					if (Event.StartTime > InContext.GetParameters().EndTime)
 					{
 						break;
 					}
 
 					InContext.Check(Event.StartTime, Event.StartTime, bIgnoreEventDepth ? 0 : Event.Depth, Event);
+
+					if (!InContext.ShouldContinueSearching())
+					{
+						break;
+					}
 				}
 			}
 			else
 			{
-				for (const FFileActivitySharedState::FIoTimingEvent& Event : SharedState.GetAllEvents())
+				// Events are sorted by start time.
+				// Find the first event with StartTime >= searched EndTime.
+				int32 StartIndex = Algo::LowerBoundBy(Events, InContext.GetParameters().EndTime,
+					[](const FFileActivitySharedState::FIoTimingEvent& Event) { return Event.StartTime; });
+
+				// Start at the last event with StartTime < searched EndTime.
+				for (int32 Index = StartIndex - 1; Index >= 0; --Index)
 				{
+					const FFileActivitySharedState::FIoTimingEvent& Event = Events[Index];
+
 					if (bShowOnlyErrors && ((Event.Type & 0xF0) == 0))
 					{
 						continue;
 					}
 
-					if (!bIgnoreDuration && Event.EndTime <= InContext.GetParameters().StartTime)
+					if (Event.EndTime <= InContext.GetParameters().StartTime ||
+						Event.StartTime >= InContext.GetParameters().EndTime)
 					{
 						continue;
 					}
 
-					if (!InContext.ShouldContinueSearching() || Event.StartTime >= InContext.GetParameters().EndTime)
+					InContext.Check(Event.StartTime, Event.EndTime, bIgnoreEventDepth ? 0 : Event.Depth, Event);
+
+					if (!InContext.ShouldContinueSearching())
 					{
 						break;
 					}
-
-					InContext.Check(Event.StartTime, Event.EndTime, bIgnoreEventDepth ? 0 : Event.Depth, Event);
 				}
 			}
 		},
diff --git a/Engine/Source/Developer/TraceInsights/Private/Insights/ViewModels/FileActivityTimingTrack.h b/Engine/Source/Developer/TraceInsights/Private/Insights/ViewModels/FileActivityTimingTrack.h
--- a/Engine/Source/Developer/TraceInsights/Private/Insights/ViewModels/FileActivityTimingTrack.h
+++ b/Engine/Source/Developer/TraceInsights/Private/Insights/ViewModels/FileActivityTimingTrack.h
@@ -47,6 +47,8 @@ public:
 		uint64 Size;
 		uint64 ActualSize;
 		int32 FileActivityIndex;
+		uint64 FileHandle; // file handle
+		uint64 ReadWriteHandle; // for Read/Write operations
 	};
 
 public:
diff --git a/Engine/Source/Developer/TraceServices/Private/Analyzers/PlatformFileTraceAnalysis.cpp b/Engine/Source/Developer/TraceServices/Private/Analyzers/PlatformFileTraceAnalysis.cpp
--- a/Engine/Source/Developer/TraceServices/Private/Analyzers/PlatformFileTraceAnalysis.cpp
+++ b/Engine/Source/Developer/TraceServices/Private/Analyzers/PlatformFileTraceAnalysis.cpp
@@ -28,6 +28,8 @@ void FPlatformFileTraceAnalyzer::OnAnalysisBegin(const FOnAnalysisContext& Conte
 
 	Builder.RouteEvent(RouteId_BeginOpen, "PlatformFile", "BeginOpen");
 	Builder.RouteEvent(RouteId_EndOpen, "PlatformFile", "EndOpen");
+	Builder.RouteEvent(RouteId_BeginReOpen, "PlatformFile", "BeginReOpen");
+	Builder.RouteEvent(RouteId_EndReOpen, "PlatformFile", "EndReOpen");
 	Builder.RouteEvent(RouteId_BeginClose, "PlatformFile", "BeginClose");
 	Builder.RouteEvent(RouteId_EndClose, "PlatformFile", "EndClose");
 	Builder.RouteEvent(RouteId_BeginRead, "PlatformFile", "BeginRead");
@@ -45,13 +47,14 @@ bool FPlatformFileTraceAnalyzer::OnEvent(uint16 RouteId, EStyle Style, const FOn
 	{
 	case RouteId_BeginOpen:
 	{
-		double Time = Context.EventTime.AsSeconds(EventData.GetValue<uint64>("Cycle"));
+		uint64 Cycle = EventData.GetValue<uint64>("Cycle");
+		double Time = Context.EventTime.AsSeconds(Cycle);
 		Session.UpdateDurationSeconds(Time);
 		uint32 ThreadId = FTraceAnalyzerUtils::GetThreadIdField(Context);
 		#if DEBUG_PLATFORMFILETRACE
 		if (PendingOpenMap.Contains(ThreadId))
 		{
-			PLATFORMFILETRACE_WARNING("Duplicated BeginOpen event!?");
+			PLATFORMFILETRACE_WARNING("BeginOpen: Duplicated event!?");
 		}
 		#endif
 		const TCHAR* FileName = reinterpret_cast<const TCHAR*>(EventData.GetAttachment());
@@ -63,59 +66,134 @@ bool FPlatformFileTraceAnalyzer::OnEvent(uint16 RouteId, EStyle Style, const FOn
 	}
 	case RouteId_EndOpen:
 	{
-		double Time = Context.EventTime.AsSeconds(EventData.GetValue<uint64>("Cycle"));
+		uint64 Cycle = EventData.GetValue<uint64>("Cycle");
+		double Time = Context.EventTime.AsSeconds(Cycle);
 		Session.UpdateDurationSeconds(Time);
 		uint64 FileHandle = EventData.GetValue<uint64>("FileHandle");
 		uint32 ThreadId = FTraceAnalyzerUtils::GetThreadIdField(Context);
 		const FPendingActivity* FindOpen = PendingOpenMap.Find(ThreadId);
 		if (FindOpen)
 		{
-			if (FileHandle == uint64(-1)) // TODO: Portable invalid file handle?
+			if (FileHandle == uint64(-1)) // invalid file handle
 			{
 				FileActivityProvider.EndActivity(FindOpen->FileIndex, FindOpen->ActivityIndex, 0, Time, true);
 			}
 			else
 			{
+				#if DEBUG_PLATFORMFILETRACE
+				if (OpenFilesMap.Contains(FileHandle))
+				{
+					PLATFORMFILETRACE_WARNING("EndOpen: File already open!?");
+				}
+				#endif
 				OpenFilesMap.Add(FileHandle, FindOpen->FileIndex);
 				FileActivityProvider.EndActivity(FindOpen->FileIndex, FindOpen->ActivityIndex, 0, Time, false);
 			}
+			FileActivityProvider.SetActivityFileHandle(FindOpen->FileIndex, FindOpen->ActivityIndex, FileHandle);
 			PendingOpenMap.Remove(ThreadId);
 		}
 		else
 		{
-			PLATFORMFILETRACE_WARNING("BeginOpen event not traced!?");
+			PLATFORMFILETRACE_WARNING("EndOpen: BeginOpen event not traced!?");
+		}
+		break;
+	}
+	case RouteId_BeginReOpen:
+	{
+		uint64 Cycle = EventData.GetValue<uint64>("Cycle");
+		double Time = Context.EventTime.AsSeconds(Cycle);
+		Session.UpdateDurationSeconds(Time);
+		uint64 FileHandle = EventData.GetValue<uint64>("OldFileHandle");
+		uint32 ThreadId = FTraceAnalyzerUtils::GetThreadIdField(Context);
+		#if DEBUG_PLATFORMFILETRACE
+		if (PendingReOpenMap.Contains(ThreadId))
+		{
+			PLATFORMFILETRACE_WARNING("BeginReOpen: Duplicated event!?");
+		}
+		#endif
+		const uint32* FindFileIndex = OpenFilesMap.Find(FileHandle);
+		if (!FindFileIndex)
+		{
+			PLATFORMFILETRACE_WARNING("BeginReOpen: File is not open!?");
+		}
+		else
+		{
+			uint32 FileIndex = *FindFileIndex;
+			FPendingActivity& Open = PendingReOpenMap.Add(ThreadId);
+			Open.ActivityIndex = FileActivityProvider.BeginActivity(FileIndex, FileActivityType_ReOpen, ThreadId, 0, 0, Time);
+			Open.FileIndex = FileIndex;
+		}
+		break;
+	}
+	case RouteId_EndReOpen:
+	{
+		uint64 Cycle = EventData.GetValue<uint64>("Cycle");
+		double Time = Context.EventTime.AsSeconds(Cycle);
+		Session.UpdateDurationSeconds(Time);
+		uint64 NewFileHandle = EventData.GetValue<uint64>("NewFileHandle");
+		uint32 ThreadId = FTraceAnalyzerUtils::GetThreadIdField(Context);
+		const FPendingActivity* FindOpen = PendingReOpenMap.Find(ThreadId);
+		if (FindOpen)
+		{
+			if (NewFileHandle == uint64(-1)) // invalid file handle
+			{
+				FileActivityProvider.EndActivity(FindOpen->FileIndex, FindOpen->ActivityIndex, 0, Time, true);
+			}
+			else
+			{
+				#if DEBUG_PLATFORMFILETRACE
+				if (OpenFilesMap.Contains(NewFileHandle))
+				{
+					PLATFORMFILETRACE_WARNING("EndReOpen: File already open!?");
+				}
+				#endif
+				OpenFilesMap.Add(NewFileHandle, FindOpen->FileIndex);
+				FileActivityProvider.EndActivity(FindOpen->FileIndex, FindOpen->ActivityIndex, 0, Time, false);
+			}
+			FileActivityProvider.SetActivityFileHandle(FindOpen->FileIndex, FindOpen->ActivityIndex, NewFileHandle);
+			PendingReOpenMap.Remove(ThreadId);
+		}
+		else
+		{
+			PLATFORMFILETRACE_WARNING("EndReOpen: BeginReOpen event not traced!?");
 		}
 		break;
 	}
 	case RouteId_BeginClose:
 	{
-		double Time = Context.EventTime.AsSeconds(EventData.GetValue<uint64>("Cycle"));
+		uint64 Cycle = EventData.GetValue<uint64>("Cycle");
+		double Time = Context.EventTime.AsSeconds(Cycle);
 		Session.UpdateDurationSeconds(Time);
 		uint64 FileHandle = EventData.GetValue<uint64>("FileHandle");
 		uint32 ThreadId = FTraceAnalyzerUtils::GetThreadIdField(Context);
 		#if DEBUG_PLATFORMFILETRACE
 		if (PendingCloseMap.Contains(ThreadId))
 		{
-			PLATFORMFILETRACE_WARNING("Duplicated BeginClose event!?");
+			PLATFORMFILETRACE_WARNING("BeginClose: Duplicated event!?");
 		}
 		#endif
 		const uint32* FindFileIndex = OpenFilesMap.Find(FileHandle);
+		uint32 FileIndex;
 		if (FindFileIndex)
 		{
+			FileIndex = *FindFileIndex;
 			OpenFilesMap.Remove(FileHandle);
-			FPendingActivity& Close = PendingCloseMap.Add(ThreadId);
-			Close.ActivityIndex = FileActivityProvider.BeginActivity(*FindFileIndex, FileActivityType_Close, ThreadId, 0, 0, Time);
-			Close.FileIndex = *FindFileIndex;
 		}
 		else
 		{
-			PLATFORMFILETRACE_WARNING("File is not open!?");
+			PLATFORMFILETRACE_WARNING("BeginClose: File is not open!?");
+			FileIndex = FileActivityProvider.GetUnknownFileIndex();
 		}
+		FPendingActivity& Close = PendingCloseMap.Add(ThreadId);
+		Close.ActivityIndex = FileActivityProvider.BeginActivity(FileIndex, FileActivityType_Close, ThreadId, 0, 0, Time);
+		FileActivityProvider.SetActivityFileHandle(FileIndex, Close.ActivityIndex, FileHandle);
+		Close.FileIndex = FileIndex;
 		break;
 	}
 	case RouteId_EndClose:
 	{
-		double Time = Context.EventTime.AsSeconds(EventData.GetValue<uint64>("Cycle"));
+		uint64 Cycle = EventData.GetValue<uint64>("Cycle");
+		double Time = Context.EventTime.AsSeconds(Cycle);
 		Session.UpdateDurationSeconds(Time);
 		uint32 ThreadId = FTraceAnalyzerUtils::GetThreadIdField(Context);
 		const FPendingActivity* FindClose = PendingCloseMap.Find(ThreadId);
@@ -126,13 +204,14 @@ bool FPlatformFileTraceAnalyzer::OnEvent(uint16 RouteId, EStyle Style, const FOn
 		}
 		else
 		{
-			PLATFORMFILETRACE_WARNING("BeginClose event not traced!?");
+			PLATFORMFILETRACE_WARNING("EndClose: BeginClose event not traced!?");
 		}
 		break;
 	}
 	case RouteId_BeginRead:
 	{
-		double Time = Context.EventTime.AsSeconds(EventData.GetValue<uint64>("Cycle"));
+		uint64 Cycle = EventData.GetValue<uint64>("Cycle");
+		double Time = Context.EventTime.AsSeconds(Cycle);
 		Session.UpdateDurationSeconds(Time);
 		uint64 ReadHandle = EventData.GetValue<uint64>("ReadHandle");
 		uint64 FileHandle = EventData.GetValue<uint64>("FileHandle");
@@ -147,10 +226,19 @@ bool FPlatformFileTraceAnalyzer::OnEvent(uint16 RouteId, EStyle Style, const FOn
 		}
 		else
 		{
+			PLATFORMFILETRACE_WARNING("BeginRead: File is not open!?");
 			FileIndex = FileActivityProvider.GetUnknownFileIndex();
 			OpenFilesMap.Add(FileHandle, FileIndex);
 		}
 		uint64 ReadIndex = FileActivityProvider.BeginActivity(FileIndex, FileActivityType_Read, ThreadId, Offset, Size, Time);
+		FileActivityProvider.SetActivityFileHandle(FileIndex, ReadIndex, FileHandle);
+		FileActivityProvider.SetActivityReadWriteHandle(FileIndex, ReadIndex, ReadHandle);
+		#if DEBUG_PLATFORMFILETRACE
+		if (ActiveReadsMap.Contains(ReadHandle))
+		{
+			PLATFORMFILETRACE_WARNING("BeginRead: Duplicated ReadHandle!?");
+		}
+		#endif
 		FPendingActivity& Read = ActiveReadsMap.Add(ReadHandle);
 		Read.FileIndex = FileIndex;
 		Read.ActivityIndex = ReadIndex;
@@ -158,7 +246,8 @@ bool FPlatformFileTraceAnalyzer::OnEvent(uint16 RouteId, EStyle Style, const FOn
 	}
 	case RouteId_EndRead:
 	{
-		double Time = Context.EventTime.AsSeconds(EventData.GetValue<uint64>("Cycle"));
+		uint64 Cycle = EventData.GetValue<uint64>("Cycle");
+		double Time = Context.EventTime.AsSeconds(Cycle);
 		Session.UpdateDurationSeconds(Time);
 		uint64 ReadHandle = EventData.GetValue<uint64>("ReadHandle");
 		uint64 SizeRead = EventData.GetValue<uint64>("SizeRead");
@@ -167,17 +256,19 @@ bool FPlatformFileTraceAnalyzer::OnEvent(uint16 RouteId, EStyle Style, const FOn
 		if (FindRead)
 		{
 			FileActivityProvider.EndActivity(FindRead->FileIndex, FindRead->ActivityIndex, SizeRead, Time, false);
+			FileActivityProvider.CheckActivityReadWriteHandle(FindRead->FileIndex, FindRead->ActivityIndex, ReadHandle);
 			ActiveReadsMap.Remove(ReadHandle);
 		}
 		else
 		{
-			PLATFORMFILETRACE_WARNING("BeginRead event not traced!?");
+			PLATFORMFILETRACE_WARNING("EndRead: BeginRead event not traced!?");
 		}
 		break;
 	}
 	case RouteId_BeginWrite:
 	{
-		double Time = Context.EventTime.AsSeconds(EventData.GetValue<uint64>("Cycle"));
+		uint64 Cycle = EventData.GetValue<uint64>("Cycle");
+		double Time = Context.EventTime.AsSeconds(Cycle);
 		Session.UpdateDurationSeconds(Time);
 		uint64 WriteHandle = EventData.GetValue<uint64>("WriteHandle");
 		uint64 FileHandle = EventData.GetValue<uint64>("FileHandle");
@@ -192,10 +283,19 @@ bool FPlatformFileTraceAnalyzer::OnEvent(uint16 RouteId, EStyle Style, const FOn
 		}
 		else
 		{
+			PLATFORMFILETRACE_WARNING("BeginWrite: File is not open!?");
 			FileIndex = FileActivityProvider.GetUnknownFileIndex();
 			OpenFilesMap.Add(FileHandle, FileIndex);
 		}
 		uint64 WriteIndex = FileActivityProvider.BeginActivity(FileIndex, FileActivityType_Write, ThreadId, Offset, Size, Time);
+		FileActivityProvider.SetActivityFileHandle(FileIndex, WriteIndex, FileHandle);
+		FileActivityProvider.SetActivityReadWriteHandle(FileIndex, WriteIndex, WriteHandle);
+		#if DEBUG_PLATFORMFILETRACE
+		if (ActiveWritesMap.Contains(WriteHandle))
+		{
+			PLATFORMFILETRACE_WARNING("BeginWrite: Duplicated WriteHandle!?");
+		}
+		#endif
 		FPendingActivity& Write = ActiveWritesMap.Add(WriteHandle);
 		Write.FileIndex = FileIndex;
 		Write.ActivityIndex = WriteIndex;
@@ -203,7 +303,8 @@ bool FPlatformFileTraceAnalyzer::OnEvent(uint16 RouteId, EStyle Style, const FOn
 	}
 	case RouteId_EndWrite:
 	{
-		double Time = Context.EventTime.AsSeconds(EventData.GetValue<uint64>("Cycle"));
+		uint64 Cycle = EventData.GetValue<uint64>("Cycle");
+		double Time = Context.EventTime.AsSeconds(Cycle);
 		Session.UpdateDurationSeconds(Time);
 		uint64 WriteHandle = EventData.GetValue<uint64>("WriteHandle");
 		uint64 SizeWritten = EventData.GetValue<uint64>("SizeWritten");
@@ -212,11 +313,12 @@ bool FPlatformFileTraceAnalyzer::OnEvent(uint16 RouteId, EStyle Style, const FOn
 		if (FindWrite)
 		{
 			FileActivityProvider.EndActivity(FindWrite->FileIndex, FindWrite->ActivityIndex, SizeWritten, Time, false);
+			FileActivityProvider.CheckActivityReadWriteHandle(FindWrite->FileIndex, FindWrite->ActivityIndex, WriteHandle);
 			ActiveWritesMap.Remove(WriteHandle);
 		}
 		else
 		{
-			PLATFORMFILETRACE_WARNING("BeginWrite event not traced!?");
+			PLATFORMFILETRACE_WARNING("EndWrite: BeginWrite event not traced!?");
 		}
 		break;
 	}
diff --git a/Engine/Source/Developer/TraceServices/Private/Analyzers/PlatformFileTraceAnalysis.h b/Engine/Source/Developer/TraceServices/Private/Analyzers/PlatformFileTraceAnalysis.h
--- a/Engine/Source/Developer/TraceServices/Private/Analyzers/PlatformFileTraceAnalysis.h
+++ b/Engine/Source/Developer/TraceServices/Private/Analyzers/PlatformFileTraceAnalysis.h
@@ -25,6 +25,8 @@ private:
 	{
 		RouteId_BeginOpen,
 		RouteId_EndOpen,
+		RouteId_BeginReOpen,
+		RouteId_EndReOpen,
 		RouteId_BeginClose,
 		RouteId_EndClose,
 		RouteId_BeginRead,
@@ -46,6 +48,7 @@ private:
 	FFileActivityProvider& FileActivityProvider;
 	TMap<uint64, uint32> OpenFilesMap;
 	TMap<uint32, FPendingActivity> PendingOpenMap;
+	TMap<uint32, FPendingActivity> PendingReOpenMap;
 	TMap<uint32, FPendingActivity> PendingCloseMap;
 	TMap<uint64, FPendingActivity> ActiveReadsMap;
 	TMap<uint64, FPendingActivity> ActiveWritesMap;
diff --git a/Engine/Source/Developer/TraceServices/Private/Model/FileActivity.cpp b/Engine/Source/Developer/TraceServices/Private/Model/FileActivity.cpp
--- a/Engine/Source/Developer/TraceServices/Private/Model/FileActivity.cpp
+++ b/Engine/Source/Developer/TraceServices/Private/Model/FileActivity.cpp
@@ -79,6 +79,8 @@ uint64 FFileActivityProvider::BeginActivity(uint32 FileIndex, EFileActivityType
 	FileActivity.ThreadId = ThreadId;
 	FileActivity.Failed = false;
 	FileActivity.ActivityType = Type;
+	FileActivity.FileHandle = uint64(-1);
+	FileActivity.ReadWriteHandle = uint64(-1);
 	return FileInfo.ActivityTimeline->AppendBeginEvent(Time, &FileActivity);
 }
 
@@ -91,6 +93,27 @@ void FFileActivityProvider::EndActivity(uint32 FileIndex, uint64 ActivityIndex,
 	Activity->Failed = Failed;
 }
 
+void FFileActivityProvider::SetActivityFileHandle(uint32 FileIndex, uint64 ActivityIndex, uint64 FileHandle)
+{
+	FFileInfoInternal& FileInfo = Files[FileIndex];
+	FFileActivity* Activity = FileInfo.ActivityTimeline->GetEvent(ActivityIndex);
+	Activity->FileHandle = FileHandle;
+}
+
+void FFileActivityProvider::SetActivityReadWriteHandle(uint32 FileIndex, uint64 ActivityIndex, uint64 ReadWriteHandle)
+{
+	FFileInfoInternal& FileInfo = Files[FileIndex];
+	FFileActivity* Activity = FileInfo.ActivityTimeline->GetEvent(ActivityIndex);
+	Activity->ReadWriteHandle = ReadWriteHandle;
+}
+
+void FFileActivityProvider::CheckActivityReadWriteHandle(uint32 FileIndex, uint64 ActivityIndex, uint64 ReadWriteHandle)
+{
+	FFileInfoInternal& FileInfo = Files[FileIndex];
+	FFileActivity* Activity = FileInfo.ActivityTimeline->GetEvent(ActivityIndex);
+	check(Activity->ReadWriteHandle == ReadWriteHandle);
+}
+
 const TCHAR* FFileActivityProvider::GetFilePath(uint32 FileIndex) const
 {
 	return Files[FileIndex].FileInfo.Path;
diff --git a/Engine/Source/Developer/TraceServices/Private/Model/FileActivity.h b/Engine/Source/Developer/TraceServices/Private/Model/FileActivity.h
--- a/Engine/Source/Developer/TraceServices/Private/Model/FileActivity.h
+++ b/Engine/Source/Developer/TraceServices/Private/Model/FileActivity.h
@@ -35,6 +35,9 @@ public:
 	uint32 GetUnknownFileIndex();
 	uint64 BeginActivity(uint32 FileIndex, EFileActivityType Type, uint32 ThreadId, uint64 Offset, uint64 Size, double Time);
 	void EndActivity(uint32 FileIndex, uint64 ActivityIndex, uint64 ActualSize, double Time, bool Failed);
+	void SetActivityFileHandle(uint32 FileIndex, uint64 ActivityIndex, uint64 FileHandle);
+	void SetActivityReadWriteHandle(uint32 FileIndex, uint64 ActivityIndex, uint64 ReadWriteHandle);
+	void CheckActivityReadWriteHandle(uint32 FileIndex, uint64 ActivityIndex, uint64 ReadWriteHandle);
 	const TCHAR* GetFilePath(uint32 FileIndex) const;
 
 private:
diff --git a/Engine/Source/Developer/TraceServices/Public/TraceServices/Model/LoadTimeProfiler.h b/Engine/Source/Developer/TraceServices/Public/TraceServices/Model/LoadTimeProfiler.h
--- a/Engine/Source/Developer/TraceServices/Public/TraceServices/Model/LoadTimeProfiler.h
+++ b/Engine/Source/Developer/TraceServices/Public/TraceServices/Model/LoadTimeProfiler.h
@@ -20,6 +20,7 @@ struct FFileInfo
 enum EFileActivityType
 {
 	FileActivityType_Open,
+	FileActivityType_ReOpen,
 	FileActivityType_Close,
 	FileActivityType_Read,
 	FileActivityType_Write,
@@ -41,6 +42,8 @@ struct FFileActivity
 	uint32 ThreadId = 0;
 	EFileActivityType ActivityType = FileActivityType_Invalid;
 	bool Failed = false;
+	uint64 FileHandle;
+	uint64 ReadWriteHandle;
 };
 
 class IFileActivityProvider
diff --git a/Engine/Source/Runtime/Core/Private/ProfilingDebugging/PlatformFileTrace.cpp b/Engine/Source/Runtime/Core/Private/ProfilingDebugging/PlatformFileTrace.cpp
--- a/Engine/Source/Runtime/Core/Private/ProfilingDebugging/PlatformFileTrace.cpp
+++ b/Engine/Source/Runtime/Core/Private/ProfilingDebugging/PlatformFileTrace.cpp
@@ -27,6 +27,16 @@ UE_TRACE_EVENT_BEGIN(PlatformFile, EndOpen)
 	UE_TRACE_EVENT_FIELD(uint64, FileHandle)
 UE_TRACE_EVENT_END()
 
+UE_TRACE_EVENT_BEGIN(PlatformFile, BeginReOpen)
+	UE_TRACE_EVENT_FIELD(uint64, Cycle)
+	UE_TRACE_EVENT_FIELD(uint64, OldFileHandle)
+UE_TRACE_EVENT_END()
+
+UE_TRACE_EVENT_BEGIN(PlatformFile, EndReOpen)
+	UE_TRACE_EVENT_FIELD(uint64, Cycle)
+	UE_TRACE_EVENT_FIELD(uint64, NewFileHandle)
+UE_TRACE_EVENT_END()
+
 UE_TRACE_EVENT_BEGIN(PlatformFile, BeginClose)
 	UE_TRACE_EVENT_FIELD(uint64, Cycle)
 	UE_TRACE_EVENT_FIELD(uint64, FileHandle)
@@ -99,12 +109,33 @@ void FPlatformFileTrace::EndOpen(uint64 FileHandle)
 
 void FPlatformFileTrace::FailOpen(const TCHAR* Path)
 {
-	// TODO: Separate event for failure in the trace?
 	UE_TRACE_LOG(PlatformFile, EndOpen, FileChannel)
 		<< EndOpen.Cycle(FPlatformTime::Cycles64())
 		<< EndOpen.FileHandle(uint64(-1));
 }
 
+void FPlatformFileTrace::BeginReOpen(uint64 OldFileHandle)
+{
+	UE_TRACE_LOG(PlatformFile, BeginReOpen, FileChannel)
+		<< BeginReOpen.Cycle(FPlatformTime::Cycles64())
+		<< BeginReOpen.OldFileHandle(OldFileHandle);
+}
+
+void FPlatformFileTrace::EndReOpen(uint64 NewFileHandle)
+{
+	UE_TRACE_LOG(PlatformFile, EndReOpen, FileChannel)
+		<< EndReOpen.Cycle(FPlatformTime::Cycles64())
+		<< EndReOpen.NewFileHandle(NewFileHandle);
+#if PLATFORMFILETRACE_DEBUG_ENABLED
+	{
+		FScopeLock OpenHandlesScopeLock(&OpenHandlesLock);
+		++OpenHandles.FindOrAdd(NewFileHandle, 0);
+	}
+#else
+	++OpenFileHandleCount;
+#endif
+}
+
 void FPlatformFileTrace::BeginClose(uint64 FileHandle)
 {
 	UE_TRACE_LOG(PlatformFile, BeginClose, FileChannel)
@@ -144,7 +175,6 @@ void FPlatformFileTrace::EndClose(uint64 FileHandle)
 
 void FPlatformFileTrace::FailClose(uint64 FileHandle)
 {
-	// TODO: Separate event for failure in the trace?
 	UE_TRACE_LOG(PlatformFile, EndClose, FileChannel)
 		<< EndClose.Cycle(FPlatformTime::Cycles64());
 }
diff --git a/Engine/Source/Runtime/Core/Private/Windows/WindowsPlatformFile.cpp b/Engine/Source/Runtime/Core/Private/Windows/WindowsPlatformFile.cpp
--- a/Engine/Source/Runtime/Core/Private/Windows/WindowsPlatformFile.cpp
+++ b/Engine/Source/Runtime/Core/Private/Windows/WindowsPlatformFile.cpp
@@ -498,8 +498,24 @@ public:
 	{
 		if (IsValid())
 		{
+			TRACE_PLATFORMFILE_BEGIN_REOPEN(Handle);
 			HANDLE NewFileHandle = ReOpenFile(Handle, DesiredAccess, ShareMode, Flags);
-			CloseHandle(Handle);
+			TRACE_PLATFORMFILE_END_REOPEN(NewFileHandle);
+
+			TRACE_PLATFORMFILE_BEGIN_CLOSE(Handle);
+			BOOL CloseResult = CloseHandle(Handle);
+#if PLATFORMFILETRACE_ENABLED
+			if (CloseResult)
+			{
+				TRACE_PLATFORMFILE_END_CLOSE(Handle);
+			}
+			else
+			{
+				TRACE_PLATFORMFILE_FAIL_CLOSE(Handle);
+			}
+#else
+			(void)CloseResult;
+#endif
 			Handle = NewFileHandle;
 		}
 	}
@@ -736,8 +752,24 @@ public:
 	{
 		if (IsValid())
 		{
+			TRACE_PLATFORMFILE_BEGIN_REOPEN(FileHandle);
 			HANDLE NewFileHandle = ReOpenFile(FileHandle, DesiredAccess, ShareMode, Flags);
-			CloseHandle(FileHandle);
+			TRACE_PLATFORMFILE_END_REOPEN(NewFileHandle);
+
+			TRACE_PLATFORMFILE_BEGIN_CLOSE(FileHandle);
+			BOOL CloseResult = CloseHandle(FileHandle);
+#if PLATFORMFILETRACE_ENABLED
+			if (CloseResult)
+			{
+				TRACE_PLATFORMFILE_END_CLOSE(FileHandle);
+			}
+			else
+			{
+				TRACE_PLATFORMFILE_FAIL_CLOSE(FileHandle);
+			}
+#else
+			(void)CloseResult;
+#endif
 			FileHandle = NewFileHandle;
 		}
 	}
diff --git a/Engine/Source/Runtime/Core/Public/ProfilingDebugging/PlatformFileTrace.h b/Engine/Source/Runtime/Core/Public/ProfilingDebugging/PlatformFileTrace.h
--- a/Engine/Source/Runtime/Core/Public/ProfilingDebugging/PlatformFileTrace.h
+++ b/Engine/Source/Runtime/Core/Public/ProfilingDebugging/PlatformFileTrace.h
@@ -18,6 +18,8 @@ struct FPlatformFileTrace
 	static void BeginOpen(const TCHAR* Path);
 	static void EndOpen(uint64 FileHandle);
 	static void FailOpen(const TCHAR* Path);
+	static void BeginReOpen(uint64 OldFileHandle);
+	static void EndReOpen(uint64 NewFileHandle);
 	static void BeginClose(uint64 FileHandle);
 	static void EndClose(uint64 FileHandle);
 	static void FailClose(uint64 FileHandle);
@@ -38,6 +40,12 @@ struct FPlatformFileTrace
 #define TRACE_PLATFORMFILE_FAIL_OPEN(Path) \
 	FPlatformFileTrace::FailOpen(Path);
 
+#define TRACE_PLATFORMFILE_BEGIN_REOPEN(OldFileHandle) \
+	FPlatformFileTrace::BeginReOpen(uint64(OldFileHandle));
+
+#define TRACE_PLATFORMFILE_END_REOPEN(NewFileHandle) \
+	FPlatformFileTrace::EndReOpen(uint64(NewFileHandle));
+
 #define TRACE_PLATFORMFILE_BEGIN_CLOSE(FileHandle) \
 	FPlatformFileTrace::BeginClose(uint64(FileHandle));
 
