Disable Task,Memory and Loadtime modules when running analysis in Editor

diff --git a/Engine/Source/Developer/TraceServices/Private/TraceServicesModule.cpp b/Engine/Source/Developer/TraceServices/Private/TraceServicesModule.cpp
--- a/Engine/Source/Developer/TraceServices/Private/TraceServicesModule.cpp
+++ b/Engine/Source/Developer/TraceServices/Private/TraceServicesModule.cpp
@@ -79,30 +79,38 @@ TSharedPtr<TraceServices::IModuleService> FTraceServicesModule::CreateModuleServ
 void FTraceServicesModule::StartupModule()
 {
 	IModularFeatures::Get().RegisterModularFeature(TraceServices::ModuleFeatureName, &TimingProfilerModule);
-	IModularFeatures::Get().RegisterModularFeature(TraceServices::ModuleFeatureName, &LoadTimeProfilerModule);
-#if EXPERIMENTAL_STATSTRACE_ENABLED
-	IModularFeatures::Get().RegisterModularFeature(TraceServices::ModuleFeatureName, &StatsModule);
-#endif
 	IModularFeatures::Get().RegisterModularFeature(TraceServices::ModuleFeatureName, &CsvProfilerModule);
 	IModularFeatures::Get().RegisterModularFeature(TraceServices::ModuleFeatureName, &CountersModule);
 	IModularFeatures::Get().RegisterModularFeature(TraceServices::ModuleFeatureName, &NetProfilerModule);
-	IModularFeatures::Get().RegisterModularFeature(TraceServices::ModuleFeatureName, &MemoryModule);
 	IModularFeatures::Get().RegisterModularFeature(TraceServices::ModuleFeatureName, &DiagnosticsModule);
+
+#if EXPERIMENTAL_STATSTRACE_ENABLED
+	IModularFeatures::Get().RegisterModularFeature(TraceServices::ModuleFeatureName, &StatsModule);
+#endif
+
+#if !WITH_EDITOR
+	IModularFeatures::Get().RegisterModularFeature(TraceServices::ModuleFeatureName, &LoadTimeProfilerModule);
+	IModularFeatures::Get().RegisterModularFeature(TraceServices::ModuleFeatureName, &MemoryModule);
 	IModularFeatures::Get().RegisterModularFeature(TraceServices::ModuleFeatureName, &TasksModule);
+#endif
 }
 
 void FTraceServicesModule::ShutdownModule()
 {
+#if !WITH_EDITOR
 	IModularFeatures::Get().UnregisterModularFeature(TraceServices::ModuleFeatureName, &TasksModule);
-	IModularFeatures::Get().UnregisterModularFeature(TraceServices::ModuleFeatureName, &DiagnosticsModule);
 	IModularFeatures::Get().UnregisterModularFeature(TraceServices::ModuleFeatureName, &MemoryModule);
-	IModularFeatures::Get().UnregisterModularFeature(TraceServices::ModuleFeatureName, &NetProfilerModule);
-	IModularFeatures::Get().UnregisterModularFeature(TraceServices::ModuleFeatureName, &CountersModule);
-	IModularFeatures::Get().UnregisterModularFeature(TraceServices::ModuleFeatureName, &CsvProfilerModule);
+	IModularFeatures::Get().UnregisterModularFeature(TraceServices::ModuleFeatureName, &LoadTimeProfilerModule);
+#endif
+
 #if EXPERIMENTAL_STATSTRACE_ENABLED
 	IModularFeatures::Get().UnregisterModularFeature(TraceServices::ModuleFeatureName, &StatsModule);
 #endif
-	IModularFeatures::Get().UnregisterModularFeature(TraceServices::ModuleFeatureName, &LoadTimeProfilerModule);
+
+	IModularFeatures::Get().UnregisterModularFeature(TraceServices::ModuleFeatureName, &DiagnosticsModule);
+	IModularFeatures::Get().UnregisterModularFeature(TraceServices::ModuleFeatureName, &NetProfilerModule);
+	IModularFeatures::Get().UnregisterModularFeature(TraceServices::ModuleFeatureName, &CountersModule);
+	IModularFeatures::Get().UnregisterModularFeature(TraceServices::ModuleFeatureName, &CsvProfilerModule);
 	IModularFeatures::Get().UnregisterModularFeature(TraceServices::ModuleFeatureName, &TimingProfilerModule);
 
 	AnalysisService.Reset();
