Fixed incorrect format specifier for UE_LOG message.

diff --git a/Engine/Source/Developer/TraceServices/Private/Model/TasksProfiler.cpp b/Engine/Source/Developer/TraceServices/Private/Model/TasksProfiler.cpp
--- a/Engine/Source/Developer/TraceServices/Private/Model/TasksProfiler.cpp
+++ b/Engine/Source/Developer/TraceServices/Private/Model/TasksProfiler.cpp
@@ -113,7 +113,7 @@ namespace TraceServices
 		FTaskInfo* Task = TryGetOrCreateTask(TaskId);
 		if (Task == nullptr)
 		{
-			UE_LOG(LogTraceServices, Log, TEXT("TaskCreated(TaskId %d, Timestamp %d) skipped"), TaskId, Timestamp);
+			UE_LOG(LogTraceServices, Log, TEXT("TaskCreated(TaskId %d, Timestamp %.6f) skipped"), TaskId, Timestamp);
 			return;
 		}
 
