Fixed arguments/formatting of UE_LOGs in TraceServices.

diff --git a/Engine/Source/Developer/TraceServices/Private/Analyzers/NetTraceAnalyzer.cpp b/Engine/Source/Developer/TraceServices/Private/Analyzers/NetTraceAnalyzer.cpp
--- a/Engine/Source/Developer/TraceServices/Private/Analyzers/NetTraceAnalyzer.cpp
+++ b/Engine/Source/Developer/TraceServices/Private/Analyzers/NetTraceAnalyzer.cpp
@@ -5,8 +5,7 @@
 #include "TraceServices/Model/Threads.h"
 #include "Logging/LogMacros.h"
 
-DECLARE_LOG_CATEGORY_EXTERN(LogNetTrace, Log, All);
-DEFINE_LOG_CATEGORY(LogNetTrace);
+DEFINE_LOG_CATEGORY_STATIC(LogNetTrace, Log, All);
 
 namespace TraceServices
 {
@@ -173,7 +172,7 @@ void FNetTraceAnalyzer::HandlePacketContentEvent(const FOnEventContext& Context,
 	const uint8 GameInstanceId = EventData.GetValue<uint8>("GameInstanceId");
 	const uint8 PacketType =  EventData.GetValue<uint8>("PacketType");
 
-	//UE_LOG(LogNetTrace, Display, TEXT("FNetTraceAnalyzer::HandlePacketContentEvent: GameInstanceId: %u, ConnectionId: %u, %s"), GameInstanceId, ConnectionId, PacketType ? TEXT("Incoming") : TEXT("Outgoing"));
+	//UE_LOG(LogNetTrace, Display, TEXT("FNetTraceAnalyzer::HandlePacketContentEvent: GameInstanceId: %u, ConnectionId: %u, %s"), (uint32)GameInstanceId, (uint32)ConnectionId, PacketType ? TEXT("Incoming") : TEXT("Outgoing"));
 
 	TSharedRef<FNetTraceGameInstanceState> GameInstanceState = GetOrCreateActiveGameInstanceState(GameInstanceId);
 	FNetTraceConnectionState* ConnectionState = GetActiveConnectionState(GameInstanceId, ConnectionId);
@@ -243,7 +242,7 @@ void FNetTraceAnalyzer::HandlePacketContentEvent(const FOnEventContext& Context,
 					}
 					else
 					{
-						UE_LOG(LogNetTrace, Warning, TEXT("PacketContentEvent GameInstanceId: %u, ConnectionId: %u %s, Missing NameIndex: %u"), GameInstanceId, ConnectionId, ConnectionMode ? TEXT("Incoming") : TEXT("Outgoing"), DecodedNameOrObjectId);	
+						UE_LOG(LogNetTrace, Warning, TEXT("PacketContentEvent GameInstanceId: %u, ConnectionId: %u %s, Missing NameIndex: %llu"), (uint32)GameInstanceId, (uint32)ConnectionId, ConnectionMode ? TEXT("Incoming") : TEXT("Outgoing"), DecodedNameOrObjectId);	
 					}
 				}
 
@@ -464,7 +463,7 @@ void FNetTraceAnalyzer::HandlePacketEvent(const FOnEventContext& Context, const
 	ConnectionState->CurrentPacketStartIndex[ConnectionMode] = ConnectionData.ContentEvents.Num();
 	ConnectionState->CurrentPacketBitOffset[ConnectionMode] = 0U;
 
-	//UE_LOG(LogNetTrace, Log, TEXT("PacketEvent GameInstanceId: %u, ConnectionId: %u, %s, Seq: %u PacketBits: %u"), GameInstanceId, ConnectionId, ConnectionMode ? TEXT("Incoming") : TEXT("Outgoing"), SequenceNumber, Packet.ContentSizeInBits);
+	//UE_LOG(LogNetTrace, Log, TEXT("PacketEvent GameInstanceId: %u, ConnectionId: %u, %s, Seq: %u PacketBits: %u"), (uint32)GameInstanceId, (uint32)ConnectionId, ConnectionMode ? TEXT("Incoming") : TEXT("Outgoing"), SequenceNumber, Packet.ContentSizeInBits);
 }
 
 void FNetTraceAnalyzer::HandlePacketDroppedEvent(const FOnEventContext& Context, const FEventData& EventData)
diff --git a/Engine/Source/Developer/TraceServices/Private/Model/DbgHelpResolver.cpp b/Engine/Source/Developer/TraceServices/Private/Model/DbgHelpResolver.cpp
--- a/Engine/Source/Developer/TraceServices/Private/Model/DbgHelpResolver.cpp
+++ b/Engine/Source/Developer/TraceServices/Private/Model/DbgHelpResolver.cpp
@@ -243,12 +243,12 @@ bool FDbgHelpResolver::LoadModuleSymbols(uint64 Base, uint64 Size, const TCHAR*
 
 	if (!bSymbolsLoaded)
 	{
-		UE_LOG(LogDbgHelp, Warning, TEXT("Unable to load symbols for %s at %p"), Path, Base);
+		UE_LOG(LogDbgHelp, Warning, TEXT("Unable to load symbols for %s at 0x%016llx"), Path, Base);
 		++ModulesFailed;
 	}
 	else
 	{
-		UE_LOG(LogDbgHelp, Display, TEXT("Loaded symbols for %s at %p."), Path, Base);
+		UE_LOG(LogDbgHelp, Display, TEXT("Loaded symbols for %s at 0x%016llx."), Path, Base);
 		++ModulesLoaded;
 	}
 
diff --git a/Engine/Source/Developer/TraceServices/Private/Model/SymslibResolver.cpp b/Engine/Source/Developer/TraceServices/Private/Model/SymslibResolver.cpp
--- a/Engine/Source/Developer/TraceServices/Private/Model/SymslibResolver.cpp
+++ b/Engine/Source/Developer/TraceServices/Private/Model/SymslibResolver.cpp
@@ -419,7 +419,7 @@ FSymslibResolver::EModuleStatus FSymslibResolver::LoadModule(FModuleEntry* Modul
 
 	const uint32 SymbolCount = syms_get_proc_count(ModuleInstance); // get symbol count
 	const uint32 LineCount = syms_get_line_count(ModuleInstance); // get line count
-	UE_LOG(LogSymslib, Display, TEXT("Loaded symbols for '%s' at base 0x%016x, %u symbols and %u lines."), Module->Name, Module->Base, SymbolCount, LineCount);
+	UE_LOG(LogSymslib, Display, TEXT("Loaded symbols for '%s' at base 0x%016llx, %u symbols and %u lines."), Module->Name, Module->Base, SymbolCount, LineCount);
 
 	return EModuleStatus::Loaded;
 }
@@ -441,7 +441,7 @@ bool FSymslibResolver::ResolveSymbol(uint64 Address, FResolvedSymbol* Target, FS
 	FModuleEntry* Module = GetModuleForAddress(Address);
 	if (!Module)
 	{
-		UE_LOG(LogSymslib, Warning, TEXT("No module mapped to address 0x%016x."), Address);
+		UE_LOG(LogSymslib, Warning, TEXT("No module mapped to address 0x%016llx."), Address);
 		UpdateResolvedSymbol(Target, ESymbolQueryResult::NotLoaded, GUnknownModuleTextSymsLib, GUnknownModuleTextSymsLib, GUnknownModuleTextSymsLib, 0);
 		return false;
 	}
