[Insights] Fixed backward compatibility with 32bit platforms for LLM tags tracing.

diff --git a/Engine/Source/Developer/TraceServices/Private/Analyzers/MemoryAnalysis.cpp b/Engine/Source/Developer/TraceServices/Private/Analyzers/MemoryAnalysis.cpp
--- a/Engine/Source/Developer/TraceServices/Private/Analyzers/MemoryAnalysis.cpp
+++ b/Engine/Source/Developer/TraceServices/Private/Analyzers/MemoryAnalysis.cpp
@@ -58,16 +58,42 @@ bool FMemoryAnalyzer::OnEvent(uint16 RouteId, EStyle Style, const FOnEventContex
 		const uint8 TrackerId = EventData.GetValue<uint8>("TrackerId");
 		const uint64 Cycle = EventData.GetValue<uint64>("Cycle");
 		const double Time = Context.EventTime.AsSeconds(Cycle);
-		const TArrayReader<int64>& Tags = EventData.GetArray<int64>("Tags");
+		const TArrayReader<int64>& Tags = EventData.GetArray<int64>("Tags"); // was traced as (void*)[]
 		const TArrayReader<int64>& Samples = EventData.GetArray<int64>("Values");
 
 		const uint32 TagsCount = Tags.Num();
+		const int64* TagsData = Tags.GetData();
+
+		TArray<int64> TagsCopy;
+		if (TagsData == nullptr)
+		{
+			// For backward compatibility with 32bit platforms.
+			const TArrayReader<uint32>& Tags32 = EventData.GetArray<uint32>("Tags");
+			check(Tags32.Num() == TagsCount);
+			const uint32* TagsData32 = Tags32.GetData();
+			if (!TagsData32)
+			{
+				break;
+			}
+			TagsCopy.Reserve(TagsCount);
+			for (uint32 TagIndex = 0; TagIndex < TagsCount; ++TagIndex)
+			{
+				TagsCopy.Push(TagsData32[TagIndex]);
+			}
+			TagsData = TagsCopy.GetData();
+		}
+
+		check(Samples.Num() == TagsCount);
+		const int64* SamplesData = Samples.GetData();
+		check(SamplesData != nullptr);
 		TArray<FMemoryTagSample> Values;
-		for (uint32 i = 0; i < TagsCount; ++i)
+		Values.Reserve(TagsCount);
+		for (uint32 TagIndex = 0; TagIndex < TagsCount; ++TagIndex)
 		{
-			Values.Push(FMemoryTagSample{ Samples[i] });
+			Values.Push(FMemoryTagSample{ SamplesData[TagIndex] });
 		}
-		Provider->AddTagSnapshot(TrackerId, Time, Tags.GetData(), Values.GetData(), TagsCount);
+
+		Provider->AddTagSnapshot(TrackerId, Time, TagsData, Values.GetData(), TagsCount);
 		Sample++;
 
 		Session.UpdateDurationSeconds(Time);
