UnrealInsights: Add aggregation of values in the asset loading tables

diff --git a/Engine/Source/Developer/TraceInsights/Private/Insights/Table/ViewModels/UntypedTable.cpp b/Engine/Source/Developer/TraceInsights/Private/Insights/Table/ViewModels/UntypedTable.cpp
--- a/Engine/Source/Developer/TraceInsights/Private/Insights/Table/ViewModels/UntypedTable.cpp
+++ b/Engine/Source/Developer/TraceInsights/Private/Insights/Table/ViewModels/UntypedTable.cpp
@@ -186,6 +186,7 @@ void FUntypedTable::CreateColumns(const TraceServices::ITableLayout& TableLayout
 	for (int32 ColumnIndex = 0; ColumnIndex < ColumnCount; ++ColumnIndex)
 	{
 		TraceServices::ETableColumnType ColumnType = TableLayout.GetColumnType(ColumnIndex);
+		uint32 ColumnDisplayHintFlags = TableLayout.GetColumnDisplayHintFlags(ColumnIndex);
 		const TCHAR* ColumnName = TableLayout.GetColumnName(ColumnIndex);
 
 		TSharedRef<FTableColumn> ColumnRef = MakeShared<FTableColumn>(FName(ColumnName));
@@ -222,37 +223,58 @@ void FUntypedTable::CreateColumns(const TraceServices::ITableLayout& TableLayout
 
 		case TraceServices::TableColumnType_Int:
 			Column.SetDataType(ETableCellDataType::Int64);
-			//Aggregation = ETableColumnAggregation::Sum;
+			if (ColumnDisplayHintFlags & TraceServices::TableColumnDisplayHint_Summable)
+			{
+				Aggregation = ETableColumnAggregation::Sum;
+			}
 			HorizontalAlignment = HAlign_Right;
 			InitialColumnWidth = 60.0f;
-			//TODO: if (Hint == AsMemory)
-			//{
-			//	FormatterPtr = MakeShared<FInt64ValueFormatterAsMemory>();
-			//}
-			//else // AsNumber
-			FormatterPtr = MakeShared<FInt64ValueFormatterAsNumber>();
+			if (ColumnDisplayHintFlags & TraceServices::TableColumnDisplayHint_Memory)
+			{
+				FormatterPtr = MakeShared<FInt64ValueFormatterAsMemory>();
+			}
+			else
+			{
+				FormatterPtr = MakeShared<FInt64ValueFormatterAsNumber>();
+			}
 			SorterPtr = MakeShared<FSorterByInt64Value>(ColumnRef);
 			break;
 
 		case TraceServices::TableColumnType_Float:
 			Column.SetDataType(ETableCellDataType::Float);
-			//Aggregation = ETableColumnAggregation::Sum;
+			if (ColumnDisplayHintFlags & TraceServices::TableColumnDisplayHint_Summable)
+			{
+				Aggregation = ETableColumnAggregation::Sum;
+			}
 			HorizontalAlignment = HAlign_Right;
 			InitialColumnWidth = 60.0f;
-			//TODO: if (Hint == AsTimeMs)
-			//else // if (Hint == AsTimeAuto)
-			FormatterPtr = MakeShared<FFloatValueFormatterAsTimeAuto>();
+			if (ColumnDisplayHintFlags & TraceServices::TableColumnDisplayHint_Time)
+			{
+				FormatterPtr = MakeShared<FFloatValueFormatterAsTimeAuto>();
+			}
+			else
+			{
+				FormatterPtr = MakeShared<FFloatValueFormatterAsNumber>();
+			}
 			SorterPtr = MakeShared<FSorterByFloatValue>(ColumnRef);
 			break;
 
 		case TraceServices::TableColumnType_Double:
 			Column.SetDataType(ETableCellDataType::Double);
-			//Aggregation = ETableColumnAggregation::Sum;
+			if (ColumnDisplayHintFlags & TraceServices::TableColumnDisplayHint_Summable)
+			{
+				Aggregation = ETableColumnAggregation::Sum;
+			}
 			HorizontalAlignment = HAlign_Right;
 			InitialColumnWidth = 80.0f;
-			//TODO: if (Hint == AsTimeMs)
-			//else // if (Hint == AsTimeAuto)
-			FormatterPtr = MakeShared<FDoubleValueFormatterAsTimeAuto>();
+			if (ColumnDisplayHintFlags & TraceServices::TableColumnDisplayHint_Time)
+			{
+				FormatterPtr = MakeShared<FDoubleValueFormatterAsTimeAuto>();
+			}
+			else
+			{
+				FormatterPtr = MakeShared<FDoubleValueFormatterAsNumber>();
+			}
 			SorterPtr = MakeShared<FSorterByDoubleValue>(ColumnRef);
 			break;
 
diff --git a/Engine/Source/Developer/TraceInsights/Private/Insights/Table/Widgets/STableTreeView.cpp b/Engine/Source/Developer/TraceInsights/Private/Insights/Table/Widgets/STableTreeView.cpp
--- a/Engine/Source/Developer/TraceInsights/Private/Insights/Table/Widgets/STableTreeView.cpp
+++ b/Engine/Source/Developer/TraceInsights/Private/Insights/Table/Widgets/STableTreeView.cpp
@@ -1284,10 +1284,20 @@ void STableTreeView::UpdateAggregatedValues(FTableTreeNode& GroupNode)
 		switch (Column.GetAggregation())
 		{
 			case ETableColumnAggregation::Sum:
-				STableTreeView::UpdateAggregationRec<int64>(Column, GroupNode, 0, true, [](int64 InValue, TOptional<FTableCellValue> InTableCellValue)
+				if (Column.GetDataType() == ETableCellDataType::Float || Column.GetDataType() == ETableCellDataType::Double)
 				{
-					return InValue + InTableCellValue->AsInt64();
-				});
+					STableTreeView::UpdateAggregationRec<double>(Column, GroupNode, 0, true, [](double InValue, TOptional<FTableCellValue> InTableCellValue)
+					{
+						return InValue + InTableCellValue->AsDouble();
+					});
+				}
+				else
+				{
+					STableTreeView::UpdateAggregationRec<int64>(Column, GroupNode, 0, true, [](int64 InValue, TOptional<FTableCellValue> InTableCellValue)
+					{
+						return InValue + InTableCellValue->AsInt64();
+					});
+				}
 				break;
 
 			case ETableColumnAggregation::Min:
diff --git a/Engine/Source/Developer/TraceServices/Private/Model/CsvProfilerPrivate.h b/Engine/Source/Developer/TraceServices/Private/Model/CsvProfilerPrivate.h
--- a/Engine/Source/Developer/TraceServices/Private/Model/CsvProfilerPrivate.h
+++ b/Engine/Source/Developer/TraceServices/Private/Model/CsvProfilerPrivate.h
@@ -72,6 +72,10 @@ private:
 		virtual uint64 GetColumnCount() const override;
 		virtual const TCHAR* GetColumnName(uint64 ColumnIndex) const override;
 		virtual ETableColumnType GetColumnType(uint64 ColumnIndex) const override;
+		virtual uint32 GetColumnDisplayHintFlags(uint64 ColumnIndex) const override
+		{
+			return 0;
+		}
 
 	private:
 		const TArray<FStatSeries*>& StatSeries;
diff --git a/Engine/Source/Developer/TraceServices/Private/Model/LoadTimeProfiler.cpp b/Engine/Source/Developer/TraceServices/Private/Model/LoadTimeProfiler.cpp
--- a/Engine/Source/Developer/TraceServices/Private/Model/LoadTimeProfiler.cpp
+++ b/Engine/Source/Developer/TraceServices/Private/Model/LoadTimeProfiler.cpp
@@ -102,15 +102,16 @@ FLoadTimeProfilerProvider::FLoadTimeProfilerProvider(IAnalysisSession& InSession
 	RequestsTable.EditLayout().
 		AddColumn(&FLoadRequest::Name, TEXT("Name")).
 		AddColumn(&FLoadRequest::ThreadId, TEXT("ThreadId")).
-		AddColumn(&FLoadRequest::StartTime, TEXT("StartTime")).
-		AddColumn(&FLoadRequest::EndTime, TEXT("EndTime")).
+		AddColumn(&FLoadRequest::StartTime, TEXT("StartTime"), TableColumnDisplayHint_Time).
+		AddColumn(&FLoadRequest::EndTime, TEXT("EndTime"), TableColumnDisplayHint_Time).
 		AddColumn<int32>(
 			[](const FLoadRequest& Row)
 			{
 				return Row.Packages.Num();
 			},
-			TEXT("PackageCount")).
-		AddColumn(&FLoadTimeProfilerProvider::PackageSizeSum, TEXT("Size")).
+			TEXT("PackageCount"),
+			TableColumnDisplayHint_Summable).
+		AddColumn(&FLoadTimeProfilerProvider::PackageSizeSum, TEXT("Size"), TableColumnDisplayHint_Memory | TableColumnDisplayHint_Summable).
 		AddColumn<const TCHAR*>(
 			[](const FLoadRequest& Row)
 			{
@@ -120,12 +121,12 @@ FLoadTimeProfilerProvider::FLoadTimeProfilerProvider(IAnalysisSession& InSession
 
 	AggregatedStatsTableLayout.
 		AddColumn(&FLoadTimeProfilerAggregatedStats::Name, TEXT("Name")).
-		AddColumn(&FLoadTimeProfilerAggregatedStats::Count, TEXT("Count")).
-		AddColumn(&FLoadTimeProfilerAggregatedStats::Total, TEXT("Total")).
-		AddColumn(&FLoadTimeProfilerAggregatedStats::Min, TEXT("Min")).
-		AddColumn(&FLoadTimeProfilerAggregatedStats::Max, TEXT("Max")).
-		AddColumn(&FLoadTimeProfilerAggregatedStats::Average, TEXT("Avg")).
-		AddColumn(&FLoadTimeProfilerAggregatedStats::Median, TEXT("Med"));
+		AddColumn(&FLoadTimeProfilerAggregatedStats::Count, TEXT("Count"), TableColumnDisplayHint_Summable).
+		AddColumn(&FLoadTimeProfilerAggregatedStats::Total, TEXT("Total"), TableColumnDisplayHint_Time | TableColumnDisplayHint_Summable).
+		AddColumn(&FLoadTimeProfilerAggregatedStats::Min, TEXT("Min"), TableColumnDisplayHint_Time).
+		AddColumn(&FLoadTimeProfilerAggregatedStats::Max, TEXT("Max"), TableColumnDisplayHint_Time).
+		AddColumn(&FLoadTimeProfilerAggregatedStats::Average, TEXT("Avg"), TableColumnDisplayHint_Time).
+		AddColumn(&FLoadTimeProfilerAggregatedStats::Median, TEXT("Med"), TableColumnDisplayHint_Time);
 
 	PackagesTableLayout.
 		AddColumn<const TCHAR*>([](const FPackagesTableRow& Row)
@@ -133,11 +134,11 @@ FLoadTimeProfilerProvider::FLoadTimeProfilerProvider(IAnalysisSession& InSession
 				return Row.PackageInfo->Name;
 			},
 			TEXT("Package")).
-		AddColumn(&FPackagesTableRow::SerializedHeaderSize, TEXT("SerializedHeaderSize")).
-		AddColumn(&FPackagesTableRow::SerializedExportsCount, TEXT("SerializedExportsCount")).
-		AddColumn(&FPackagesTableRow::SerializedExportsSize, TEXT("SerializedExportsSize")).
-		AddColumn(&FPackagesTableRow::MainThreadTime, TEXT("MainThreadTime")).
-		AddColumn(&FPackagesTableRow::AsyncLoadingThreadTime, TEXT("AsyncLoadingThreadTime"));
+		AddColumn(&FPackagesTableRow::SerializedHeaderSize, TEXT("SerializedHeaderSize"), TableColumnDisplayHint_Memory | TableColumnDisplayHint_Summable).
+		AddColumn(&FPackagesTableRow::SerializedExportsCount, TEXT("SerializedExportsCount"), TableColumnDisplayHint_Summable).
+		AddColumn(&FPackagesTableRow::SerializedExportsSize, TEXT("SerializedExportsSize"), TableColumnDisplayHint_Memory | TableColumnDisplayHint_Summable).
+		AddColumn(&FPackagesTableRow::MainThreadTime, TEXT("MainThreadTime"), TableColumnDisplayHint_Time | TableColumnDisplayHint_Summable).
+		AddColumn(&FPackagesTableRow::AsyncLoadingThreadTime, TEXT("AsyncLoadingThreadTime"), TableColumnDisplayHint_Time | TableColumnDisplayHint_Summable);
 
 	ExportsTableLayout.
 		AddColumn<const TCHAR*>([](const FExportsTableRow& Row)
@@ -155,9 +156,9 @@ FLoadTimeProfilerProvider::FLoadTimeProfilerProvider(IAnalysisSession& InSession
 				return GetLoadTimeProfilerObjectEventTypeString(Row.EventType);
 			},
 			TEXT("EventType")).
-		AddColumn(&FExportsTableRow::SerializedSize, TEXT("SerializedSize")).
-		AddColumn(&FExportsTableRow::MainThreadTime, TEXT("MainThreadTime")).
-		AddColumn(&FExportsTableRow::AsyncLoadingThreadTime, TEXT("AsyncLoadingThreadTime"));
+		AddColumn(&FExportsTableRow::SerializedSize, TEXT("SerializedSize"), TableColumnDisplayHint_Memory | TableColumnDisplayHint_Summable).
+		AddColumn(&FExportsTableRow::MainThreadTime, TEXT("MainThreadTime"), TableColumnDisplayHint_Time | TableColumnDisplayHint_Summable).
+		AddColumn(&FExportsTableRow::AsyncLoadingThreadTime, TEXT("AsyncLoadingThreadTime"), TableColumnDisplayHint_Time | TableColumnDisplayHint_Summable);
 }
 
 bool FLoadTimeProfilerProvider::GetCpuThreadTimelineIndex(uint32 ThreadId, uint32& OutTimelineIndex) const
diff --git a/Engine/Source/Developer/TraceServices/Private/Model/Tables.h b/Engine/Source/Developer/TraceServices/Private/Model/Tables.h
--- a/Engine/Source/Developer/TraceServices/Private/Model/Tables.h
+++ b/Engine/Source/Developer/TraceServices/Private/Model/Tables.h
@@ -154,11 +154,12 @@ class TTableLayout
 {
 public:
 	template<typename ColumnNativeType>
-	TTableLayout<RowType>& AddColumn(ColumnNativeType RowType::* MemberVariableColumn, const TCHAR* ColumnName)
+	TTableLayout<RowType>& AddColumn(ColumnNativeType RowType::* MemberVariableColumn, const TCHAR* ColumnName, uint32 DisplayHintFlags = 0)
 	{
 		Columns.Add({
 			ColumnName,
 			GetColumnTypeFromNativeType<ColumnNativeType>(),
+			DisplayHintFlags,
 			[MemberVariableColumn](const RowType& Row) -> FColumnValueContainer
 			{
 				return FColumnValueContainer(Row.*MemberVariableColumn);
@@ -168,11 +169,12 @@ public:
 	}
 
 	template<typename ColumnNativeType>
-	TTableLayout<RowType>& AddColumn(ColumnNativeType(RowType::* MemberFunctionColumn)() const, const TCHAR* ColumnName)
+	TTableLayout<RowType>& AddColumn(ColumnNativeType(RowType::* MemberFunctionColumn)() const, const TCHAR* ColumnName, uint32 DisplayHintFlags = 0)
 	{
 		Columns.Add({
 			ColumnName,
 			GetColumnTypeFromNativeType<ColumnNativeType>(),
+			DisplayHintFlags,
 			[MemberFunctionColumn](const RowType& Row) -> FColumnValueContainer
 			{
 				return FColumnValueContainer((Row.*MemberFunctionColumn)());
@@ -182,11 +184,12 @@ public:
 	}
 
 	template<typename ColumnNativeType>
-	TTableLayout<RowType>& AddColumn(ColumnNativeType(*FunctionColumn)(const RowType&), const TCHAR* ColumnName)
+	TTableLayout<RowType>& AddColumn(ColumnNativeType(*FunctionColumn)(const RowType&), const TCHAR* ColumnName, uint32 DisplayHintFlags = 0)
 	{
 		Columns.Add({
 			ColumnName,
 			GetColumnTypeFromNativeType<ColumnNativeType>(),
+			DisplayHintFlags,
 			[FunctionColumn](const RowType& Row) -> FColumnValueContainer
 			{
 				return FColumnValueContainer(FunctionColumn(Row));
@@ -210,6 +213,11 @@ public:
 		return Columns[ColumnIndex].Type;
 	}
 
+	uint32 GetColumnDisplayHintFlags(uint64 ColumnIndex) const override
+	{
+		return Columns[ColumnIndex].DisplayHintFlags;
+	}
+
 	FColumnValueContainer GetColumnValue(const RowType& Row, uint64 ColumnIndex) const
 	{
 		return Columns[ColumnIndex].Projector(Row);
@@ -220,6 +228,7 @@ private:
 	{
 		FString Name;
 		ETableColumnType Type;
+		uint32 DisplayHintFlags;
 		TFunction<FColumnValueContainer(const RowType&)> Projector;
 	};
 
diff --git a/Engine/Source/Developer/TraceServices/Public/TraceServices/Containers/Tables.h b/Engine/Source/Developer/TraceServices/Public/TraceServices/Containers/Tables.h
--- a/Engine/Source/Developer/TraceServices/Public/TraceServices/Containers/Tables.h
+++ b/Engine/Source/Developer/TraceServices/Public/TraceServices/Containers/Tables.h
@@ -17,6 +17,13 @@ enum ETableColumnType
 	TableColumnType_CString,
 };
 
+enum ETableColumnDisplayHint
+{
+	TableColumnDisplayHint_Time = 1 << 0,
+	TableColumnDisplayHint_Memory = 1 << 1,
+	TableColumnDisplayHint_Summable = 1 << 31
+};
+
 class ITableLayout
 {
 public:
@@ -24,6 +31,7 @@ public:
 	virtual uint64 GetColumnCount() const = 0;
 	virtual const TCHAR* GetColumnName(uint64 ColumnIndex) const = 0;
 	virtual ETableColumnType GetColumnType(uint64 ColumnIndex) const = 0;
+	virtual uint32 GetColumnDisplayHintFlags(uint64 ColumnIndex) const = 0;
 };
 
 class IUntypedTableReader
