Fix for the use of uint16 as a character type.

diff --git a/Engine/Source/Developer/TraceServices/Private/Analyzers/GpuProfilerTraceAnalysis.cpp b/Engine/Source/Developer/TraceServices/Private/Analyzers/GpuProfilerTraceAnalysis.cpp
--- a/Engine/Source/Developer/TraceServices/Private/Analyzers/GpuProfilerTraceAnalysis.cpp
+++ b/Engine/Source/Developer/TraceServices/Private/Analyzers/GpuProfilerTraceAnalysis.cpp
@@ -35,10 +35,10 @@ bool FGpuProfilerAnalyzer::OnEvent(uint16 RouteId, EStyle Style, const FOnEventC
 	case RouteId_EventSpec:
 	{
 		uint32 EventType = EventData.GetValue<uint32>("EventType");
-		const auto& Name = EventData.GetArray<uint16>("Name");
+		const auto& Name = EventData.GetArray<UTF16CHAR>("Name");
 
-		const FString EventName(Name.Num(), Name.GetData());
-		EventTypeMap.Add(EventType, TimingProfilerProvider.AddGpuTimer(*EventName));
+		auto NameTChar = StringCast<TCHAR>(Name.GetData(), Name.Num());
+		EventTypeMap.Add(EventType, TimingProfilerProvider.AddGpuTimer(FStringView(NameTChar.Get(), NameTChar.Length())));
 		break;
 	}
 	case RouteId_Frame:
@@ -85,7 +85,7 @@ bool FGpuProfilerAnalyzer::OnEvent(uint16 RouteId, EStyle Style, const FOnEventC
 				else
 				{
 					FTimingProfilerEvent Event;
-					Event.TimerIndex = TimingProfilerProvider.AddGpuTimer(TEXT("<unknown>"));
+					Event.TimerIndex = TimingProfilerProvider.AddGpuTimer(TEXT("<unknown>"_SV));
 					EventTypeMap.Add(EventType, Event.TimerIndex);
 					ThisTimeline.AppendBeginEvent(LastTime, Event);
 				}
diff --git a/Engine/Source/Developer/TraceServices/Private/Model/TimingProfiler.cpp b/Engine/Source/Developer/TraceServices/Private/Model/TimingProfiler.cpp
--- a/Engine/Source/Developer/TraceServices/Private/Model/TimingProfiler.cpp
+++ b/Engine/Source/Developer/TraceServices/Private/Model/TimingProfiler.cpp
@@ -100,7 +100,7 @@ void FTimingProfilerProvider::SetTimerName(uint32 TimerId, const TCHAR* Name)
 	Timer.NameHash = NameHash;
 }
 
-uint32 FTimingProfilerProvider::AddGpuTimer(const TCHAR* Name)
+uint32 FTimingProfilerProvider::AddGpuTimer(FStringView Name)
 {
 	Session.WriteAccessCheck();
 
@@ -108,15 +108,15 @@ uint32 FTimingProfilerProvider::AddGpuTimer(const TCHAR* Name)
 	return Timer.Id;
 }
 
-FTimingProfilerTimer& FTimingProfilerProvider::AddTimerInternal(const TCHAR* Name, bool IsGpuTimer)
+FTimingProfilerTimer& FTimingProfilerProvider::AddTimerInternal(FStringView Name, bool IsGpuTimer)
 {
 	FTimingProfilerTimer& Timer = Timers.AddDefaulted_GetRef();
 	Timer.Id = Timers.Num() - 1;
 	Timer.Name = Session.StoreString(Name);
 	uint32 NameHash = 0;
-	for (const TCHAR* c = Name; *c; ++c)
+	for (TCHAR c : Name)
 	{
-		NameHash = (NameHash + *c) * 0x2c2c57ed;
+		NameHash = (NameHash + c) * 0x2c2c57ed;
 	}
 	Timer.NameHash = NameHash;
 	Timer.IsGpuTimer = IsGpuTimer;
diff --git a/Engine/Source/Developer/TraceServices/Private/Model/TimingProfilerPrivate.h b/Engine/Source/Developer/TraceServices/Private/Model/TimingProfilerPrivate.h
--- a/Engine/Source/Developer/TraceServices/Private/Model/TimingProfilerPrivate.h
+++ b/Engine/Source/Developer/TraceServices/Private/Model/TimingProfilerPrivate.h
@@ -23,7 +23,7 @@ public:
 	FTimingProfilerProvider(IAnalysisSession& InSession);
 	virtual ~FTimingProfilerProvider();
 	uint32 AddCpuTimer(const TCHAR* Name);
-	uint32 AddGpuTimer(const TCHAR* Name);
+	uint32 AddGpuTimer(FStringView Name);
 	void SetTimerName(uint32 TimerId, const TCHAR* Name);
 	uint32 AddMetadata(uint32 MasterTimerId, TArray<uint8>&& Metadata);
 	TimelineInternal& EditCpuThreadTimeline(uint32 ThreadId);
@@ -43,7 +43,7 @@ public:
 	virtual TArrayView<const uint8> GetMetadata(uint32 TimerId) const override;
 
 private:
-	FTimingProfilerTimer& AddTimerInternal(const TCHAR* Name, bool IsGpuEvent);
+	FTimingProfilerTimer& AddTimerInternal(FStringView Name, bool IsGpuEvent);
 
 	struct FMetadata 
 	{
