[Insights] Made FSlabAllocator thread safe. This fixes memory stomps caused by the symbol resolver's string store overwriting the memory allocated for various providers (ex. crash in FFileActivityProvider).

diff --git a/Engine/Source/Developer/TraceServices/Private/Common/SlabAllocator.h b/Engine/Source/Developer/TraceServices/Private/Common/SlabAllocator.h
--- a/Engine/Source/Developer/TraceServices/Private/Common/SlabAllocator.h
+++ b/Engine/Source/Developer/TraceServices/Private/Common/SlabAllocator.h
@@ -2,8 +2,9 @@
 
 #pragma once
 
-#include "HAL/UnrealMemory.h"
 #include "Containers/Array.h"
+#include "HAL/CriticalSection.h"
+#include "HAL/UnrealMemory.h"
 #include "TraceServices/Containers/Allocators.h"
 
 namespace TraceServices
@@ -21,6 +22,7 @@ public:
 
 	~FSlabAllocator()
 	{
+		FScopeLock _(&Cs);
 		for (void* Slab : Slabs)
 		{
 			FMemory::Free(Slab);
@@ -29,6 +31,7 @@ public:
 
 	virtual void* Allocate(uint64 Size) override
 	{
+		FScopeLock _(&Cs);
 		uint64 AllocationSize = Size + (-int64(Size) & 15);
 		if (!CurrentSlab || CurrentSlabAllocatedSize + AllocationSize > SlabSize)
 		{
@@ -44,6 +47,7 @@ public:
 	}
 
 private:
+	FCriticalSection Cs;
 	TArray<void*> Slabs;
 	uint8* CurrentSlab = nullptr;
 	const uint64 SlabSize;
