Fixed missing include.

diff --git a/Engine/Source/Developer/TraceServices/Private/Common/SlabAllocator.h b/Engine/Source/Developer/TraceServices/Private/Common/SlabAllocator.h
--- a/Engine/Source/Developer/TraceServices/Private/Common/SlabAllocator.h
+++ b/Engine/Source/Developer/TraceServices/Private/Common/SlabAllocator.h
@@ -5,6 +5,7 @@
 #include "Containers/Array.h"
 #include "HAL/CriticalSection.h"
 #include "HAL/UnrealMemory.h"
+#include "Misc/ScopeLock.h"
 #include "TraceServices/Containers/Allocators.h"
 
 namespace TraceServices
