[Insights] Cache layer

Introduces a persistent storage system for Unreal Insights. The cache is based on fixed block sizes. Users give a unique name which is associated with a unique caching index. User can then allocate fixed sized blocks of memory and store data inside. The system make no assumption of the saved data. Each block is ref counted and as it is released the memory is committed to disk.

Additionally a small user data section is provided per cache index for persistent storage of data which does not fit neatly into the block model, such as element counts.

diff --git a/Engine/Source/Developer/TraceServices/Private/AnalysisCache.cpp b/Engine/Source/Developer/TraceServices/Private/AnalysisCache.cpp
new file mode 100644
--- /dev/null
+++ b/Engine/Source/Developer/TraceServices/Private/AnalysisCache.cpp
@@ -0,0 +1,447 @@
+﻿// Copyright Epic Games, Inc. All Rights Reserved.
+
+#include "AnalysisCache.h"
+
+#include "Algo/Count.h"
+#include "Algo/Find.h"
+#include "Async/MappedFileHandle.h"
+#include "Containers/StringView.h"
+#include "HAL/PlatformFileManager.h"
+#include "Logging/LogMacros.h"
+#include "Memory/SharedBuffer.h"
+#include "Misc/Paths.h"
+#include "Misc/PathViews.h"
+#include "Serialization/BufferArchive.h"
+#include "Serialization/BufferWriter.h"
+#include "Serialization/CompactBinaryPackage.h"
+#include "Serialization/CompactBinaryWriter.h"
+#include "Serialization/MemoryReader.h"
+
+DEFINE_LOG_CATEGORY_STATIC(LogAnalysisCache, Log, All);
+
+namespace TraceServices {
+			
+//////////////////////////////////////////////////////////////////////
+FCacheId FAnalysisCache::FFileContents::GetId(const TCHAR* Name, uint16 Flags)
+{
+	FIndexEntry* Entry = Algo::FindByPredicate(IndexEntries, [Name](const FIndexEntry& Entry){ return Entry.Name.Equals(Name); });
+	if (Entry)
+	{
+		return Entry->Id;
+	}
+
+	// Name was not previously registered, create new entry
+	const uint32 NewId = IndexEntries.Num() + 1;
+	FIndexEntry& NewEntry = IndexEntries.AddZeroed_GetRef();
+	NewEntry.Name = Name;
+	NewEntry.Id = NewId;
+	NewEntry.Flags = uint32(Flags);
+	return NewId;
+}
+
+//////////////////////////////////////////////////////////////////////
+uint16 FAnalysisCache::FFileContents::GetFlags(FCacheId InId)
+{
+	FIndexEntry* Entry = Algo::FindBy(IndexEntries, InId, [](const FIndexEntry& InEntry){ return InEntry.Id; });
+	if (Entry)
+	{
+		return uint16(Entry->Flags & 0xffff);
+	}
+	return 0;
+}
+	
+//////////////////////////////////////////////////////////////////////
+FMutableMemoryView FAnalysisCache::FFileContents::GetUserData(FCacheId InId)
+{
+	FIndexEntry* Entry = Algo::FindBy(IndexEntries, InId, [](const FIndexEntry& InEntry){ return InEntry.Id; });
+	if (Entry)
+	{
+		return FMutableMemoryView(Entry->UserData, UserDataSize);
+	}
+	return FMutableMemoryView();
+}
+
+//////////////////////////////////////////////////////////////////////
+bool FAnalysisCache::FFileContents::Save(IFileHandle* File)
+{
+	if (bTransientMode)
+	{
+		return true;
+	}
+	
+	check(File);
+	File->Seek(0);
+
+	FCbWriter Writer;
+	Writer.BeginObject();
+	Writer << "Version" << CurrentVersion;
+	Writer.BeginArray("Index"_ASV);
+	for (auto Entry : IndexEntries)
+	{
+		Writer.BeginObject();
+		Writer << "N"_ASV << Entry.Name;
+		Writer << "I"_ASV << Entry.Id;
+		Writer << "F"_ASV << Entry.Flags;
+		Writer.AddBinary("UD"_ASV, &Entry.UserData, UserDataSize);
+		Writer.EndObject();
+	}
+	Writer.EndArray();
+	Writer.BeginArray("Blocks"_ASV);
+	for (const FBlockEntry& Entry : Blocks)
+	{
+		Writer.AddBinary(&Entry, sizeof(FBlockEntry));
+	}
+	Writer.EndArray();
+	Writer.EndObject();
+	
+	FCbPackage Package(Writer.Save().AsObject());
+
+	FUniqueBuffer Buffer = FUniqueBuffer::Alloc(ReservedSize);
+	FBufferWriter BufferWriter(Buffer.GetData(), Buffer.GetSize());
+	Package.Save(BufferWriter);
+	
+	return File->Write((uint8*)Buffer.GetData(), Buffer.GetSize());
+}
+
+//////////////////////////////////////////////////////////////////////
+bool FAnalysisCache::FFileContents::Load(IFileHandle* File)
+{
+	if (bTransientMode)
+	{
+		return true;
+	}
+	
+	check(File);
+	File->Seek(0);
+	
+	FUniqueBuffer Buffer = FUniqueBuffer::Alloc(ReservedSize);
+	File->Read((uint8*)Buffer.GetData(), Buffer.GetSize());
+	
+	FMemoryReaderView Ar(MakeArrayView<uint8>((uint8*)Buffer.GetData(), Buffer.GetSize()));
+	
+	FCbPackage Package;
+	Package.Load(Ar);
+
+	uint32 PackageVersion = Package.GetObject().Find("Version").AsUInt32();
+	UE_LOG(LogAnalysisCache, Display, TEXT("Cache file (version %u) loaded."), PackageVersion);
+	if (PackageVersion != CurrentVersion)
+	{
+		// todo: Handle this better
+		return false;
+	}
+
+	FCbArrayView IndexArray = Package.GetObject().Find("Index"_ASV).AsArrayView();
+	IndexEntries.Reserve(IndexArray.Num());
+	for (FCbFieldView IndexEntry : IndexArray)
+	{
+		FCbObjectView IndexEntryObj = IndexEntry.AsObjectView();
+		FAnsiStringView NameView = IndexEntryObj["N"_ASV].AsString();
+		const uint32 Id = IndexEntryObj["I"_ASV].AsUInt32();
+		const uint32 Flags = IndexEntryObj["F"_ASV].AsUInt32();
+		FMemoryView UserData = IndexEntryObj["UD"_ASV].AsBinaryView();
+		FIndexEntry& Entry = IndexEntries.AddZeroed_GetRef();
+		Entry.Name = FString(NameView);
+		Entry.Id = Id;
+		Entry.Flags = Flags;
+		FMutableMemoryView UserDataDst(&Entry.UserData, UserDataSize);
+		FMutableMemoryView Remainder = UserDataDst.CopyFrom(UserData);
+		check(Remainder.GetSize() == 0);
+	}
+
+	FCbArrayView BlockArray = Package.GetObject().Find("Blocks"_ASV).AsArrayView();
+	Blocks.Reserve(BlockArray.Num());
+	
+	for (FCbFieldView BlockEntryView : BlockArray)
+	{
+		FBlockEntry& Block = Blocks.AddZeroed_GetRef();
+		FMutableMemoryView BlockView = FMutableMemoryView(&Block, sizeof(FBlockEntry));
+		BlockView.CopyFrom(BlockEntryView.AsBinaryView());
+	}
+
+	return true;
+}
+
+//////////////////////////////////////////////////////////////////////
+uint64 FAnalysisCache::FFileContents::UpdateBlock(FMemoryView Block, BlockKeyType BlockKey, IFileHandle* File)
+{
+	if (bTransientMode)
+	{
+		return 0;
+	}
+	
+	check(File);
+	const uint64 EntryIndex = Algo::BinarySearchBy(Blocks, BlockKey, [&](const FBlockEntry& InEntry){return InEntry.BlockKey;});
+	const FIoHash CurrentHash = FIoHash::HashBuffer(Block);
+	
+	if (EntryIndex != INDEX_NONE)
+	{
+		FBlockEntry& Entry = Blocks[EntryIndex];
+		// todo: This only works if size hasn't changed!
+		if (CurrentHash != Entry.Hash)
+		{
+			//todo: Add compression
+			File->Seek(Entry.Offset);
+			if(!File->Write((const uint8*)Block.GetData(), Block.GetSize()))
+			{
+				UE_LOG(LogAnalysisCache, Error, TEXT("Failed to update block 0x%x at offset %d kb"), BlockKey, Entry.Offset / 1024);
+				return 0;
+			}
+
+			Entry.Hash = CurrentHash;
+			return Block.GetSize();
+		}
+	}
+	else
+	{
+		// Write to end of file and add to blocks array
+		File->SeekFromEnd(0);
+		const uint64 Offset = File->Tell();
+		check(Offset >= ReservedSize);
+		if(!File->Write((const uint8*) Block.GetData(), Block.GetSize()))
+		{
+			UE_LOG(LogAnalysisCache, Error, TEXT("Failed to update block 0x%x at offset %u kb"), BlockKey, Offset / 1024);
+		}
+
+		// Insert entry in blocks list and sort array
+		Blocks.Emplace(FBlockEntry{BlockKey, 0, Offset, 0, Block.GetSize(), CurrentHash});
+		Algo::SortBy(Blocks, [](const FBlockEntry& Entry){ return Entry.BlockKey; });
+
+		return Block.GetSize();
+	}
+
+	return 0;
+}
+
+//////////////////////////////////////////////////////////////////////
+uint64 FAnalysisCache::FFileContents::LoadBlock(FMutableMemoryView Block, BlockKeyType BlockKey, IFileHandle* File)
+{
+	if (bTransientMode)
+	{
+		return 0;
+	}
+	
+	check(File);
+	const uint64 EntryIndex = Algo::BinarySearchBy(Blocks, BlockKey, [&](const FBlockEntry& InEntry){return InEntry.BlockKey;});
+	if (EntryIndex == INDEX_NONE)
+	{
+		UE_LOG(LogAnalysisCache, Error, TEXT("Trying to load unknown block 0x%x."), BlockKey);
+		return 0;
+	}
+	
+	FBlockEntry& Entry = Blocks[EntryIndex];
+	check(Entry.UncompressedSize <= Block.GetSize());
+	
+	if (!File->Seek(Entry.Offset))
+	{
+		UE_LOG(LogAnalysisCache, Error, TEXT("Block 0x%x was located on an invalid offset %u kb."), BlockKey, Entry.Offset / 1024);
+		return 0;
+	}
+
+	//todo: Add compression
+	if (!File->Read((uint8*)Block.GetData(), Entry.UncompressedSize))
+	{
+		UE_LOG(LogAnalysisCache, Error, TEXT("Unable to read block 0x%x on offset %u kb with size %u kb."), BlockKey, Entry.Offset / 1024, Entry.UncompressedSize / 1024);
+		return 0;
+	}
+
+	return Block.GetSize();
+}
+
+//////////////////////////////////////////////////////////////////////
+FAnalysisCache::FAnalysisCache(const TCHAR* Name)
+	: Stats()
+{
+	IPlatformFile& PlatformFile = FPlatformFileManager::Get().GetPlatformFile();
+
+	// Find the cache file path.
+	// todo: This assumes TraceSessions folder.
+	// todo: This will need to be refined as we move away from files
+	const FString& BaseDirectory = FPaths::ProjectSavedDir();
+	const FString BaseName(FPathViews::GetBaseFilename(Name));
+	CacheFilePath = FPaths::Combine(*BaseDirectory, TEXT("TraceSessions"), *BaseName);
+	CacheFilePath += TEXT(".ucache");
+
+	// Opening the file we can encounter one of 4 scenarios:
+	// 1. File does not exist, we create it
+	// 2. File exist, we can read the contents
+	// 3. File exist but we could not open the file for reading. Multiple processes are competing. Put the cache in transient mode.
+	// 4. File does not exist, we cannot create it. Multiple processes are competing to create the file. Put the cache in transient mode.
+	if (!PlatformFile.FileExists(*CacheFilePath))
+	{
+		if (const TUniquePtr<IFileHandle> File(PlatformFile.OpenWrite(*CacheFilePath)); File)
+		{
+			Contents.Save(File.Get());
+		}
+		else
+		{
+			// Unable to open for write. Most likely this is because another instance is racing to create the file
+			UE_LOG(LogAnalysisCache, Warning,
+			       TEXT("Unable to write the cache file %s, possibly already open in another session. Putting cache in transient mode."),
+			       *CacheFilePath);
+			Contents.bTransientMode = true;
+		}
+	}
+	else
+	{
+		if (const TUniquePtr<IFileHandle> File(PlatformFile.OpenRead(*CacheFilePath)); File.IsValid())
+		{
+			if (const bool Result = Contents.Load(File.Get()); !Result)
+			{
+				UE_LOG(LogAnalysisCache, Error, TEXT("Failed to open cache file table of contents."));
+				//todo: Recover by deleting file?
+			}
+
+			UE_LOG(LogAnalysisCache, VeryVerbose, TEXT("Cache contains %d blocks:"), Contents.Blocks.Num());
+			UE_LOG(LogAnalysisCache, VeryVerbose, TEXT("   %10s   %10s   %13s   %13s   %13s"), TEXT("Cache index"), TEXT("Block index"), TEXT("Offset"), TEXT("Uncompressed"), TEXT("Compressed"));
+
+			for (const FFileContents::FBlockEntry& Block : Contents.Blocks)
+			{
+				UE_LOG(LogAnalysisCache, VeryVerbose, TEXT("   %10d   %10d   %10d kb   %10d kb   %10d kb"), GetCacheId(Block.BlockKey), GetBlockIndex(Block.BlockKey), Block.Offset / 1024, Block.UncompressedSize / 1024, Block.CompressedSize / 1024);
+			}
+
+
+			//todo: At this point we can check the file size and if it doesn't match we can trim unknown segments.
+		}
+		else
+		{
+			// Unable to open for read. Most likely this is because another instance is using the file
+			UE_LOG(LogAnalysisCache, Warning,
+			       TEXT("Unable to read the cache file %s, possibly already open in another session. Putting cache in transient mode."),
+			       *CacheFilePath);
+			Contents.bTransientMode = true;
+		}
+	}
+
+	// Build a dictionary of number of blocks per cache id.
+	for (FFileContents::FBlockEntry& Block : Contents.Blocks)
+	{
+		const uint32 CacheId = GetCacheId(Block.BlockKey);
+		IndexBlockCount.FindOrAdd(CacheId)++;
+	}
+	
+	CacheFile = TUniquePtr<IFileHandle>(PlatformFile.OpenRead(*CacheFilePath, true));
+	CacheFileWrite = TUniquePtr<IFileHandle>(PlatformFile.OpenWrite(*CacheFilePath, true, true));
+
+	// Make sure we have valid file handles or is in transient mode.
+	check(Contents.bTransientMode || (CacheFile.IsValid() && CacheFileWrite.IsValid()));
+}
+
+/////////////////////////////////////////////////////////////////////
+FAnalysisCache::~FAnalysisCache()
+{
+	// Remove all references to cached block, forcing them to write.
+	CachedBlocks.Empty();
+
+	// Save the table of contents.
+	if (const bool Result = Contents.Save(CacheFileWrite.Get()); !Result )
+	{
+		UE_LOG(LogAnalysisCache, Error, TEXT("Failed to update cache files table of contents."));
+	}
+
+	UE_LOG(LogAnalysisCache, Display, TEXT("Closing analysis cache \"%s\". %0.2f Mb read, %0.2f Mb written."), *CacheFilePath,
+	       Stats.BytesRead / (1024.f*1024.f), Stats.BytesWritten / (1024.f*1024.f));
+}
+
+/////////////////////////////////////////////////////////////////////
+uint32 FAnalysisCache::GetCacheId(const TCHAR* Name, uint16 Flags)
+{
+	return Contents.GetId(Name, Flags);
+}
+
+/////////////////////////////////////////////////////////////////////
+FMutableMemoryView FAnalysisCache::GetUserData(FCacheId CacheId)
+{
+	return Contents.GetUserData(CacheId);
+}
+
+/////////////////////////////////////////////////////////////////////
+FSharedBuffer FAnalysisCache::CreateBlocks(FCacheId CacheId, uint32 BlockCount)
+{
+	const uint32 BlockIndex = IndexBlockCount.FindOrAdd(CacheId);
+	const BlockKeyType BlockKey = CreateBlockKey(CacheId, BlockIndex);
+
+	// Allocate memory and make the shared buffer with freeing callback.
+	const uint64 TotalBytes = IAnalysisCache::BlockSizeBytes * BlockCount;
+	void* Block = FMemory::Malloc(TotalBytes, BlockAlignment);
+	FMemory::Memzero(Block, TotalBytes);
+	FSharedBuffer Blocks = FSharedBuffer::TakeOwnership(Block, TotalBytes, [this, CacheId, BlockIndex](void* InBlock, uint64 Size)
+	{
+		ReleaseBlocks((uint8*)InBlock, CacheId, BlockIndex, Size);
+	});
+
+	// Increment the block count
+	IndexBlockCount[CacheId] += BlockCount;
+
+	// Add the blocks into our internal caching mechanism
+	if (!(Contents.GetFlags(CacheId) & ECacheFlags_NoGlobalCaching))
+	{
+		check(!CachedBlocks.Contains(BlockKey));
+		CachedBlocks.Add(BlockKey, FSharedBuffer::MakeView(Blocks.GetView(), Blocks));
+	}
+
+	return Blocks;
+}
+
+/////////////////////////////////////////////////////////////////////
+FSharedBuffer FAnalysisCache::GetBlocks(FCacheId CacheId, uint32 BlockIndexStart, uint32 BlockCount)
+{
+	const BlockKeyType CacheBlockKey = CreateBlockKey(CacheId, BlockIndexStart);
+
+	const uint32 ExistingBlockCount = IndexBlockCount.FindOrAdd(CacheId);
+	if (BlockIndexStart >= ExistingBlockCount || (BlockIndexStart + BlockCount) > ExistingBlockCount)
+	{
+		UE_LOG(LogAnalysisCache, Error, TEXT("Block range %u to %u is invalid for cache id %u."), BlockIndexStart,
+		       BlockIndexStart + BlockCount, CacheId);
+		return FSharedBuffer();
+	}
+	
+	// Look in our currently help block cache
+	// todo: Cached blocks are only keyed on first block index. What if a different range is requested? Overlap.
+	if (FSharedBuffer* Block = CachedBlocks.Find(CacheBlockKey))
+	{
+		return *Block;
+	}
+
+	// Allocate a contiguous chunk of memory that fits all the blocks
+	const uint64 TotalBytes = IAnalysisCache::BlockSizeBytes * BlockCount;
+	uint8* BlockBuffer = (uint8*)FMemory::Malloc(TotalBytes, BlockAlignment);
+
+	for (uint32 Block = 0; Block < BlockCount; ++Block)
+	{
+		const BlockKeyType BlockKey = CreateBlockKey(CacheId, BlockIndexStart + Block);
+		const FMutableMemoryView BlockView(BlockBuffer + (Block*IAnalysisCache::BlockSizeBytes), IAnalysisCache::BlockSizeBytes);
+		const uint64 BytesRead = Contents.LoadBlock(BlockView, BlockKey, CacheFile.Get());
+		Stats.BytesRead += BytesRead;
+	}
+	
+	// Take ownership of memory and register freeing callback.
+	FSharedBuffer Blocks = FSharedBuffer::TakeOwnership(BlockBuffer, TotalBytes, [this, CacheId, BlockIndexStart](void* Block, uint64 Size)
+	{
+		ReleaseBlocks((uint8*)Block, CacheId, BlockIndexStart, Size);
+	});
+
+	// Add the blocks into our internal caching mechanism
+	if (!(Contents.GetFlags(CacheId) & ECacheFlags_NoGlobalCaching))
+	{
+		CachedBlocks.Add(CacheBlockKey, Blocks);
+	}
+
+	return Blocks;
+}
+
+/////////////////////////////////////////////////////////////////////
+void FAnalysisCache::ReleaseBlocks(uint8* BlockBuffer, FCacheId CacheId, uint32 BlockIndexStart, uint64 Size)
+{
+	const uint32 BlockCount = Size / IAnalysisCache::BlockSizeBytes;
+	for (uint32 Block = 0; Block < BlockCount; ++Block)
+	{
+		void* BlockStart = BlockBuffer + (Block * IAnalysisCache::BlockSizeBytes);
+		const FMemoryView BlockView = FMemoryView(BlockStart, IAnalysisCache::BlockSizeBytes);
+		const BlockKeyType BlockKey = CreateBlockKey(CacheId, BlockIndexStart + Block);
+		const uint64 BytesWritten = Contents.UpdateBlock(BlockView, BlockKey, CacheFileWrite.Get());
+		Stats.BytesWritten += BytesWritten;
+	}
+}
+
+/////////////////////////////////////////////////////////////////////
+} //namespace TraceServices
diff --git a/Engine/Source/Developer/TraceServices/Private/AnalysisCache.h b/Engine/Source/Developer/TraceServices/Private/AnalysisCache.h
new file mode 100644
--- /dev/null
+++ b/Engine/Source/Developer/TraceServices/Private/AnalysisCache.h
@@ -0,0 +1,148 @@
+﻿// Copyright Epic Games, Inc. All Rights Reserved.
+#pragma once
+
+#include "Containers/UnrealString.h"
+#include "IO/IoHash.h"
+#include "Memory/SharedBuffer.h"
+#include "Templates/UniquePtr.h"
+#include "TraceServices/AnalysisService.h"
+#include "TraceServices/Model/AnalysisCache.h"
+
+class IMappedFileRegion;
+class IMappedFileHandle;
+
+namespace TraceServices {
+
+class IAnalysisSession;
+	
+//////////////////////////////////////////////////////////////////////
+class FAnalysisCache : public IAnalysisCache
+{
+public:
+	explicit FAnalysisCache(const TCHAR* Name);
+	virtual ~FAnalysisCache() override;
+
+	virtual FCacheId GetCacheId(const TCHAR* Name, uint16 Flags) override;
+	virtual FMutableMemoryView GetUserData(FCacheId CacheId) override;
+	virtual FSharedBuffer CreateBlocks(FCacheId CacheId, uint32 Count) override;
+	virtual FSharedBuffer GetBlocks(FCacheId CacheId, uint32 BlockIndexStart, uint32 BlockCount = 1) override;
+	
+private:
+	/**
+	 * Compound key identifying a block.
+	 * Warning: Changing this type also affects file contents.
+	 */
+	typedef uint32 BlockKeyType;
+
+	/**
+	 * Wrapper for file contents
+	 */
+	struct FFileContents
+	{
+		// Current file format version
+		static constexpr uint32 CurrentVersion = 1;
+		// Number of bytes at the start of the file reserved for table of contents.
+		static constexpr uint32 ReservedSize = 512 * 1024;
+		// Number of bytes allocated for user data (per named entry)
+		static constexpr uint32 UserDataSize = 64;
+		
+		FCacheId GetId(const TCHAR* Name, uint16 Flags);
+		FMutableMemoryView GetUserData(FCacheId Id);
+		uint16 GetFlags(FCacheId Id);
+		bool Save(IFileHandle* File);
+		bool Load(IFileHandle* File);
+		uint64 UpdateBlock(FMemoryView Block, BlockKeyType BlockKey, IFileHandle* File);
+		uint64 LoadBlock(FMutableMemoryView Block, BlockKeyType BlockKey, IFileHandle* File);
+
+		struct FIndexEntry
+		{
+			FString Name;
+			uint32 Id;
+			uint32 Flags;
+			uint8 UserData[UserDataSize];
+		};
+		
+		struct FBlockEntry
+		{
+			BlockKeyType BlockKey;
+			uint32 Flags;
+			uint64 Offset;
+			uint64 CompressedSize;
+			uint64 UncompressedSize;
+			FIoHash Hash;
+		};
+		
+		TArray<FIndexEntry> IndexEntries;
+		TArray<FBlockEntry> Blocks;
+		// When this is set blocks or table of contents are never written.
+		bool bTransientMode = false;
+	};
+
+	/**
+	 * Life time statistics
+	 */
+	struct FStats
+	{
+		uint64 BytesRead;
+		uint64 BytesWritten;
+	};
+	
+	/**
+	 * Called when a shared buffer is released. Checks if content has change and updates cached block correspondingly
+	 * and finally frees memory.
+	 * @param BlockBuffer Allocated memory
+	 * @param CacheId Unique cache id
+	 * @param BlockIndexStart Block index for this cache entry
+	 * @param Size of buffer
+	 */
+	void ReleaseBlocks(uint8* BlockBuffer, FCacheId CacheId, uint32 BlockIndexStart, uint64 Size);
+	
+	/**
+	 * Creates a block key identifier
+	 * @param CacheId Unique cache id
+	 * @param BlockIndex Block index for this cache entry
+	 * @return Unique key identifying a single block
+	 */
+	static BlockKeyType CreateBlockKey(FCacheId CacheId, uint32 BlockIndex)
+	{
+		// Use 3 top nibbles for cache id and the rest for block index. That gives a space for 4095 unique cache entries
+		// each containing a maximum of 7 Gb of data.
+		check(CacheId < 0xfff && BlockIndex < 0x000fffff && CacheId > 0);
+		return CacheId << 20 | BlockIndex;
+	}
+
+	/**
+	 * Extract the cache id from a block key
+	 * @param BlockKey Unique key identifying a single block
+	 * @return Unique cache id
+	 */
+	static uint32 GetCacheId(BlockKeyType BlockKey)
+	{
+		return (BlockKey & 0xfff00000) >> 20;
+	}
+
+	/**
+	 * Extract the block index from a block key
+	 * @param BlockKey Unique key identifying a single block
+	 * @return Block index for this cache entry
+	 */
+	static uint32 GetBlockIndex(BlockKeyType BlockKey)
+	{
+		return BlockKey & 0x000fffff;
+	}
+
+	FAnalysisCache() = delete;
+	FAnalysisCache(FAnalysisCache& Other) = delete;
+
+	constexpr static uint32 BlockAlignment = 1024*4;
+	
+	FString CacheFilePath;
+	FFileContents Contents;
+	FStats Stats;
+	TUniquePtr<IFileHandle> CacheFile;
+	TUniquePtr<IFileHandle> CacheFileWrite;
+	TMap<BlockKeyType,FSharedBuffer> CachedBlocks;
+	TMap<uint32,uint32> IndexBlockCount;
+};
+
+}
diff --git a/Engine/Source/Developer/TraceServices/Public/TraceServices/Model/AnalysisCache.h b/Engine/Source/Developer/TraceServices/Public/TraceServices/Model/AnalysisCache.h
new file mode 100644
--- /dev/null
+++ b/Engine/Source/Developer/TraceServices/Public/TraceServices/Model/AnalysisCache.h
@@ -0,0 +1,91 @@
+﻿// Copyright Epic Games, Inc. All Rights Reserved.
+#pragma once
+
+#include "CoreTypes.h"
+#include "Memory/MemoryView.h"
+
+namespace TraceServices
+{
+	/**
+	* An iterator for segments
+	*/
+	typedef TArrayView<FMemoryView> FSegmentIterator;
+
+
+	/**
+	 * Flags defining the behaviour of a cache entry
+	 */
+	enum ECacheFlags : uint16
+	{
+		// When set blocks are not kept in as part of the global caching.
+		ECacheFlags_NoGlobalCaching = 1 << 1,
+	};
+
+	/**
+	 * A unique id for a named cache entry.
+	 */
+	typedef uint32 FCacheId;
+	
+	/**
+	 * Analysis cache
+	 *
+	 * Allows users of the system to create tracked fixed size blocks associated with a "cache id", created using a string.
+	 * After creating block(s) users are free to write arbitrary data into them. Once all references to the block has been
+	 * dropped the buffer is committed to persistent storage (e.g. cache file).
+	 *
+	 * On consecutive analysis sessions the blocks can be retrieved using the cache id and the block index. Any block can be
+	 * requested at any point in the session in any order.
+	 *
+	 * Analysis is not thread safe and is designed to be used from the analysis thread only.
+	 *
+	 * The container classes \ref TraceServices::TCachedPagedArray or \ref TraceServices::FCachedStringStore are provided as
+	 * user facing abstractions on top of the cache.
+	 *
+	 */
+	class IAnalysisCache
+	{
+	public:
+		virtual ~IAnalysisCache() = default;
+	
+		/**
+		 * Block size in bytes
+		 */
+		constexpr static uint64 BlockSizeBytes = 8 * 1024;
+		
+		/**
+		* Returns the unique id for a named entry. If the entry does not yet exists and new id is reserved.
+		* @param Name String identifier
+		* @param Flags Details about the entry from \ref ECacheFlags.
+		* @return A valid index for the identifier
+		*/
+		virtual FCacheId GetCacheId(const TCHAR* Name, uint16 Flags = 0) = 0;
+
+		/**
+		 * Gets a pointer to the user data a given cache id. The owner of the cache id is able to store arbitrary
+		 * data which does not fit as block data, for example element counts. The user can write to this memory at any
+		 * point during the session. The block is stored to disk at the end of the session.
+		 * @param CacheId Unique cache id
+		 * @return A valid memory view of the user data block, an empty view if the index was not found.
+		 */
+		virtual FMutableMemoryView GetUserData(FCacheId CacheId) = 0;
+			
+		/**
+		* Create blocks of tracked memory. The requested blocks form a contiguous memory range. The block written
+		* to disk when no more references are alive.
+		* @param CacheId Unique cache id
+		* @param Count Number of blocks to use
+		* @return Buffer mapping to the blocks. If the cache id is invalid an empty buffer is returned.
+		*/
+		[[nodiscard]] virtual FSharedBuffer CreateBlocks(FCacheId CacheId, uint32 Count) = 0;
+
+		/**
+		* Load blocks into memory. The requested range of blocks will be loaded into a contiguous memory range.
+		* @param CacheId Unique cache id
+		* @param BlockIndexStart Block index for this cache entry
+		* @param BlockCount Consecutive block indices to load
+		* @return Buffer mapping to the blocks. If the block range is invalid the buffer is empty.
+		*/
+		[[nodiscard]] virtual FSharedBuffer GetBlocks(FCacheId CacheId, uint32 BlockIndexStart, uint32 BlockCount = 1) = 0;
+};
+	
+}
\ No newline at end of file
