Submitting fix for no unity build error on behalf of David.Bollo
#rnx

diff --git a/Engine/Source/Developer/TraceServices/Private/AnalysisCache.cpp b/Engine/Source/Developer/TraceServices/Private/AnalysisCache.cpp
--- a/Engine/Source/Developer/TraceServices/Private/AnalysisCache.cpp
+++ b/Engine/Source/Developer/TraceServices/Private/AnalysisCache.cpp
@@ -6,6 +6,7 @@
 #include "Algo/Find.h"
 #include "Async/MappedFileHandle.h"
 #include "Containers/StringView.h"
+#include "GenericPlatform/GenericPlatformFile.h"
 #include "HAL/PlatformFileManager.h"
 #include "Logging/LogMacros.h"
 #include "Memory/SharedBuffer.h"
