[Insights] Provide a cache instance in the session

Exposes a cache instance in the analysis session interface.

diff --git a/Engine/Source/Developer/TraceServices/Private/AnalysisService.cpp b/Engine/Source/Developer/TraceServices/Private/AnalysisService.cpp
--- a/Engine/Source/Developer/TraceServices/Private/AnalysisService.cpp
+++ b/Engine/Source/Developer/TraceServices/Private/AnalysisService.cpp
@@ -84,6 +84,7 @@ FAnalysisSession::FAnalysisSession(const TCHAR* SessionName, TUniquePtr<UE::Trac
 	, DurationSeconds(0.0)
 	, Allocator(32 << 20)
 	, StringStore(Allocator)
+	, Cache(*Name)
 	, DataStream(MoveTemp(InDataStream))
 {
 }
diff --git a/Engine/Source/Developer/TraceServices/Private/AnalysisServicePrivate.h b/Engine/Source/Developer/TraceServices/Private/AnalysisServicePrivate.h
--- a/Engine/Source/Developer/TraceServices/Private/AnalysisServicePrivate.h
+++ b/Engine/Source/Developer/TraceServices/Private/AnalysisServicePrivate.h
@@ -2,8 +2,8 @@
 
 #pragma once
 
+#include "AnalysisCache.h"
 #include "TraceServices/AnalysisService.h"
-#include "Async/AsyncWork.h"
 #include "Containers/HashTable.h"
 #include "Misc/ScopeLock.h"
 #include "Common/StringStore.h"
@@ -49,6 +49,8 @@ public:
 	virtual const TCHAR* StoreString(const TCHAR* String) override { return StringStore.Store(String); }
 	virtual const TCHAR* StoreString(const FStringView& String) override { return StringStore.Store(String); }
 
+	virtual IAnalysisCache& GetCache() override { return Cache; }
+
 	virtual void BeginRead() const override { Lock.BeginRead(); }
 	virtual void EndRead() const override { Lock.EndRead(); }
 
@@ -73,6 +75,7 @@ private:
 	double DurationSeconds = 0.0;
 	FSlabAllocator Allocator;
 	FStringStore StringStore;
+	FAnalysisCache Cache;
 	TArray<UE::Trace::IAnalyzer*> Analyzers;
 	TArray<IProvider*> Providers;
 	TMap<FName, IProvider*> ProvidersMap;
diff --git a/Engine/Source/Developer/TraceServices/Public/TraceServices/AnalysisService.h b/Engine/Source/Developer/TraceServices/Public/TraceServices/AnalysisService.h
--- a/Engine/Source/Developer/TraceServices/Public/TraceServices/AnalysisService.h
+++ b/Engine/Source/Developer/TraceServices/Public/TraceServices/AnalysisService.h
@@ -8,6 +8,7 @@
 #include "Delegates/Delegate.h"
 #include "Trace/DataStream.h"
 #include "TraceServices/Model/AnalysisSession.h"
+#include "TraceServices/Model/AnalysisCache.h"
 #include "TraceServices/Model/Log.h"
 #include "TraceServices/Model/Bookmarks.h"
 #include "TraceServices/Model/Frames.h"
diff --git a/Engine/Source/Developer/TraceServices/Public/TraceServices/Model/AnalysisSession.h b/Engine/Source/Developer/TraceServices/Public/TraceServices/Model/AnalysisSession.h
--- a/Engine/Source/Developer/TraceServices/Public/TraceServices/Model/AnalysisSession.h
+++ b/Engine/Source/Developer/TraceServices/Public/TraceServices/Model/AnalysisSession.h
@@ -19,7 +19,8 @@ namespace TraceServices
 {
 
 class ILinearAllocator;
-
+class IAnalysisCache;
+	
 class IProvider
 {
 public:
@@ -49,6 +50,8 @@ public:
 	virtual ILinearAllocator& GetLinearAllocator() = 0;
 	virtual const TCHAR* StoreString(const TCHAR* String) = 0;
 	virtual const TCHAR* StoreString(const FStringView& String) = 0;
+
+	virtual IAnalysisCache& GetCache() = 0;
 	
 	virtual void AddAnalyzer(UE::Trace::IAnalyzer* Analyzer) = 0;
 
