[Insights] Persistent symbol storage for module provider

Implements persistent symbol storage using the cache layer. Symbols are stored in a cached paged array and with references to a cached string store for module and symbol names and filenames. Previously saved symbols are loaded at the start of the session and any new symbols that are discovered during the session is stored on close.

diff --git a/Engine/Source/Developer/TraceServices/Private/Model/ModuleProvider.cpp b/Engine/Source/Developer/TraceServices/Private/Model/ModuleProvider.cpp
--- a/Engine/Source/Developer/TraceServices/Private/Model/ModuleProvider.cpp
+++ b/Engine/Source/Developer/TraceServices/Private/Model/ModuleProvider.cpp
@@ -1,11 +1,15 @@
 // Copyright Epic Games, Inc. All Rights Reserved.
 
 #include "ModuleProvider.h"
+#include <atomic>
+#include "Algo/Transform.h"
+#include "Common/CachedPagedArray.h"
+#include "Common/CachedStringStore.h"
+#include "Common/Utils.h"
 #include "Containers/Map.h"
+#include "Misc/PathViews.h"
+#include "TraceServices/Model/AnalysisCache.h"
 #include "TraceServices/Model/AnalysisSession.h"
-#include "TraceServices/Model/Diagnostics.h"
-#include "HAL/RunnableThread.h"
-#include <atomic>
 
 // Choose which symbol resolver to use.
 // If both are enabled and the RAD Syms library fails to initialize, it will fall back to DbgHelp (on Windows).
@@ -45,12 +49,12 @@ struct FModuleEntry
 };
 
 /////////////////////////////////////////////////////////////////////
-template<typename SymbolResolver>
+template<typename SymbolResolverType>
 class TModuleProvider : public IModuleAnalysisProvider
 {
 public:
 								TModuleProvider(IAnalysisSession& Session);
-	virtual 					~TModuleProvider() {}
+	virtual 					~TModuleProvider() override;
 
 	//Query interface
 	const FResolvedSymbol*		GetSymbol(uint64 Address) override;
@@ -60,39 +64,65 @@ public:
 	void						OnModuleLoad(const FStringView& Module, uint64 Base, uint32 Size, const uint8* Checksum, uint32 ChecksumSize) override;
 	void 						OnModuleUnload(uint64 Base) override;
 	void						OnAnalysisComplete() override;
+	void						SaveSymbolsToCache(IAnalysisCache& Cache);
+	void						LoadSymbolsFromCache(IAnalysisCache& Cache);
 
 private:
 
+	struct FSavedSymbol
+	{
+		uint64 Address;
+		uint32 ModuleOffset;
+		uint32 NameOffset;
+		uint32 FileOffset;
+		uint32 Line;
+	};
+	
 	FRWLock						ModulesLock;
 	TPagedArray<FModuleEntry>	Modules;
 
 	FRWLock						SymbolsLock;
-	TPagedArray<FResolvedSymbol> Symbols;
-	TMap<uint64, const FResolvedSymbol*> SymbolCache;
-
+	// Persistently stored symbol strings
+	FCachedStringStore Strings;
+	// Efficient representation of symbols
+	TPagedArray<FResolvedSymbol> SymbolCache;
+	// Lookup table to for symbols
+	TMap<uint64, const FResolvedSymbol*> SymbolCacheLookup;
+	
 	IAnalysisSession&			Session;
 	FString						Platform;
-	SymbolResolver 				Resolver;
+	TUniquePtr<SymbolResolverType>	Resolver;
 };
 
 /////////////////////////////////////////////////////////////////////
-template<typename SymbolProvider>
-TModuleProvider<SymbolProvider>::TModuleProvider(IAnalysisSession& Session)
+template<typename SymbolResolverType>
+TModuleProvider<SymbolResolverType>::TModuleProvider(IAnalysisSession& Session)
 	: Modules(Session.GetLinearAllocator(), 128)
-	, Symbols(Session.GetLinearAllocator(), 1024*1024)
+	, Strings(TEXT("ModuleProvider.Strings"), Session.GetCache())
+	, SymbolCache(Session.GetLinearAllocator(), 1024*1024)
 	, Session(Session)
-	, Resolver(Session)
 {
+	Resolver = TUniquePtr<SymbolResolverType>(new SymbolResolverType(Session));
+	LoadSymbolsFromCache(Session.GetCache());
+}
+
+/////////////////////////////////////////////////////////////////////
+template <typename SymbolResolverType>
+TModuleProvider<SymbolResolverType>::~TModuleProvider()
+{
+	// Delete the resolver in order to flush all pending resolves
+	Resolver.Reset();
+	SaveSymbolsToCache(Session.GetCache());
 }
 
 /////////////////////////////////////////////////////////////////////
-template<typename SymbolProvider>
-const FResolvedSymbol* TModuleProvider<SymbolProvider>::GetSymbol(uint64 Address)
+template<typename SymbolResolverType>
+const FResolvedSymbol* TModuleProvider<SymbolResolverType>::GetSymbol(uint64 Address)
 {
 	{
 		// Attempt to read from the cached symbols.
 		FReadScopeLock _(SymbolsLock);
-		if (const FResolvedSymbol** Entry = SymbolCache.Find(Address))
+		if (const FResolvedSymbol** Entry = SymbolCacheLookup.Find(Address))
 		{
 			return *Entry;
 		}
@@ -102,31 +132,31 @@ const FResolvedSymbol* TModuleProvider<SymbolProvider>::GetSymbol(uint64 Address
 	{
 		// Add a pending entry to our cache
 		FWriteScopeLock _(SymbolsLock);
-		if (SymbolCache.Contains(Address))
+		if (SymbolCacheLookup.Contains(Address))
 		{
-			return SymbolCache[Address];
+			return SymbolCacheLookup[Address];
 		}
-		ResolvedSymbol = &Symbols.EmplaceBack(ESymbolQueryResult::Pending, nullptr, nullptr, nullptr, 0);
-		SymbolCache.Add(Address, ResolvedSymbol);
+		ResolvedSymbol = &SymbolCache.EmplaceBack(ESymbolQueryResult::Pending, nullptr, nullptr, nullptr, 0);
+		SymbolCacheLookup.Add(Address, ResolvedSymbol);
 	}
 
 	// If not in cache yet, queue it up in the resolver
 	check(ResolvedSymbol);
-	Resolver.QueueSymbolResolve(Address, ResolvedSymbol);
+	Resolver->QueueSymbolResolve(Address, ResolvedSymbol);
 
 	return ResolvedSymbol;
 }
 
 /////////////////////////////////////////////////////////////////////
-template<typename SymbolResolver>
-void TModuleProvider<SymbolResolver>::GetStats(FStats* OutStats) const
+template<typename SymbolResolverType>
+void TModuleProvider<SymbolResolverType>::GetStats(FStats* OutStats) const
 {
-	Resolver.GetStats(OutStats);
+	Resolver->GetStats(OutStats);
 }
 
 /////////////////////////////////////////////////////////////////////
-template<typename SymbolProvider>
-void TModuleProvider<SymbolProvider>::OnModuleLoad(const FStringView& Module, uint64 Base, uint32 Size, const uint8* ImageId, uint32 ImageIdSize)
+template<typename SymbolResolverType>
+void TModuleProvider<SymbolResolverType>::OnModuleLoad(const FStringView& Module, uint64 Base, uint32 Size, const uint8* ImageId, uint32 ImageIdSize)
 {
 	if (Module.Len() == 0)
 	{
@@ -141,21 +171,75 @@ void TModuleProvider<SymbolProvider>::OnModuleLoad(const FStringView& Module, ui
 		Modules.EmplaceBack(Name, Path, Base, Size, false);
 	}
 
-	Resolver.QueueModuleLoad(Module, Base, Size, ImageId, ImageIdSize);
+	Resolver->QueueModuleLoad(Module, Base, Size, ImageId, ImageIdSize);
 }
 
 /////////////////////////////////////////////////////////////////////
-template<typename SymbolProvider>
-void TModuleProvider<SymbolProvider>::OnModuleUnload(uint64 Base)
+template<typename SymbolResolverType>
+void TModuleProvider<SymbolResolverType>::OnModuleUnload(uint64 Base)
 {
 	//todo: Find entry, set bLoaded to false
 }
 
 /////////////////////////////////////////////////////////////////////
-template<typename SymbolResolver>
-void TModuleProvider<SymbolResolver>::OnAnalysisComplete()
+template<typename SymbolResolverType>
+void TModuleProvider<SymbolResolverType>::OnAnalysisComplete()
+{
+	Resolver->OnAnalysisComplete();
+}
+
+/////////////////////////////////////////////////////////////////////
+template <typename SymbolResolverType>
+void TModuleProvider<SymbolResolverType>::SaveSymbolsToCache(IAnalysisCache& Cache)
 {
-	Resolver.OnAnalysisComplete();
+	// Create a temporary reverse lookup for symbol -> address
+	TMap<const FResolvedSymbol*, uint64> SymbolReverseLookup;
+	SymbolReverseLookup.Reserve(SymbolCacheLookup.Num());
+	Algo::Transform(SymbolCacheLookup, SymbolReverseLookup, [](const TTuple<uint64, const FResolvedSymbol*>& Pair) { 
+		return TTuple<const FResolvedSymbol*, uint64>(Pair.Value, Pair.Key);
+	});
+
+	// Save new symbols
+	TCachedPagedArray<FSavedSymbol, 1024> SavedSymbols(TEXT("ModuleProvider.Symbols"), Cache);
+	const uint32 NumPreviouslySavedSymbols = SavedSymbols.Num();
+	uint32 NumSavedSymbols = 0;
+	for (uint32 SymbolIndex = SavedSymbols.Num(); SymbolIndex < SymbolCache.Num(); ++SymbolIndex)
+	{
+		const FResolvedSymbol& Symbol = SymbolCache[SymbolIndex];
+		if (Symbol.GetResult() != ESymbolQueryResult::OK)
+		{
+			continue;
+		}
+		const uint64* Address = SymbolReverseLookup.Find(&Symbol);
+		const uint32 ModuleOffset = Strings.Store_GetOffset(Symbol.Module);
+		const uint32 NameOffset = Strings.Store_GetOffset(Symbol.Name);
+		const uint32 FileOffset = Strings.Store_GetOffset(Symbol.File);
+		SavedSymbols.EmplaceBack(FSavedSymbol{*Address, ModuleOffset, NameOffset, FileOffset, Symbol.Line});
+		++NumSavedSymbols;
+	}
+	UE_LOG(LogTraceServices, Display, TEXT("Added %d symbols to the %d previously saved symbols."), NumPreviouslySavedSymbols, NumSavedSymbols);
+}
+
+template <typename SymbolResolverType>
+void TModuleProvider<SymbolResolverType>::LoadSymbolsFromCache(IAnalysisCache& Cache)
+{
+	// Load saved symbols
+	TCachedPagedArray<FSavedSymbol, 1024> SavedSymbols(TEXT("ModuleProvider.Symbols"), Cache);
+	for (uint64 SymbolIndex = 0; SymbolIndex < SavedSymbols.Num(); ++SymbolIndex)
+	{
+		const FSavedSymbol& Symbol = SavedSymbols[SymbolIndex];
+		const TCHAR* Module = Strings.GetStringAtOffset(Symbol.ModuleOffset);
+		const TCHAR* Name = Strings.GetStringAtOffset(Symbol.NameOffset);
+		const TCHAR* File = Strings.GetStringAtOffset(Symbol.FileOffset);
+		if (Module == nullptr || Name == nullptr || File == nullptr)
+		{
+			UE_LOG(LogTraceServices, Warning, TEXT("Found cached symbol (adress %llx) which referenced unknown string."), Symbol.Address);
+			continue;
+		}
+		FResolvedSymbol& Resolved = SymbolCache.EmplaceBack(ESymbolQueryResult::OK, Module, Name, File, Symbol.Line);
+		SymbolCacheLookup.Add(Symbol.Address, &Resolved);
+	}
+	UE_LOG(LogTraceServices, Display, TEXT("Loaded %d symbols from cache."), SymbolCacheLookup.Num());
 }
 
 /////////////////////////////////////////////////////////////////////
