Fix Non Unity Build Errors

diff --git a/Engine/Source/Developer/TraceInsights/Private/Insights/Common/InsightsMenuBuilder.cpp b/Engine/Source/Developer/TraceInsights/Private/Insights/Common/InsightsMenuBuilder.cpp
--- a/Engine/Source/Developer/TraceInsights/Private/Insights/Common/InsightsMenuBuilder.cpp
+++ b/Engine/Source/Developer/TraceInsights/Private/Insights/Common/InsightsMenuBuilder.cpp
@@ -3,6 +3,7 @@
 #include "InsightsMenuBuilder.h"
 
 #include "Framework/Docking/TabManager.h"
+#include "Framework/MultiBox/MultiBoxBuilder.h"
 #include "WorkspaceMenuStructure.h"
 #include "WorkspaceMenuStructureModule.h"
 
diff --git a/Engine/Source/Developer/TraceServices/Private/TraceServicesModule.cpp b/Engine/Source/Developer/TraceServices/Private/TraceServicesModule.cpp
--- a/Engine/Source/Developer/TraceServices/Private/TraceServicesModule.cpp
+++ b/Engine/Source/Developer/TraceServices/Private/TraceServicesModule.cpp
@@ -86,10 +86,7 @@ void FTraceServicesModule::StartupModule()
 	IModularFeatures::Get().RegisterModularFeature(TraceServices::ModuleFeatureName, &NetProfilerModule);
 	IModularFeatures::Get().RegisterModularFeature(TraceServices::ModuleFeatureName, &DiagnosticsModule);
 	IModularFeatures::Get().RegisterModularFeature(TraceServices::ModuleFeatureName, &PlatformEventsModule);
-
-#if EXPERIMENTAL_STATSTRACE_ENABLED
 	IModularFeatures::Get().RegisterModularFeature(TraceServices::ModuleFeatureName, &StatsModule);
-#endif
 
 #if !WITH_EDITOR
 	IModularFeatures::Get().RegisterModularFeature(TraceServices::ModuleFeatureName, &LoadTimeProfilerModule);
@@ -106,10 +103,7 @@ void FTraceServicesModule::ShutdownModule()
 	IModularFeatures::Get().UnregisterModularFeature(TraceServices::ModuleFeatureName, &LoadTimeProfilerModule);
 #endif
 
-#if EXPERIMENTAL_STATSTRACE_ENABLED
 	IModularFeatures::Get().UnregisterModularFeature(TraceServices::ModuleFeatureName, &StatsModule);
-#endif
-
 	IModularFeatures::Get().UnregisterModularFeature(TraceServices::ModuleFeatureName, &PlatformEventsModule);
 	IModularFeatures::Get().UnregisterModularFeature(TraceServices::ModuleFeatureName, &DiagnosticsModule);
 	IModularFeatures::Get().UnregisterModularFeature(TraceServices::ModuleFeatureName, &NetProfilerModule);
diff --git a/Engine/Source/Developer/TraceServices/Public/Common/CachedStringStore.h b/Engine/Source/Developer/TraceServices/Public/Common/CachedStringStore.h
--- a/Engine/Source/Developer/TraceServices/Public/Common/CachedStringStore.h
+++ b/Engine/Source/Developer/TraceServices/Public/Common/CachedStringStore.h
@@ -5,6 +5,7 @@
 #include "Containers/UnrealString.h"
 #include "HAL/UnrealMemory.h"
 #include "Memory/SharedBuffer.h"
+#include "Templates/Tuple.h"
 #include "TraceServices/Model/AnalysisCache.h"
 
 namespace TraceServices
diff --git a/Engine/Source/Runtime/Core/Public/Stats/StatsTrace.h b/Engine/Source/Runtime/Core/Public/Stats/StatsTrace.h
--- a/Engine/Source/Runtime/Core/Public/Stats/StatsTrace.h
+++ b/Engine/Source/Runtime/Core/Public/Stats/StatsTrace.h
@@ -5,10 +5,8 @@
 #include "CoreTypes.h"
 #include "Trace/Config.h"
 
-#define EXPERIMENTAL_STATSTRACE_ENABLED 0
-
 #if !defined(STATSTRACE_ENABLED)
-#if UE_TRACE_ENABLED && STATS && EXPERIMENTAL_STATSTRACE_ENABLED && !UE_BUILD_SHIPPING
+#if UE_TRACE_ENABLED && STATS && !UE_BUILD_SHIPPING
 #define STATSTRACE_ENABLED 1
 #else
 #define STATSTRACE_ENABLED 0
