[Insights] Fix log output

#rnx

diff --git a/Engine/Source/Developer/TraceServices/Private/Model/ModuleProvider.cpp b/Engine/Source/Developer/TraceServices/Private/Model/ModuleProvider.cpp
--- a/Engine/Source/Developer/TraceServices/Private/Model/ModuleProvider.cpp
+++ b/Engine/Source/Developer/TraceServices/Private/Model/ModuleProvider.cpp
@@ -217,7 +217,7 @@ void TModuleProvider<SymbolResolverType>::SaveSymbolsToCache(IAnalysisCache& Cac
 		SavedSymbols.EmplaceBack(FSavedSymbol{*Address, ModuleOffset, NameOffset, FileOffset, Symbol.Line});
 		++NumSavedSymbols;
 	}
-	UE_LOG(LogTraceServices, Display, TEXT("Added %d symbols to the %d previously saved symbols."), NumPreviouslySavedSymbols, NumSavedSymbols);
+	UE_LOG(LogTraceServices, Display, TEXT("Added %d symbols to the %d previously saved symbols."), NumSavedSymbols, NumPreviouslySavedSymbols);
 }
 
 template <typename SymbolResolverType>
