[Insights] Fixed formula used to compute number of symbols in a batch (when parallelizing symbol resolving).

diff --git a/Engine/Source/Developer/TraceServices/Private/Model/SymslibResolver.cpp b/Engine/Source/Developer/TraceServices/Private/Model/SymslibResolver.cpp
--- a/Engine/Source/Developer/TraceServices/Private/Model/SymslibResolver.cpp
+++ b/Engine/Source/Developer/TraceServices/Private/Model/SymslibResolver.cpp
@@ -211,13 +211,17 @@ void FSymslibResolver::DispatchQueuedAddresses()
 	{
 		return;
 	}
-	
+
 	TArray<FQueuedAddress> WorkingSet(ResolveQueue);
-	TasksInFlight += SymbolTasksInParallel;
-	const uint32 Stride = (WorkingSet.Num() + SymbolTasksInParallel - 1) / SymbolTasksInParallel;
-	
+
+	uint32 Stride = (WorkingSet.Num() - 1) / SymbolTasksInParallel + 1;
+	constexpr uint32 MinStride = 4;
+	Stride = FMath::Max(Stride, MinStride);
+	const uint32 ActualSymbolTasksInParallel = (WorkingSet.Num() + Stride - 1) / Stride;
+	TasksInFlight += ActualSymbolTasksInParallel;
+
 	// Use background priority in order to not not interfere with Slate
-	ParallelFor(SymbolTasksInParallel, [this, &WorkingSet, Stride](uint32 Index) {
+	ParallelFor(ActualSymbolTasksInParallel, [this, &WorkingSet, Stride](uint32 Index) {
 		const uint32 StartIndex = Index * Stride;
 		const uint32 EndIndex = FMath::Min(StartIndex + Stride, (uint32)WorkingSet.Num());
 		TArrayView<FQueuedAddress> QueuedWork(&WorkingSet[StartIndex], EndIndex - StartIndex);
