[Insights] Only save cache file when first  used.

Refactored cache file creation so that is all contained inside FFileContents wrapper. File is now created the first time a block is commited which means a cache file will not be created for traces that doesn't use the cache (currently all non-memory traces).

diff --git a/Engine/Source/Developer/TraceServices/Private/AnalysisCache.cpp b/Engine/Source/Developer/TraceServices/Private/AnalysisCache.cpp
--- a/Engine/Source/Developer/TraceServices/Private/AnalysisCache.cpp
+++ b/Engine/Source/Developer/TraceServices/Private/AnalysisCache.cpp
@@ -21,7 +21,60 @@
 DEFINE_LOG_CATEGORY_STATIC(LogAnalysisCache, Log, All);
 
 namespace TraceServices {
-			
+//////////////////////////////////////////////////////////////////////
+FAnalysisCache::FFileContents::FFileContents(const TCHAR* FilePath)
+	: CacheFilePath(FilePath)
+{
+	IPlatformFile& PlatformFile = FPlatformFileManager::Get().GetPlatformFile();
+	
+	// Opening the file we can encounter one of 3 scenarios:
+	// 1. File does not exist, create on first save
+	// 2. File exist, we can read the contents
+	// 3. File exist but we could not open the file for reading. Multiple processes are competing. Put the cache in transient mode.
+	if (PlatformFile.FileExists(*CacheFilePath))
+	{
+		if (const TUniquePtr<IFileHandle> File(PlatformFile.OpenRead(*CacheFilePath)); File.IsValid())
+		{
+			if (const bool Result = Load(); !Result)
+			{
+				UE_LOG(LogAnalysisCache, Error, TEXT("Failed to open cache file table of contents."));
+				//todo: Recover by deleting file?
+			}
+
+			UE_LOG(LogAnalysisCache, VeryVerbose, TEXT("Cache contains %d blocks:"), Blocks.Num());
+			UE_LOG(LogAnalysisCache, VeryVerbose, TEXT("   %10s   %10s   %13s   %13s   %13s"), TEXT("Cache index"), TEXT("Block index"), TEXT("Offset"), TEXT("Uncompressed"), TEXT("Compressed"));
+
+			for (const FFileContents::FBlockEntry& Block : Blocks)
+			{
+				UE_LOG(LogAnalysisCache, VeryVerbose, TEXT("   %10d   %10d   %10d kb   %10d kb   %10d kb"), GetCacheId(Block.BlockKey), GetBlockIndex(Block.BlockKey), Block.Offset / 1024, Block.UncompressedSize / 1024, Block.CompressedSize / 1024);
+			}
+
+			//todo: At this point we can check the file size and if it doesn't match we can trim unknown segments.
+		}
+		else
+		{
+			// Unable to open for read. Most likely this is because another instance is using the file
+			UE_LOG(LogAnalysisCache, Warning,
+			       TEXT("Unable to read the cache file %s, possibly already open in another session. Putting cache in transient mode."),
+			       *CacheFilePath);
+			bTransientMode = true;
+		}
+	}
+}
+
+//////////////////////////////////////////////////////////////////////
+FAnalysisCache::FFileContents::~FFileContents()
+{
+	// Save the table of contents.
+	if (!Blocks.IsEmpty())
+	{
+		if (const bool Result = Save(); !Result )
+		{
+			UE_LOG(LogAnalysisCache, Error, TEXT("Failed to update cache files table of contents."));
+		}
+	}
+}
+
 //////////////////////////////////////////////////////////////////////
 FCacheId FAnalysisCache::FFileContents::GetId(const TCHAR* Name, uint16 Flags)
 {
@@ -63,14 +116,14 @@ FMutableMemoryView FAnalysisCache::FFileContents::GetUserData(FCacheId InId)
 }
 
 //////////////////////////////////////////////////////////////////////
-bool FAnalysisCache::FFileContents::Save(IFileHandle* File)
+bool FAnalysisCache::FFileContents::Save()
 {
-	if (bTransientMode)
+	IFileHandle* File = GetFileHandleForWrite();
+	if (!File)
 	{
 		return true;
 	}
-	
-	check(File);
+
 	File->Seek(0);
 
 	FCbWriter Writer;
@@ -105,14 +158,14 @@ bool FAnalysisCache::FFileContents::Save(IFileHandle* File)
 }
 
 //////////////////////////////////////////////////////////////////////
-bool FAnalysisCache::FFileContents::Load(IFileHandle* File)
+bool FAnalysisCache::FFileContents::Load()
 {
-	if (bTransientMode)
+	IFileHandle* File = GetFileHandleForRead();
+	if (!File)
 	{
 		return true;
 	}
 	
-	check(File);
 	File->Seek(0);
 	
 	FUniqueBuffer Buffer = FUniqueBuffer::Alloc(ReservedSize);
@@ -163,14 +216,14 @@ bool FAnalysisCache::FFileContents::Load(IFileHandle* File)
 }
 
 //////////////////////////////////////////////////////////////////////
-uint64 FAnalysisCache::FFileContents::UpdateBlock(FMemoryView Block, BlockKeyType BlockKey, IFileHandle* File)
+uint64 FAnalysisCache::FFileContents::UpdateBlock(FMemoryView Block, BlockKeyType BlockKey)
 {
-	if (bTransientMode)
+	IFileHandle* File = GetFileHandleForWrite();
+	if (!File)
 	{
 		return 0;
 	}
-	
-	check(File);
+
 	const uint64 EntryIndex = Algo::BinarySearchBy(Blocks, BlockKey, [&](const FBlockEntry& InEntry){return InEntry.BlockKey;});
 	const FIoHash CurrentHash = FIoHash::HashBuffer(Block);
 	
@@ -195,12 +248,14 @@ uint64 FAnalysisCache::FFileContents::UpdateBlock(FMemoryView Block, BlockKeyTyp
 	else
 	{
 		// Write to end of file and add to blocks array
-		File->SeekFromEnd(0);
+		const bool bSeekResult = File->SeekFromEnd(0);
+		check(bSeekResult);
 		const uint64 Offset = File->Tell();
 		check(Offset >= ReservedSize);
 		if(!File->Write((const uint8*) Block.GetData(), Block.GetSize()))
 		{
 			UE_LOG(LogAnalysisCache, Error, TEXT("Failed to update block 0x%x at offset %u kb"), BlockKey, Offset / 1024);
+			return 0;
 		}
 
 		// Insert entry in blocks list and sort array
@@ -214,14 +269,14 @@ uint64 FAnalysisCache::FFileContents::UpdateBlock(FMemoryView Block, BlockKeyTyp
 }
 
 //////////////////////////////////////////////////////////////////////
-uint64 FAnalysisCache::FFileContents::LoadBlock(FMutableMemoryView Block, BlockKeyType BlockKey, IFileHandle* File)
+uint64 FAnalysisCache::FFileContents::LoadBlock(FMutableMemoryView Block, BlockKeyType BlockKey)
 {
-	if (bTransientMode)
+	IFileHandle* File = GetFileHandleForRead();
+	if (!File)
 	{
 		return 0;
 	}
 	
-	check(File);
 	const uint64 EntryIndex = Algo::BinarySearchBy(Blocks, BlockKey, [&](const FBlockEntry& InEntry){return InEntry.BlockKey;});
 	if (EntryIndex == INDEX_NONE)
 	{
@@ -249,82 +304,90 @@ uint64 FAnalysisCache::FFileContents::LoadBlock(FMutableMemoryView Block, BlockK
 }
 
 //////////////////////////////////////////////////////////////////////
-FAnalysisCache::FAnalysisCache(const TCHAR* Name)
-	: Stats()
+IFileHandle* FAnalysisCache::FFileContents::GetFileHandleForWrite()
 {
+	if (bTransientMode)
+	{
+		return nullptr;
+	}
+	
+	if (CacheFileWrite.IsValid())
+	{
+		return CacheFileWrite.Get();
+	}
+
 	IPlatformFile& PlatformFile = FPlatformFileManager::Get().GetPlatformFile();
+	const bool bCreated = !PlatformFile.FileExists(*CacheFilePath);
+	
+	CacheFileWrite = TUniquePtr<IFileHandle>(PlatformFile.OpenWrite(*CacheFilePath, true, true));
+	if (!CacheFileWrite.IsValid())
+	{
+		// Unable to open for write. Most likely this is because another instance is using the file
+		UE_LOG(LogAnalysisCache, Warning,
+			   TEXT("Unable to write to the cache file %s, possibly already open in another session. Putting cache in transient mode."),
+			   *CacheFilePath);
+		bTransientMode = true;
+		return nullptr;
+	}
 
+	if (bCreated)
+	{
+		// Save the table of contents to reserve space
+		Save();
+	}
+	
+	return CacheFileWrite.Get();
+}
+	
+//////////////////////////////////////////////////////////////////////
+IFileHandle* FAnalysisCache::FFileContents::GetFileHandleForRead()
+{
+	if (bTransientMode)
+	{
+		return nullptr;
+	}
+	
+	if (CacheFile.IsValid())
+	{
+		return CacheFile.Get();
+	}
+
+	IPlatformFile& PlatformFile = FPlatformFileManager::Get().GetPlatformFile();
+	CacheFile = TUniquePtr<IFileHandle>(PlatformFile.OpenRead(*CacheFilePath, true));
+
+	if (!CacheFile.IsValid())
+	{
+		// Unable to open for read. Most likely this is because another instance is using the file
+		UE_LOG(LogAnalysisCache, Warning,
+			   TEXT("Unable to read the cache file %s, possibly already open in another session. Putting cache in transient mode."),
+			   *CacheFilePath);
+		bTransientMode = true;
+		return nullptr;
+	}
+	
+	return CacheFile.Get();
+}
+	
+//////////////////////////////////////////////////////////////////////
+FAnalysisCache::FAnalysisCache(const TCHAR* Name)
+	: Stats()
+{
 	// Find the cache file path.
 	// todo: This assumes TraceSessions folder.
 	// todo: This will need to be refined as we move away from files
 	const FString& BaseDirectory = FPaths::ProjectSavedDir();
 	const FString BaseName(FPathViews::GetBaseFilename(Name));
-	CacheFilePath = FPaths::Combine(*BaseDirectory, TEXT("TraceSessions"), *BaseName);
+	FString CacheFilePath = FPaths::Combine(*BaseDirectory, TEXT("TraceSessions"), *BaseName);
 	CacheFilePath += TEXT(".ucache");
 
-	// Opening the file we can encounter one of 4 scenarios:
-	// 1. File does not exist, we create it
-	// 2. File exist, we can read the contents
-	// 3. File exist but we could not open the file for reading. Multiple processes are competing. Put the cache in transient mode.
-	// 4. File does not exist, we cannot create it. Multiple processes are competing to create the file. Put the cache in transient mode.
-	if (!PlatformFile.FileExists(*CacheFilePath))
-	{
-		if (const TUniquePtr<IFileHandle> File(PlatformFile.OpenWrite(*CacheFilePath)); File)
-		{
-			Contents.Save(File.Get());
-		}
-		else
-		{
-			// Unable to open for write. Most likely this is because another instance is racing to create the file
-			UE_LOG(LogAnalysisCache, Warning,
-			       TEXT("Unable to write the cache file %s, possibly already open in another session. Putting cache in transient mode."),
-			       *CacheFilePath);
-			Contents.bTransientMode = true;
-		}
-	}
-	else
-	{
-		if (const TUniquePtr<IFileHandle> File(PlatformFile.OpenRead(*CacheFilePath)); File.IsValid())
-		{
-			if (const bool Result = Contents.Load(File.Get()); !Result)
-			{
-				UE_LOG(LogAnalysisCache, Error, TEXT("Failed to open cache file table of contents."));
-				//todo: Recover by deleting file?
-			}
-
-			UE_LOG(LogAnalysisCache, VeryVerbose, TEXT("Cache contains %d blocks:"), Contents.Blocks.Num());
-			UE_LOG(LogAnalysisCache, VeryVerbose, TEXT("   %10s   %10s   %13s   %13s   %13s"), TEXT("Cache index"), TEXT("Block index"), TEXT("Offset"), TEXT("Uncompressed"), TEXT("Compressed"));
-
-			for (const FFileContents::FBlockEntry& Block : Contents.Blocks)
-			{
-				UE_LOG(LogAnalysisCache, VeryVerbose, TEXT("   %10d   %10d   %10d kb   %10d kb   %10d kb"), GetCacheId(Block.BlockKey), GetBlockIndex(Block.BlockKey), Block.Offset / 1024, Block.UncompressedSize / 1024, Block.CompressedSize / 1024);
-			}
-
-
-			//todo: At this point we can check the file size and if it doesn't match we can trim unknown segments.
-		}
-		else
-		{
-			// Unable to open for read. Most likely this is because another instance is using the file
-			UE_LOG(LogAnalysisCache, Warning,
-			       TEXT("Unable to read the cache file %s, possibly already open in another session. Putting cache in transient mode."),
-			       *CacheFilePath);
-			Contents.bTransientMode = true;
-		}
-	}
+	Contents = MakeUnique<FFileContents>(*CacheFilePath);
 
 	// Build a dictionary of number of blocks per cache id.
-	for (FFileContents::FBlockEntry& Block : Contents.Blocks)
+	for (FFileContents::FBlockEntry& Block : Contents->Blocks)
 	{
 		const uint32 CacheId = GetCacheId(Block.BlockKey);
 		IndexBlockCount.FindOrAdd(CacheId)++;
 	}
-	
-	CacheFile = TUniquePtr<IFileHandle>(PlatformFile.OpenRead(*CacheFilePath, true));
-	CacheFileWrite = TUniquePtr<IFileHandle>(PlatformFile.OpenWrite(*CacheFilePath, true, true));
-
-	// Make sure we have valid file handles or is in transient mode.
-	check(Contents.bTransientMode || (CacheFile.IsValid() && CacheFileWrite.IsValid()));
 }
 
 /////////////////////////////////////////////////////////////////////
@@ -332,27 +395,23 @@ FAnalysisCache::~FAnalysisCache()
 {
 	// Remove all references to cached block, forcing them to write.
 	CachedBlocks.Empty();
-
-	// Save the table of contents.
-	if (const bool Result = Contents.Save(CacheFileWrite.Get()); !Result )
-	{
-		UE_LOG(LogAnalysisCache, Error, TEXT("Failed to update cache files table of contents."));
-	}
-
-	UE_LOG(LogAnalysisCache, Display, TEXT("Closing analysis cache \"%s\". %0.2f Mb read, %0.2f Mb written."), *CacheFilePath,
+	// Delete file contents wrapper
+	Contents.Reset();
+	
+	UE_LOG(LogAnalysisCache, Display, TEXT("Closing analysis cache, %0.2f Mb read, %0.2f Mb written."),
 	       Stats.BytesRead / (1024.f*1024.f), Stats.BytesWritten / (1024.f*1024.f));
 }
 
 /////////////////////////////////////////////////////////////////////
 uint32 FAnalysisCache::GetCacheId(const TCHAR* Name, uint16 Flags)
 {
-	return Contents.GetId(Name, Flags);
+	return Contents->GetId(Name, Flags);
 }
 
 /////////////////////////////////////////////////////////////////////
 FMutableMemoryView FAnalysisCache::GetUserData(FCacheId CacheId)
 {
-	return Contents.GetUserData(CacheId);
+	return Contents->GetUserData(CacheId);
 }
 
 /////////////////////////////////////////////////////////////////////
@@ -374,7 +433,7 @@ FSharedBuffer FAnalysisCache::CreateBlocks(FCacheId CacheId, uint32 BlockCount)
 	IndexBlockCount[CacheId] += BlockCount;
 
 	// Add the blocks into our internal caching mechanism
-	if (!(Contents.GetFlags(CacheId) & ECacheFlags_NoGlobalCaching))
+	if (!(Contents->GetFlags(CacheId) & ECacheFlags_NoGlobalCaching))
 	{
 		check(!CachedBlocks.Contains(BlockKey));
 		CachedBlocks.Add(BlockKey, FSharedBuffer::MakeView(Blocks.GetView(), Blocks));
@@ -411,7 +470,7 @@ FSharedBuffer FAnalysisCache::GetBlocks(FCacheId CacheId, uint32 BlockIndexStart
 	{
 		const BlockKeyType BlockKey = CreateBlockKey(CacheId, BlockIndexStart + Block);
 		const FMutableMemoryView BlockView(BlockBuffer + (Block*IAnalysisCache::BlockSizeBytes), IAnalysisCache::BlockSizeBytes);
-		const uint64 BytesRead = Contents.LoadBlock(BlockView, BlockKey, CacheFile.Get());
+		const uint64 BytesRead = Contents->LoadBlock(BlockView, BlockKey);
 		Stats.BytesRead += BytesRead;
 	}
 	
@@ -422,7 +481,7 @@ FSharedBuffer FAnalysisCache::GetBlocks(FCacheId CacheId, uint32 BlockIndexStart
 	});
 
 	// Add the blocks into our internal caching mechanism
-	if (!(Contents.GetFlags(CacheId) & ECacheFlags_NoGlobalCaching))
+	if (!(Contents->GetFlags(CacheId) & ECacheFlags_NoGlobalCaching))
 	{
 		CachedBlocks.Add(CacheBlockKey, Blocks);
 	}
@@ -439,7 +498,7 @@ void FAnalysisCache::ReleaseBlocks(uint8* BlockBuffer, FCacheId CacheId, uint32
 		void* BlockStart = BlockBuffer + (Block * IAnalysisCache::BlockSizeBytes);
 		const FMemoryView BlockView = FMemoryView(BlockStart, IAnalysisCache::BlockSizeBytes);
 		const BlockKeyType BlockKey = CreateBlockKey(CacheId, BlockIndexStart + Block);
-		const uint64 BytesWritten = Contents.UpdateBlock(BlockView, BlockKey, CacheFileWrite.Get());
+		const uint64 BytesWritten = Contents->UpdateBlock(BlockView, BlockKey);
 		Stats.BytesWritten += BytesWritten;
 	}
 }
diff --git a/Engine/Source/Developer/TraceServices/Private/AnalysisCache.h b/Engine/Source/Developer/TraceServices/Private/AnalysisCache.h
--- a/Engine/Source/Developer/TraceServices/Private/AnalysisCache.h
+++ b/Engine/Source/Developer/TraceServices/Private/AnalysisCache.h
@@ -45,14 +45,18 @@ private:
 		static constexpr uint32 ReservedSize = 512 * 1024;
 		// Number of bytes allocated for user data (per named entry)
 		static constexpr uint32 UserDataSize = 64;
-		
+
+		explicit FFileContents(const TCHAR* FilePath);
+		~FFileContents();
 		FCacheId GetId(const TCHAR* Name, uint16 Flags);
 		FMutableMemoryView GetUserData(FCacheId Id);
 		uint16 GetFlags(FCacheId Id);
-		bool Save(IFileHandle* File);
-		bool Load(IFileHandle* File);
-		uint64 UpdateBlock(FMemoryView Block, BlockKeyType BlockKey, IFileHandle* File);
-		uint64 LoadBlock(FMutableMemoryView Block, BlockKeyType BlockKey, IFileHandle* File);
+		bool Save();
+		bool Load();
+		uint64 UpdateBlock(FMemoryView Block, BlockKeyType BlockKey);
+		uint64 LoadBlock(FMutableMemoryView Block, BlockKeyType BlockKey);
+		IFileHandle* GetFileHandleForWrite();
+		IFileHandle* GetFileHandleForRead();
 
 		struct FIndexEntry
 		{
@@ -72,6 +76,9 @@ private:
 			FIoHash Hash;
 		};
 		
+		FString CacheFilePath;
+		TUniquePtr<IFileHandle> CacheFile;
+		TUniquePtr<IFileHandle> CacheFileWrite;
 		TArray<FIndexEntry> IndexEntries;
 		TArray<FBlockEntry> Blocks;
 		// When this is set blocks or table of contents are never written.
@@ -136,11 +143,8 @@ private:
 
 	constexpr static uint32 BlockAlignment = 1024*4;
 	
-	FString CacheFilePath;
-	FFileContents Contents;
+	TUniquePtr<FFileContents> Contents;
 	FStats Stats;
-	TUniquePtr<IFileHandle> CacheFile;
-	TUniquePtr<IFileHandle> CacheFileWrite;
 	TMap<BlockKeyType,FSharedBuffer> CachedBlocks;
 	TMap<uint32,uint32> IndexBlockCount;
 };
