[Insights] Always return a valid callstack, even if it's not found.

Instead of returning a nullptr, return a  placeholder callstack when an unknown callstack has been requested.

#rnx

diff --git a/Engine/Source/Developer/TraceServices/Private/Model/CallstacksProvider.cpp b/Engine/Source/Developer/TraceServices/Private/Model/CallstacksProvider.cpp
--- a/Engine/Source/Developer/TraceServices/Private/Model/CallstacksProvider.cpp
+++ b/Engine/Source/Developer/TraceServices/Private/Model/CallstacksProvider.cpp
@@ -13,6 +13,9 @@ namespace TraceServices
 
 /////////////////////////////////////////////////////////////////////
 static const FResolvedSymbol GNeverResolveSymbol(ESymbolQueryResult::NotLoaded, nullptr, nullptr, nullptr, 0);
+static const FResolvedSymbol GNotFoundSymbol(ESymbolQueryResult::NotFound, TEXT("Unknown"), nullptr, nullptr, 0);
+static const FStackFrame GNotFoundStackFrame = { 0, &GNotFoundSymbol};
+static const FCallstack GNotFoundCallstack(&GNotFoundStackFrame, 1);
 
 /////////////////////////////////////////////////////////////////////
 #ifdef TRACE_CALLSTACK_STATS
@@ -95,7 +98,7 @@ const FCallstack* FCallstacksProvider::GetCallstack(uint64 CallstackId) const
 {
 	FRWScopeLock ReadLock(EntriesLock, SLT_ReadOnly);
 	const FCallstack* const* FindResult = CallstackEntries.Find(CallstackId);
-	return FindResult ? *FindResult : nullptr;
+	return FindResult ? *FindResult : &GNotFoundCallstack;
 }
 
 /////////////////////////////////////////////////////////////////////
@@ -108,7 +111,7 @@ void FCallstacksProvider::GetCallstacks(const TArrayView<uint64>& CallstackIds,
 	for (uint64 CallstackId : CallstackIds)
 	{
 		const FCallstack* const* FindResult = CallstackEntries.Find(CallstackId);
-		OutCallstacks[OutIdx] = FindResult ? *FindResult : nullptr;
+		OutCallstacks[OutIdx] = FindResult ? *FindResult : &GNotFoundCallstack;
 		OutIdx++;
 	}
 }
