[Insights]
  - TraceServices: Added IModuleService::GetAvailableModulesEx(..), IModuleService::GetEnabledModules(..) and IModule::ShouldBeEnabledByDefault().
  - TraceServices: Fixed analysis to use only the enabled modules (to correctly use the state set by ShouldBeEnabledByDefault + SetModuleEnabled API).
  - SessionInfo: Added "!" in front of name for modules that are registered, but not enabled.
  - Changed FTraceServicesModule::StartupModule() to allow registration of all modules.
  - Changed the Memory, Tasks and LoadTimeProfiler modules to report that should be disabled by default in editor.
  - Fixed Insights unit tests (UE-116865). Re-enabled the LoadTimeProfiler module for unit tests even if it is disabled by default for editor.

diff --git a/Engine/Source/Developer/TraceInsights/Private/Insights/Tests/InsightsTestUtils.cpp b/Engine/Source/Developer/TraceInsights/Private/Insights/Tests/InsightsTestUtils.cpp
--- a/Engine/Source/Developer/TraceInsights/Private/Insights/Tests/InsightsTestUtils.cpp
+++ b/Engine/Source/Developer/TraceInsights/Private/Insights/Tests/InsightsTestUtils.cpp
@@ -5,6 +5,8 @@
 #include "Misc/AutomationTest.h"
 #include "Modules/ModuleManager.h"
 #include "TraceServices/AnalysisService.h"
+#include "TraceServices/ITraceServicesModule.h"
+#include "TraceServices/ModuleService.h"
 
 #include "Insights/IUnrealInsightsModule.h"
 
@@ -13,6 +15,11 @@
 FInsightsTestUtils::FInsightsTestUtils(FAutomationTestBase* InTest) :
 	Test(InTest)
 {
+#if WITH_EDITOR
+	ITraceServicesModule& TraceServicesModule = FModuleManager::LoadModuleChecked<ITraceServicesModule>("TraceServices");
+	TSharedPtr<TraceServices::IModuleService> ModuleService = TraceServicesModule.GetModuleService();
+	ModuleService->SetModuleEnabled(FName("TraceModule_LoadTimeProfiler"), true);
+#endif
 }
 
 ////////////////////////////////////////////////////////////////////////////////////////////////////
diff --git a/Engine/Source/Developer/TraceInsights/Private/Insights/Widgets/SSessionInfoWindow.cpp b/Engine/Source/Developer/TraceInsights/Private/Insights/Widgets/SSessionInfoWindow.cpp
--- a/Engine/Source/Developer/TraceInsights/Private/Insights/Widgets/SSessionInfoWindow.cpp
+++ b/Engine/Source/Developer/TraceInsights/Private/Insights/Widgets/SSessionInfoWindow.cpp
@@ -423,18 +423,30 @@ FText SSessionInfoWindow::GetStatusText() const
 FText SSessionInfoWindow::GetModulesText() const
 {
 	FString ModulesStr;
-	TArray<TraceServices::FModuleInfo> Modules;
+	TArray<TraceServices::FModuleInfoEx> Modules;
 
 	TSharedPtr<TraceServices::IModuleService> ModuleService = FInsightsManager::Get()->GetModuleService();
 	if (ModuleService)
 	{
-		ModuleService->GetAvailableModules(Modules);
+		ModuleService->GetAvailableModulesEx(Modules);
 	}
 
-	for (const TraceServices::FModuleInfo& Module : Modules)
+	bool bIsFirst = true;
+	for (const TraceServices::FModuleInfoEx& Module : Modules)
 	{
-		ModulesStr += Module.DisplayName;
-		ModulesStr += TEXT(", ");
+		if (bIsFirst)
+		{
+			bIsFirst = false;
+		}
+		else
+		{
+			ModulesStr += TEXT(", ");
+		}
+		if (!Module.bIsEnabled)
+		{
+			ModulesStr += TEXT("!");
+		}
+		ModulesStr += Module.Info.DisplayName;
 	}
 
 	return FText::FromString(ModulesStr);
diff --git a/Engine/Source/Developer/TraceServices/Private/ModuleService.cpp b/Engine/Source/Developer/TraceServices/Private/ModuleService.cpp
--- a/Engine/Source/Developer/TraceServices/Private/ModuleService.cpp
+++ b/Engine/Source/Developer/TraceServices/Private/ModuleService.cpp
@@ -27,6 +27,10 @@ void FModuleService::Initialize()
 		FModuleInfo ModuleInfo;
 		Module->GetModuleInfo(ModuleInfo);
 		ModulesMap.Add(ModuleInfo.Name, Module);
+		if (Module->ShouldBeEnabledByDefault())
+		{
+			EnabledModules.Add(Module);
+		}
 	}
 
 	bIsInitialized = true;
@@ -40,16 +44,37 @@ void FModuleService::GetAvailableModules(TArray<FModuleInfo>& OutModules)
 	for (const auto& KV : ModulesMap)
 	{
 		IModule* Module = KV.Value;
-		TArray<const TCHAR*> ModuleLoggers;
-		Module->GetLoggers(ModuleLoggers);
-		if (ModuleLoggers.Num())
-		{
-			FModuleInfo& ModuleInfo = OutModules.AddDefaulted_GetRef();
-			Module->GetModuleInfo(ModuleInfo);
-		}
+		FModuleInfo& ModuleInfo = OutModules.AddDefaulted_GetRef();
+		Module->GetModuleInfo(ModuleInfo);
+	}
+}
+
+void FModuleService::GetAvailableModulesEx(TArray<FModuleInfoEx>& OutModules)
+{
+	FScopeLock Lock(&CriticalSection);
+	Initialize();
+	OutModules.Empty(ModulesMap.Num());
+	for (const auto& KV : ModulesMap)
+	{
+		IModule* Module = KV.Value;
+		FModuleInfoEx& ModuleInfoEx = OutModules.AddDefaulted_GetRef();
+		Module->GetModuleInfo(ModuleInfoEx.Info);
+		ModuleInfoEx.bIsEnabled = EnabledModules.Contains(Module);
 	}
 }
-	
+
+void FModuleService::GetEnabledModules(TArray<FModuleInfo>& OutModules)
+{
+	FScopeLock Lock(&CriticalSection);
+	Initialize();
+	OutModules.Empty(ModulesMap.Num());
+	for (IModule* Module : EnabledModules)
+	{
+		FModuleInfo& ModuleInfo = OutModules.AddDefaulted_GetRef();
+		Module->GetModuleInfo(ModuleInfo);
+	}
+}
+
 void FModuleService::SetModuleEnabled(const FName& ModuleName, bool bEnabled)
 {
 	FScopeLock Lock(&CriticalSection);
@@ -78,9 +103,8 @@ void FModuleService::OnAnalysisBegin(IAnalysisSession& Session)
 {
 	FScopeLock Lock(&CriticalSection);
 	Initialize();
-	for (const auto& KV : ModulesMap)
+	for (IModule* Module : EnabledModules)
 	{
-		IModule* Module = KV.Value;
 		Module->OnAnalysisBegin(Session);
 	}
 }
@@ -126,9 +150,8 @@ void FModuleService::GenerateReports(const IAnalysisSession& Session, const TCHA
 {
 	FScopeLock Lock(&CriticalSection);
 	Initialize();
-	for (const auto& KV : ModulesMap)
+	for (IModule* Module : EnabledModules)
 	{
-		IModule* Module = KV.Value;
 		Module->GenerateReports(Session, CmdLine, OutputDirectory);
 	}
 }
diff --git a/Engine/Source/Developer/TraceServices/Private/ModuleServicePrivate.h b/Engine/Source/Developer/TraceServices/Private/ModuleServicePrivate.h
--- a/Engine/Source/Developer/TraceServices/Private/ModuleServicePrivate.h
+++ b/Engine/Source/Developer/TraceServices/Private/ModuleServicePrivate.h
@@ -18,6 +18,8 @@ class FModuleService
 public:
 	FModuleService();
 	virtual void GetAvailableModules(TArray<FModuleInfo>& OutModules) override;
+	virtual void GetAvailableModulesEx(TArray<FModuleInfoEx>& OutModules) override;
+	virtual void GetEnabledModules(TArray<FModuleInfo>& OutModules) override;
 	virtual void SetModuleEnabled(const FName& ModuleName, bool bEnabled) override;
 	virtual void GenerateReports(const IAnalysisSession& Session, const TCHAR* CmdLine, const TCHAR* OutputDirectory) override;
 	void OnAnalysisBegin(IAnalysisSession& Session);
diff --git a/Engine/Source/Developer/TraceServices/Private/Modules/LoadTimeProfilerModule.h b/Engine/Source/Developer/TraceServices/Private/Modules/LoadTimeProfilerModule.h
--- a/Engine/Source/Developer/TraceServices/Private/Modules/LoadTimeProfilerModule.h
+++ b/Engine/Source/Developer/TraceServices/Private/Modules/LoadTimeProfilerModule.h
@@ -19,6 +19,12 @@ public:
 		return TEXT("loadtimetrace");
 	}
 	virtual void GenerateReports(const IAnalysisSession& Session, const TCHAR* CmdLine, const TCHAR* OutputDirectory) override;
+#if WITH_EDITOR
+	virtual bool ShouldBeEnabledByDefault() const override
+	{
+		return false;
+	}
+#endif
 };
 
 } // namespace TraceServices
diff --git a/Engine/Source/Developer/TraceServices/Private/Modules/MemoryModule.h b/Engine/Source/Developer/TraceServices/Private/Modules/MemoryModule.h
--- a/Engine/Source/Developer/TraceServices/Private/Modules/MemoryModule.h
+++ b/Engine/Source/Developer/TraceServices/Private/Modules/MemoryModule.h
@@ -18,7 +18,15 @@ public:
 	{
 		return nullptr;
 	}
-	virtual void GenerateReports(const IAnalysisSession& Session, const TCHAR* CmdLine, const TCHAR* OutputDirectory) override {}
+	virtual void GenerateReports(const IAnalysisSession& Session, const TCHAR* CmdLine, const TCHAR* OutputDirectory) override
+	{
+	}
+#if WITH_EDITOR
+	virtual bool ShouldBeEnabledByDefault() const override
+	{
+		return false;
+	}
+#endif
 };
 
 } // namespace TraceServices
diff --git a/Engine/Source/Developer/TraceServices/Private/Modules/TasksModule.h b/Engine/Source/Developer/TraceServices/Private/Modules/TasksModule.h
--- a/Engine/Source/Developer/TraceServices/Private/Modules/TasksModule.h
+++ b/Engine/Source/Developer/TraceServices/Private/Modules/TasksModule.h
@@ -14,13 +14,19 @@ namespace TraceServices
 		virtual void OnAnalysisBegin(IAnalysisSession& Session) override;
 		virtual void GetLoggers(TArray<const TCHAR*>& OutLoggers) override
 		{
-		};
+		}
 		virtual const TCHAR* GetCommandLineArgument() override
 		{
 			return nullptr;
 		}
-		virtual void GenerateReports(const IAnalysisSession& Session, const TCHAR* CmdLine, const TCHAR* OutputDirectory) override 
+		virtual void GenerateReports(const IAnalysisSession& Session, const TCHAR* CmdLine, const TCHAR* OutputDirectory) override
+		{
+		}
+#if WITH_EDITOR
+		virtual bool ShouldBeEnabledByDefault() const override
 		{
-		};
+			return false;
+		}
+#endif
 	};
 } // namespace TraceServices
diff --git a/Engine/Source/Developer/TraceServices/Private/TraceServicesModule.cpp b/Engine/Source/Developer/TraceServices/Private/TraceServicesModule.cpp
--- a/Engine/Source/Developer/TraceServices/Private/TraceServicesModule.cpp
+++ b/Engine/Source/Developer/TraceServices/Private/TraceServicesModule.cpp
@@ -88,22 +88,16 @@ void FTraceServicesModule::StartupModule()
 	IModularFeatures::Get().RegisterModularFeature(TraceServices::ModuleFeatureName, &DiagnosticsModule);
 	IModularFeatures::Get().RegisterModularFeature(TraceServices::ModuleFeatureName, &PlatformEventsModule);
 	IModularFeatures::Get().RegisterModularFeature(TraceServices::ModuleFeatureName, &StatsModule);
-
-#if !WITH_EDITOR
 	IModularFeatures::Get().RegisterModularFeature(TraceServices::ModuleFeatureName, &LoadTimeProfilerModule);
 	IModularFeatures::Get().RegisterModularFeature(TraceServices::ModuleFeatureName, &MemoryModule);
 	IModularFeatures::Get().RegisterModularFeature(TraceServices::ModuleFeatureName, &TasksModule);
-#endif
 }
 
 void FTraceServicesModule::ShutdownModule()
 {
-#if !WITH_EDITOR
 	IModularFeatures::Get().UnregisterModularFeature(TraceServices::ModuleFeatureName, &TasksModule);
 	IModularFeatures::Get().UnregisterModularFeature(TraceServices::ModuleFeatureName, &MemoryModule);
 	IModularFeatures::Get().UnregisterModularFeature(TraceServices::ModuleFeatureName, &LoadTimeProfilerModule);
-#endif
-
 	IModularFeatures::Get().UnregisterModularFeature(TraceServices::ModuleFeatureName, &StatsModule);
 	IModularFeatures::Get().UnregisterModularFeature(TraceServices::ModuleFeatureName, &PlatformEventsModule);
 	IModularFeatures::Get().UnregisterModularFeature(TraceServices::ModuleFeatureName, &DiagnosticsModule);
diff --git a/Engine/Source/Developer/TraceServices/Public/TraceServices/ModuleService.h b/Engine/Source/Developer/TraceServices/Public/TraceServices/ModuleService.h
--- a/Engine/Source/Developer/TraceServices/Public/TraceServices/ModuleService.h
+++ b/Engine/Source/Developer/TraceServices/Public/TraceServices/ModuleService.h
@@ -20,6 +20,12 @@ struct FModuleInfo
 	const TCHAR* DisplayName = nullptr;
 };
 
+struct FModuleInfoEx
+{
+	FModuleInfo Info;
+	bool bIsEnabled = false;
+};
+
 class IModule
 	: public IModularFeature
 {
@@ -29,6 +35,7 @@ public:
 	virtual void GetLoggers(TArray<const TCHAR*>& OutLoggers) = 0;
 	virtual void GenerateReports(const IAnalysisSession& Session, const TCHAR* CmdLine, const TCHAR* OutputDirectory) = 0;
 	virtual const TCHAR* GetCommandLineArgument() = 0;
+	virtual bool ShouldBeEnabledByDefault() const { return true; }
 };
 
 class IModuleService
@@ -36,6 +43,8 @@ class IModuleService
 public:
 	virtual ~IModuleService() = default;
 	virtual void GetAvailableModules(TArray<FModuleInfo>& OutModules) = 0;
+	virtual void GetAvailableModulesEx(TArray<FModuleInfoEx>& OutModules) = 0;
+	virtual void GetEnabledModules(TArray<FModuleInfo>& OutModules) = 0;
 	virtual void SetModuleEnabled(const FName& ModuleName, bool bEnabled) = 0;
 	virtual void GenerateReports(const IAnalysisSession& Session, const TCHAR* CmdLine, const TCHAR* OutputDirectory) = 0;
 };
