[Insights] .ucache file is now created in the same folder the session is located

diff --git a/Engine/Source/Developer/TraceInsights/Private/Insights/InsightsManager.cpp b/Engine/Source/Developer/TraceInsights/Private/Insights/InsightsManager.cpp
--- a/Engine/Source/Developer/TraceInsights/Private/Insights/InsightsManager.cpp
+++ b/Engine/Source/Developer/TraceInsights/Private/Insights/InsightsManager.cpp
@@ -7,6 +7,7 @@
 #include "Logging/MessageLog.h"
 #include "MessageLog/Public/MessageLogModule.h"
 #include "Misc/CString.h"
+#include "Misc/Paths.h"
 #include "Modules/ModuleManager.h"
 #include "Templates/UniquePtr.h"
 #include "Trace/StoreClient.h"
@@ -623,12 +624,13 @@ void FInsightsManager::LoadTrace(uint32 InTraceId, bool InAutoQuit)
 		return;
 	}
 
-	FString TraceName;
+	FString TraceName(StoreClient->GetStatus()->GetStoreDir());
 	const UE::Trace::FStoreClient::FTraceInfo* TraceInfo = StoreClient->GetTraceInfoById(InTraceId);
 	if (TraceInfo != nullptr)
 	{
 		FAnsiStringView Name = TraceInfo->GetName();
-		TraceName = FString(Name.Len(), Name.GetData());
+		TraceName = FPaths::Combine(TraceName, FString(Name.Len(), Name.GetData()));
+		TraceName = FPaths::SetExtension(TraceName, ".utrace");
 	}
 
 	Session = AnalysisService->StartAnalysis(*TraceName, MoveTemp(TraceData));
diff --git a/Engine/Source/Developer/TraceServices/Private/AnalysisCache.cpp b/Engine/Source/Developer/TraceServices/Private/AnalysisCache.cpp
--- a/Engine/Source/Developer/TraceServices/Private/AnalysisCache.cpp
+++ b/Engine/Source/Developer/TraceServices/Private/AnalysisCache.cpp
@@ -372,17 +372,15 @@ IFileHandle* FAnalysisCache::FFileContents::GetFileHandleForRead()
 }
 	
 //////////////////////////////////////////////////////////////////////
-FAnalysisCache::FAnalysisCache(const TCHAR* Name)
+FAnalysisCache::FAnalysisCache(const TCHAR* Path)
 	: Stats()
 {
 	// Find the cache file path.
-	// todo: This assumes TraceSessions folder.
 	// todo: This will need to be refined as we move away from files
-	const FString& BaseDirectory = FPaths::ProjectSavedDir();
-	const FString BaseName(FPathViews::GetBaseFilename(Name));
-	FString CacheFilePath = FPaths::Combine(*BaseDirectory, TEXT("TraceSessions"), *BaseName);
-	CacheFilePath += TEXT(".ucache");
 
+	// We expect to receive the full path to the session file
+	FString CacheFilePath(Path);
+	CacheFilePath = FPaths::SetExtension(CacheFilePath, TEXT(".ucache"));
 	Contents = MakeUnique<FFileContents>(*CacheFilePath);
 
 	// Build a dictionary of number of blocks per cache id.
diff --git a/Engine/Source/Developer/TraceServices/Private/AnalysisCache.h b/Engine/Source/Developer/TraceServices/Private/AnalysisCache.h
--- a/Engine/Source/Developer/TraceServices/Private/AnalysisCache.h
+++ b/Engine/Source/Developer/TraceServices/Private/AnalysisCache.h
@@ -19,7 +19,7 @@ class IAnalysisSession;
 class FAnalysisCache : public IAnalysisCache
 {
 public:
-	explicit FAnalysisCache(const TCHAR* Name);
+	explicit FAnalysisCache(const TCHAR* Path);
 	virtual ~FAnalysisCache() override;
 
 	virtual FCacheId GetCacheId(const TCHAR* Name, uint16 Flags) override;
