[TraceServices]
  - Added virtual destructor to all Provider classes.
  - Added "explicit" to several constructors of Provider classes.

diff --git a/Engine/Source/Developer/TraceServices/Private/Model/AllocationsProvider.h b/Engine/Source/Developer/TraceServices/Private/Model/AllocationsProvider.h
--- a/Engine/Source/Developer/TraceServices/Private/Model/AllocationsProvider.h
+++ b/Engine/Source/Developer/TraceServices/Private/Model/AllocationsProvider.h
@@ -199,7 +199,7 @@ private:
 	static constexpr double DefaultTimelineSampleGranularity = 0.0001; // 0.1ms
 
 public:
-	FAllocationsProvider(IAnalysisSession& InSession);
+	explicit FAllocationsProvider(IAnalysisSession& InSession);
 	virtual ~FAllocationsProvider();
 
 	static FName GetName();
diff --git a/Engine/Source/Developer/TraceServices/Private/Model/BookmarksPrivate.h b/Engine/Source/Developer/TraceServices/Private/Model/BookmarksPrivate.h
--- a/Engine/Source/Developer/TraceServices/Private/Model/BookmarksPrivate.h
+++ b/Engine/Source/Developer/TraceServices/Private/Model/BookmarksPrivate.h
@@ -30,7 +30,8 @@ class FBookmarkProvider
 public:
 	static const FName ProviderName;
 
-	FBookmarkProvider(IAnalysisSession& Session);
+	explicit FBookmarkProvider(IAnalysisSession& Session);
+	virtual ~FBookmarkProvider() {}
 
 	FBookmarkSpec& GetSpec(uint64 BookmarkPoint);
 	virtual uint64 GetBookmarkCount() const override { return Bookmarks.Num(); }
diff --git a/Engine/Source/Developer/TraceServices/Private/Model/CallstacksProvider.h b/Engine/Source/Developer/TraceServices/Private/Model/CallstacksProvider.h
--- a/Engine/Source/Developer/TraceServices/Private/Model/CallstacksProvider.h
+++ b/Engine/Source/Developer/TraceServices/Private/Model/CallstacksProvider.h
@@ -17,8 +17,8 @@ class IModuleProvider;
 class FCallstacksProvider : public ICallstacksProvider
 {
 public:
-						FCallstacksProvider(IAnalysisSession& Session);
-	virtual				~FCallstacksProvider() {}
+	explicit FCallstacksProvider(IAnalysisSession& Session);
+	virtual ~FCallstacksProvider() {}
 
 	const FCallstack*	GetCallstack(uint64 CallstackId) const override;
 	void				GetCallstacks(const TArrayView<uint64>& CallstackIds, FCallstack const** OutCallstacks) const override;
diff --git a/Engine/Source/Developer/TraceServices/Private/Model/Channel.h b/Engine/Source/Developer/TraceServices/Private/Model/Channel.h
--- a/Engine/Source/Developer/TraceServices/Private/Model/Channel.h
+++ b/Engine/Source/Developer/TraceServices/Private/Model/Channel.h
@@ -15,6 +15,8 @@ public:
 	static const FName ProviderName;
 
 	FChannelProvider();
+	virtual ~FChannelProvider() {}
+
 	void AnnounceChannel(const TCHAR* ChannelName, uint32 Id, bool bReadOnly);
 	void UpdateChannel(uint32 Id, bool bEnabled);
 
diff --git a/Engine/Source/Developer/TraceServices/Private/Model/ContextSwitchesPrivate.h b/Engine/Source/Developer/TraceServices/Private/Model/ContextSwitchesPrivate.h
--- a/Engine/Source/Developer/TraceServices/Private/Model/ContextSwitchesPrivate.h
+++ b/Engine/Source/Developer/TraceServices/Private/Model/ContextSwitchesPrivate.h
@@ -14,8 +14,9 @@ class FContextSwitchProvider
 {
 public:
 	static const FName ProviderName;
-	FContextSwitchProvider(IAnalysisSession& Session);
-	~FContextSwitchProvider();
+
+	explicit FContextSwitchProvider(IAnalysisSession& Session);
+	virtual ~FContextSwitchProvider();
 
 	int32 GetCoreNumber(uint32 ThreadId, double Time) const override;
 	const void EnumerateContextSwitches(uint32 ThreadId, double StartTime, double EndTime, ContextSwitchCallback Callback) const override;
diff --git a/Engine/Source/Developer/TraceServices/Private/Model/CountersPrivate.h b/Engine/Source/Developer/TraceServices/Private/Model/CountersPrivate.h
--- a/Engine/Source/Developer/TraceServices/Private/Model/CountersPrivate.h
+++ b/Engine/Source/Developer/TraceServices/Private/Model/CountersPrivate.h
@@ -226,7 +226,7 @@ class FCounterProvider
 public:
 	static const FName ProviderName;
 
-	FCounterProvider(IAnalysisSession& Session, IFrameProvider& FrameProvider);
+	explicit FCounterProvider(IAnalysisSession& Session, IFrameProvider& FrameProvider);
 	virtual ~FCounterProvider();
 	virtual uint64 GetCounterCount() const override { return Counters.Num(); }
 	virtual void EnumerateCounters(TFunctionRef<void(uint32, const ICounter&)> Callback) const override;
diff --git a/Engine/Source/Developer/TraceServices/Private/Model/CsvProfilerPrivate.h b/Engine/Source/Developer/TraceServices/Private/Model/CsvProfilerPrivate.h
--- a/Engine/Source/Developer/TraceServices/Private/Model/CsvProfilerPrivate.h
+++ b/Engine/Source/Developer/TraceServices/Private/Model/CsvProfilerPrivate.h
@@ -20,7 +20,7 @@ class FCsvProfilerProvider
 	: public ICsvProfilerProvider
 {
 public:
-	FCsvProfilerProvider(IAnalysisSession& InSession);
+	explicit FCsvProfilerProvider(IAnalysisSession& InSession);
 	virtual ~FCsvProfilerProvider();
 	virtual const IUntypedTable& GetTable(uint32 CaptureIndex) const override { return Captures[CaptureIndex]->Table; }
 	virtual void EnumerateCaptures(TFunctionRef<void(const FCaptureInfo&)> Callback) const override;
diff --git a/Engine/Source/Developer/TraceServices/Private/Model/DiagnosticsPrivate.h b/Engine/Source/Developer/TraceServices/Private/Model/DiagnosticsPrivate.h
--- a/Engine/Source/Developer/TraceServices/Private/Model/DiagnosticsPrivate.h
+++ b/Engine/Source/Developer/TraceServices/Private/Model/DiagnosticsPrivate.h
@@ -10,7 +10,7 @@ namespace TraceServices
 class FDiagnosticsProvider : public IDiagnosticsProvider
 {
 public:
-	FDiagnosticsProvider(IAnalysisSession& Session);
+	explicit FDiagnosticsProvider(IAnalysisSession& Session);
 	virtual ~FDiagnosticsProvider() {}
 
 public:
diff --git a/Engine/Source/Developer/TraceServices/Private/Model/FileActivity.h b/Engine/Source/Developer/TraceServices/Private/Model/FileActivity.h
--- a/Engine/Source/Developer/TraceServices/Private/Model/FileActivity.h
+++ b/Engine/Source/Developer/TraceServices/Private/Model/FileActivity.h
@@ -28,7 +28,9 @@ public:
 	};
 	typedef TIntervalTimeline<FFileActivity*, FTimelineSettings> TimelineInternal;
 
-	FFileActivityProvider(IAnalysisSession& Session);
+	explicit FFileActivityProvider(IAnalysisSession& Session);
+	virtual ~FFileActivityProvider() {}
+
 	virtual void EnumerateFileActivity(TFunctionRef<bool(const FFileInfo&, const Timeline&)> Callback) const override;
 	virtual const ITable<FFileActivity>& GetFileActivityTable() const override;
 	uint32 GetFileIndex(const TCHAR* Path);
diff --git a/Engine/Source/Developer/TraceServices/Private/Model/FramesPrivate.h b/Engine/Source/Developer/TraceServices/Private/Model/FramesPrivate.h
--- a/Engine/Source/Developer/TraceServices/Private/Model/FramesPrivate.h
+++ b/Engine/Source/Developer/TraceServices/Private/Model/FramesPrivate.h
@@ -15,7 +15,8 @@ class FFrameProvider
 public:
 	static const FName ProviderName;
 
-	FFrameProvider(IAnalysisSession& Session);
+	explicit FFrameProvider(IAnalysisSession& Session);
+	virtual ~FFrameProvider() {}
 
 	virtual uint64 GetFrameCount(ETraceFrameType FrameType) const override;
 	virtual void EnumerateFrames(ETraceFrameType FrameType, uint64 Start, uint64 End, TFunctionRef<void(const FFrame&)> Callback) const override;
diff --git a/Engine/Source/Developer/TraceServices/Private/Model/LoadTimeProfilerPrivate.h b/Engine/Source/Developer/TraceServices/Private/Model/LoadTimeProfilerPrivate.h
--- a/Engine/Source/Developer/TraceServices/Private/Model/LoadTimeProfilerPrivate.h
+++ b/Engine/Source/Developer/TraceServices/Private/Model/LoadTimeProfilerPrivate.h
@@ -21,7 +21,9 @@ class FLoadTimeProfilerProvider
 public:
 	typedef TMonotonicTimeline<FLoadTimeProfilerCpuEvent> CpuTimelineInternal;
 
-	FLoadTimeProfilerProvider(IAnalysisSession& Session, ICounterProvider& CounterProvider);
+	explicit FLoadTimeProfilerProvider(IAnalysisSession& Session, ICounterProvider& CounterProvider);
+	virtual ~FLoadTimeProfilerProvider() {}
+
 	virtual uint64 GetTimelineCount() const override { return CpuTimelines.Num(); }
 	virtual bool GetCpuThreadTimelineIndex(uint32 ThreadId, uint32& OutTimelineIndex) const override;
 	virtual bool ReadTimeline(uint32 Index, TFunctionRef<void(const CpuTimeline&)> Callback) const override;
@@ -82,6 +84,8 @@ private:
 		};
 
 		FLoaderFrameCounter(ELoaderFrameCounterType Type, const TPagedArray<FLoaderFrame>& Frames);
+		virtual ~FLoaderFrameCounter() {}
+
 		virtual const TCHAR* GetName() const override;
 		virtual const TCHAR* GetDescription() const override;
 		virtual bool IsFloatingPoint() const override;
diff --git a/Engine/Source/Developer/TraceServices/Private/Model/LogPrivate.h b/Engine/Source/Developer/TraceServices/Private/Model/LogPrivate.h
--- a/Engine/Source/Developer/TraceServices/Private/Model/LogPrivate.h
+++ b/Engine/Source/Developer/TraceServices/Private/Model/LogPrivate.h
@@ -41,7 +41,8 @@ public:
 
 	static const FName ProviderName;
 
-	FLogProvider(IAnalysisSession& Session);
+	explicit FLogProvider(IAnalysisSession& Session);
+	virtual ~FLogProvider() {}
 
 	FLogCategoryInfo& GetCategory(uint64 CategoryPointer);
 	FLogMessageSpec& GetMessageSpec(uint64 LogPoint);
diff --git a/Engine/Source/Developer/TraceServices/Private/Model/MemoryPrivate.h b/Engine/Source/Developer/TraceServices/Private/Model/MemoryPrivate.h
--- a/Engine/Source/Developer/TraceServices/Private/Model/MemoryPrivate.h
+++ b/Engine/Source/Developer/TraceServices/Private/Model/MemoryPrivate.h
@@ -11,7 +11,7 @@ namespace TraceServices
 class FMemoryProvider : public IMemoryProvider
 {
 public:
-	FMemoryProvider(IAnalysisSession& Session);
+	explicit FMemoryProvider(IAnalysisSession& Session);
 	virtual ~FMemoryProvider() {}
 
 	void AddEventSpec(FMemoryTagId Tag, const TCHAR* Name, FMemoryTagId ParentTag);
diff --git a/Engine/Source/Developer/TraceServices/Private/Model/ModuleProvider.cpp b/Engine/Source/Developer/TraceServices/Private/Model/ModuleProvider.cpp
--- a/Engine/Source/Developer/TraceServices/Private/Model/ModuleProvider.cpp
+++ b/Engine/Source/Developer/TraceServices/Private/Model/ModuleProvider.cpp
@@ -53,8 +53,8 @@ template<typename SymbolResolverType>
 class TModuleProvider : public IModuleAnalysisProvider
 {
 public:
-								TModuleProvider(IAnalysisSession& Session);
-	virtual 					~TModuleProvider() override;
+	explicit TModuleProvider(IAnalysisSession& Session);
+	virtual ~TModuleProvider();
 
 	//Query interface
 	const FResolvedSymbol*		GetSymbol(uint64 Address) override;
diff --git a/Engine/Source/Developer/TraceServices/Private/Model/StackSamplesPrivate.h b/Engine/Source/Developer/TraceServices/Private/Model/StackSamplesPrivate.h
--- a/Engine/Source/Developer/TraceServices/Private/Model/StackSamplesPrivate.h
+++ b/Engine/Source/Developer/TraceServices/Private/Model/StackSamplesPrivate.h
@@ -14,8 +14,9 @@ class FStackSampleProvider
 {
 public:
 	static const FName ProviderName;
-	FStackSampleProvider(IAnalysisSession& Session);
-	~FStackSampleProvider();
+
+	explicit FStackSampleProvider(IAnalysisSession& Session);
+	virtual ~FStackSampleProvider();
 
 	const TPagedArray<FStackSample>* GetStackSamples(uint32 ThreadId) const override;
 
diff --git a/Engine/Source/Developer/TraceServices/Private/Model/TasksProfilerPrivate.h b/Engine/Source/Developer/TraceServices/Private/Model/TasksProfilerPrivate.h
--- a/Engine/Source/Developer/TraceServices/Private/Model/TasksProfilerPrivate.h
+++ b/Engine/Source/Developer/TraceServices/Private/Model/TasksProfilerPrivate.h
@@ -17,6 +17,7 @@ namespace TraceServices
 	{
 	public:
 		explicit FTasksProvider(IAnalysisSession& Session);
+		virtual ~FTasksProvider() {}
 
 		void Init(uint32 Version);
 
@@ -98,4 +99,4 @@ namespace TraceServices
 		IEditableCounter* RunningTasksCounter;
 		IEditableCounter* ExecutionTimeCounter;
 	};
-}
\ No newline at end of file
+} // namespace TraceServices
diff --git a/Engine/Source/Developer/TraceServices/Private/Model/ThreadsPrivate.h b/Engine/Source/Developer/TraceServices/Private/Model/ThreadsPrivate.h
--- a/Engine/Source/Developer/TraceServices/Private/Model/ThreadsPrivate.h
+++ b/Engine/Source/Developer/TraceServices/Private/Model/ThreadsPrivate.h
@@ -16,8 +16,10 @@ class FThreadProvider
 {
 public:
 	static const FName ProviderName;
-	FThreadProvider(IAnalysisSession& Session);
-	~FThreadProvider();
+
+	explicit FThreadProvider(IAnalysisSession& Session);
+	virtual ~FThreadProvider();
+
 	void AddGameThread(uint32 Id);
 	void AddThread(uint32 Id, const TCHAR* Name, EThreadPriority Priority);
 	void SetThreadPriority(uint32 Id, EThreadPriority Priority);
diff --git a/Engine/Source/Developer/TraceServices/Private/Model/TimingProfilerPrivate.h b/Engine/Source/Developer/TraceServices/Private/Model/TimingProfilerPrivate.h
--- a/Engine/Source/Developer/TraceServices/Private/Model/TimingProfilerPrivate.h
+++ b/Engine/Source/Developer/TraceServices/Private/Model/TimingProfilerPrivate.h
@@ -20,7 +20,7 @@ class FTimingProfilerProvider
 public:
 	typedef TMonotonicTimeline<FTimingProfilerEvent> TimelineInternal;
 
-	FTimingProfilerProvider(IAnalysisSession& InSession);
+	explicit FTimingProfilerProvider(IAnalysisSession& InSession);
 	virtual ~FTimingProfilerProvider();
 	uint32 AddCpuTimer(const TCHAR* Name);
 	uint32 AddGpuTimer(FStringView Name);
diff --git a/Engine/Source/Developer/TraceServices/Public/TraceServices/Model/AllocationsProvider.h b/Engine/Source/Developer/TraceServices/Public/TraceServices/Model/AllocationsProvider.h
--- a/Engine/Source/Developer/TraceServices/Public/TraceServices/Model/AllocationsProvider.h
+++ b/Engine/Source/Developer/TraceServices/Public/TraceServices/Model/AllocationsProvider.h
@@ -123,6 +123,8 @@ public:
 	typedef UPTRINT FQueryHandle;
 
 public:
+	virtual ~IAllocationsProvider() = default;
+
 	virtual void BeginEdit() const = 0;
 	virtual void EndEdit() const = 0;
 	virtual void BeginRead() const = 0;
diff --git a/Engine/Source/Developer/TraceServices/Public/TraceServices/Model/CsvProfiler.h b/Engine/Source/Developer/TraceServices/Public/TraceServices/Model/CsvProfiler.h
--- a/Engine/Source/Developer/TraceServices/Public/TraceServices/Model/CsvProfiler.h
+++ b/Engine/Source/Developer/TraceServices/Public/TraceServices/Model/CsvProfiler.h
@@ -22,6 +22,7 @@ class ICsvProfilerProvider
 	: public IProvider
 {
 public:
+	virtual ~ICsvProfilerProvider() = default;
 	virtual void EnumerateCaptures(TFunctionRef<void(const FCaptureInfo&)> Callback) const = 0;
 	virtual const IUntypedTable& GetTable(uint32 CaptureId) const = 0;
 };
diff --git a/Engine/Source/Developer/TraceServices/Public/TraceServices/Model/Modules.h b/Engine/Source/Developer/TraceServices/Public/TraceServices/Model/Modules.h
--- a/Engine/Source/Developer/TraceServices/Public/TraceServices/Model/Modules.h
+++ b/Engine/Source/Developer/TraceServices/Public/TraceServices/Model/Modules.h
@@ -80,7 +80,9 @@ public:
 		uint32 SymbolsResolved;
 		uint32 SymbolsFailed;
 	};
-	
+
+	virtual ~IModuleProvider() = default;
+
 	/** Queries the name of the symbol at address. This function returns immedately, 
 	 * but the lookup is async. See \ref FResolvedSymbol for details. It assumed that 
 	 * all calls to this function happens before analysis has ended.
diff --git a/Engine/Source/Developer/TraceServices/Public/TraceServices/Model/NetProfiler.h b/Engine/Source/Developer/TraceServices/Public/TraceServices/Model/NetProfiler.h
--- a/Engine/Source/Developer/TraceServices/Public/TraceServices/Model/NetProfiler.h
+++ b/Engine/Source/Developer/TraceServices/Public/TraceServices/Model/NetProfiler.h
@@ -164,6 +164,8 @@ struct FNetProfilerAggregatedStats
 class INetProfilerProvider : public IProvider
 {
 public:
+	virtual ~INetProfilerProvider() = default;
+
 	// Return the version reported in the trace
 	// A return value of 0 indicates no network trace data
 	virtual uint32 GetNetTraceVersion() const = 0;
