Helper method to read an array-type field, falling back to the event attachment if the field doesn't exist.

#rnx

diff --git a/Engine/Source/Developer/TraceServices/Public/TraceServices/Utils.h b/Engine/Source/Developer/TraceServices/Public/TraceServices/Utils.h
--- a/Engine/Source/Developer/TraceServices/Public/TraceServices/Utils.h
+++ b/Engine/Source/Developer/TraceServices/Public/TraceServices/Utils.h
@@ -3,6 +3,7 @@
 #pragma once
 
 #include "Trace/Analyzer.h"
+#include "Containers/ArrayView.h"
 #include "Containers/UnrealString.h"
 
 namespace TraceServices {
@@ -23,5 +24,19 @@ inline FString LegacyAttachmentString(
 	return Out;
 }
 
+inline TArrayView<const uint8> LegacyAttachmentArray(
+	const ANSICHAR* FieldName,
+	const UE::Trace::IAnalyzer::FOnEventContext& Context)
+{
+	TArrayView<const uint8> Out = Context.EventData.GetArrayView<uint8>(FieldName);
+	if (Out.GetData() == nullptr)
+	{
+		Out = TArrayView<const uint8>(
+			Context.EventData.GetAttachment(),
+			Context.EventData.GetAttachmentSize());
+	}
+	return Out;
+}
+
 } // namespace FTraceAnalyzerUtils
 } // namespace TraceServices
