Load and platform-file tracing migrated away from attachments.

#rnx

diff --git a/Engine/Source/Developer/TraceServices/Private/Analyzers/LoadTimeTraceAnalysis.cpp b/Engine/Source/Developer/TraceServices/Private/Analyzers/LoadTimeTraceAnalysis.cpp
--- a/Engine/Source/Developer/TraceServices/Private/Analyzers/LoadTimeTraceAnalysis.cpp
+++ b/Engine/Source/Developer/TraceServices/Private/Analyzers/LoadTimeTraceAnalysis.cpp
@@ -114,10 +114,10 @@ bool FAsyncLoadingTraceAnalyzer::OnEvent(uint16 RouteId, EStyle Style, const FOn
 		if (AsyncPackage)
 		{
 			FAnalysisSessionEditScope _(Session);
-			if (EventData.GetAttachmentSize())
+			FString PackageName = FTraceAnalyzerUtils::LegacyAttachmentString<TCHAR>("Name", Context);
+			if (PackageName.Len())
 			{
-				const TCHAR* PackageName = reinterpret_cast<const TCHAR*>(EventData.GetAttachment());
-				AsyncPackage->PackageInfo->Name = Session.StoreString(PackageName);
+				AsyncPackage->PackageInfo->Name = Session.StoreString(*PackageName);
 			}
 			FPackageSummaryInfo& Summary = AsyncPackage->PackageInfo->Summary;
 			Summary.TotalHeaderSize = EventData.GetValue<uint32>("TotalHeaderSize");
@@ -281,11 +281,22 @@ bool FAsyncLoadingTraceAnalyzer::OnEvent(uint16 RouteId, EStyle Style, const FOn
 	}
 	case RouteId_BeginRequestGroup:
 	{
+		FString FormatString;
+		const uint8* FormatArgs;
+		if (EventData.GetString("Format", FormatString))
+		{
+			FormatArgs = EventData.GetArrayView<uint8>("FormatArgs").GetData();
+		}
+		else
+		{
+			const TCHAR* FormatStringPtr = reinterpret_cast<const TCHAR*>(EventData.GetAttachment());
+			FormatArgs = EventData.GetAttachment() + (FCString::Strlen(FormatStringPtr) + 1) * sizeof(TCHAR);
+			FormatString = FormatStringPtr;
+		}
+
 		FAnalysisSessionEditScope _(Session);
 		TSharedRef<FRequestGroupState> GroupState = MakeShared<FRequestGroupState>();
-		const TCHAR* FormatString = reinterpret_cast<const TCHAR*>(EventData.GetAttachment());
-		const uint8* FormatArgs = EventData.GetAttachment() + (FCString::Strlen(FormatString) + 1) * sizeof(TCHAR);
-		FFormatArgsHelper::Format(FormatBuffer, FormatBufferSize - 1, TempBuffer, FormatBufferSize - 1, FormatString, FormatArgs);
+		FFormatArgsHelper::Format(FormatBuffer, FormatBufferSize - 1, TempBuffer, FormatBufferSize - 1, *FormatString, FormatArgs);
 		GroupState->Name = Session.StoreString(FormatBuffer);
 		uint32 ThreadId = FTraceAnalyzerUtils::GetThreadIdField(Context);
 		FThreadState& ThreadState = GetThreadState(ThreadId);
diff --git a/Engine/Source/Developer/TraceServices/Private/Analyzers/PlatformFileTraceAnalysis.cpp b/Engine/Source/Developer/TraceServices/Private/Analyzers/PlatformFileTraceAnalysis.cpp
--- a/Engine/Source/Developer/TraceServices/Private/Analyzers/PlatformFileTraceAnalysis.cpp
+++ b/Engine/Source/Developer/TraceServices/Private/Analyzers/PlatformFileTraceAnalysis.cpp
@@ -57,8 +57,8 @@ bool FPlatformFileTraceAnalyzer::OnEvent(uint16 RouteId, EStyle Style, const FOn
 			PLATFORMFILETRACE_WARNING("BeginOpen: Duplicated event!?");
 		}
 		#endif
-		const TCHAR* FileName = reinterpret_cast<const TCHAR*>(EventData.GetAttachment());
-		uint32 FileIndex = FileActivityProvider.GetFileIndex(FileName);
+		const FString FileName = FTraceAnalyzerUtils::LegacyAttachmentString<TCHAR>("Path", Context);
+		uint32 FileIndex = FileActivityProvider.GetFileIndex(*FileName);
 		FPendingActivity& Open = PendingOpenMap.Add(ThreadId);
 		Open.ActivityIndex = FileActivityProvider.BeginActivity(FileIndex, FileActivityType_Open, ThreadId, 0, 0, Time);
 		Open.FileIndex = FileIndex;
diff --git a/Engine/Source/Runtime/Core/Private/ProfilingDebugging/PlatformFileTrace.cpp b/Engine/Source/Runtime/Core/Private/ProfilingDebugging/PlatformFileTrace.cpp
--- a/Engine/Source/Runtime/Core/Private/ProfilingDebugging/PlatformFileTrace.cpp
+++ b/Engine/Source/Runtime/Core/Private/ProfilingDebugging/PlatformFileTrace.cpp
@@ -20,6 +20,7 @@ UE_TRACE_CHANNEL(FileChannel)
 
 UE_TRACE_EVENT_BEGIN(PlatformFile, BeginOpen)
 	UE_TRACE_EVENT_FIELD(uint64, Cycle)
+	UE_TRACE_EVENT_FIELD(UE::Trace::WideString, Path)
 UE_TRACE_EVENT_END()
 
 UE_TRACE_EVENT_BEGIN(PlatformFile, EndOpen)
@@ -87,9 +88,9 @@ namespace
 void FPlatformFileTrace::BeginOpen(const TCHAR* Path)
 {
 	uint32 PathSize = (FCString::Strlen(Path) + 1) * sizeof(TCHAR);
-	UE_TRACE_LOG(PlatformFile, BeginOpen, FileChannel, PathSize)
+	UE_TRACE_LOG(PlatformFile, BeginOpen, FileChannel)
 		<< BeginOpen.Cycle(FPlatformTime::Cycles64())
-		<< BeginOpen.Attachment(Path, PathSize);
+		<< BeginOpen.Path(Path, PathSize);
 }
 
 void FPlatformFileTrace::EndOpen(uint64 FileHandle)
diff --git a/Engine/Source/Runtime/Core/Private/Serialization/LoadTimeTrace.cpp b/Engine/Source/Runtime/Core/Private/Serialization/LoadTimeTrace.cpp
--- a/Engine/Source/Runtime/Core/Private/Serialization/LoadTimeTrace.cpp
+++ b/Engine/Source/Runtime/Core/Private/Serialization/LoadTimeTrace.cpp
@@ -11,6 +11,8 @@
 UE_TRACE_CHANNEL_DEFINE(LoadTimeChannel)
 
 UE_TRACE_EVENT_BEGIN(LoadTime, BeginRequestGroup)
+	UE_TRACE_EVENT_FIELD(UE::Trace::WideString, Format)
+	UE_TRACE_EVENT_FIELD(uint8[], FormatArgs)
 UE_TRACE_EVENT_END()
 
 UE_TRACE_EVENT_BEGIN(LoadTime, EndRequestGroup)
@@ -23,14 +25,9 @@ FLoadTimeProfilerTrace::FRequestGroupScope::~FRequestGroupScope()
 
 void FLoadTimeProfilerTrace::FRequestGroupScope::OutputBegin()
 {
-	uint16 FormatStringSize = (uint16)((FCString::Strlen(FormatString) + 1) * sizeof(TCHAR));
-	auto Attachment = [this, FormatStringSize](uint8* Out)
-	{
-		memcpy(Out, FormatString, FormatStringSize);
-		memcpy(Out + FormatStringSize, FormatArgsBuffer, FormatArgsSize);
-	};
-	UE_TRACE_LOG(LoadTime, BeginRequestGroup, LoadTimeChannel, FormatStringSize + FormatArgsSize)
-		<< BeginRequestGroup.Attachment(Attachment);
+	UE_TRACE_LOG(LoadTime, BeginRequestGroup, LoadTimeChannel)
+		<< BeginRequestGroup.Format(FormatString)
+		<< BeginRequestGroup.FormatArgs(FormatArgsBuffer, FormatArgsSize);
 }
 
 
diff --git a/Engine/Source/Runtime/CoreUObject/Private/Serialization/LoadTimeTrace.cpp b/Engine/Source/Runtime/CoreUObject/Private/Serialization/LoadTimeTrace.cpp
--- a/Engine/Source/Runtime/CoreUObject/Private/Serialization/LoadTimeTrace.cpp
+++ b/Engine/Source/Runtime/CoreUObject/Private/Serialization/LoadTimeTrace.cpp
@@ -31,6 +31,7 @@ UE_TRACE_EVENT_BEGIN(LoadTime, PackageSummary)
 	UE_TRACE_EVENT_FIELD(uint32, TotalHeaderSize)
 	UE_TRACE_EVENT_FIELD(uint32, ImportCount)
 	UE_TRACE_EVENT_FIELD(uint32, ExportCount)
+	UE_TRACE_EVENT_FIELD(UE::Trace::WideString, Name)
 UE_TRACE_EVENT_END()
 
 UE_TRACE_EVENT_BEGIN(LoadTime, BeginCreateExport)
@@ -169,14 +170,14 @@ void FLoadTimeProfilerTracePrivate::OutputDestroyAsyncPackage(const void* AsyncP
 void FLoadTimeProfilerTracePrivate::OutputPackageSummary(const void* AsyncPackage, const FName& PackageName, uint32 TotalHeaderSize, uint32 ImportCount, uint32 ExportCount)
 {
 	TCHAR Buffer[FName::StringBufferSize];
-	uint16 NameSize = (PackageName.ToString(Buffer) + 1) * sizeof(TCHAR);
+	uint32 NameLength = PackageName.ToString(Buffer);
 
-	UE_TRACE_LOG(LoadTime, PackageSummary, LoadTimeChannel, NameSize)
+	UE_TRACE_LOG(LoadTime, PackageSummary, LoadTimeChannel)
 		<< PackageSummary.AsyncPackage(AsyncPackage)
 		<< PackageSummary.TotalHeaderSize(TotalHeaderSize)
 		<< PackageSummary.ImportCount(ImportCount)
 		<< PackageSummary.ExportCount(ExportCount)
-		<< PackageSummary.Attachment(Buffer, NameSize);
+		<< PackageSummary.Name(Buffer, NameLength);
 }
 
 void FLoadTimeProfilerTracePrivate::OutputAsyncPackageRequestAssociation(const void* AsyncPackage, uint64 RequestId)
