Fix non-unity error

diff --git a/Engine/Source/Developer/TraceServices/Private/Model/SymslibResolver.cpp b/Engine/Source/Developer/TraceServices/Private/Model/SymslibResolver.cpp
--- a/Engine/Source/Developer/TraceServices/Private/Model/SymslibResolver.cpp
+++ b/Engine/Source/Developer/TraceServices/Private/Model/SymslibResolver.cpp
@@ -10,6 +10,7 @@
 #include "GenericPlatform/GenericPlatformFile.h"
 #include "HAL/PlatformFileManager.h"
 #include "Logging/LogMacros.h"
+#include "Misc/Paths.h"
 #include "Misc/PathViews.h"
 #include "Misc/ScopeLock.h"
 #include "Misc/ScopeRWLock.h"
