[Insights] Fix TraceThreadId being assigned an incorrect SystemThreadId in ContextSwitchProvider

diff --git a/Engine/Source/Developer/TraceServices/Private/Model/ContextSwitches.cpp b/Engine/Source/Developer/TraceServices/Private/Model/ContextSwitches.cpp
--- a/Engine/Source/Developer/TraceServices/Private/Model/ContextSwitches.cpp
+++ b/Engine/Source/Developer/TraceServices/Private/Model/ContextSwitches.cpp
@@ -318,15 +318,25 @@ void FContextSwitchesProvider::Add(uint32 SystemThreadId, double Start, double E
 
 void FContextSwitchesProvider::AddThreadInfo(uint32 ThreadId, uint32 SystemThreadId)
 {
-	if (!TraceToSystemThreadIdMap.Contains(ThreadId))
+	if (SystemThreadId == 0)
 	{
-		TraceToSystemThreadIdMap.Add(ThreadId, SystemThreadId);
+		// At the start of a session some threads might be received with a SystemThreadId of 0.
+		return;
+	}
+
+	uint32* OldSystemThreadId = TraceToSystemThreadIdMap.Find(ThreadId);
+	if (OldSystemThreadId)
+	{
+		ensure(*OldSystemThreadId == SystemThreadId);
 	}
+	TraceToSystemThreadIdMap.Add(ThreadId, SystemThreadId);
 
-	if (!SystemToTraceThreadIdMap.Contains(SystemThreadId))
+	uint32* OldTraceThreadId = SystemToTraceThreadIdMap.Find(SystemThreadId);
+	if (OldTraceThreadId)
 	{
-		SystemToTraceThreadIdMap.Add(SystemThreadId, ThreadId);
+		ensure(*OldTraceThreadId == ThreadId);
 	}
+	SystemToTraceThreadIdMap.Add(SystemThreadId, ThreadId);
 }
 
 const IContextSwitchesProvider* ReadContextSwitchesProvider(const IAnalysisSession& Session)
