Fix memory corruption when a cached TagDesc's memory is moved.

diff --git a/Engine/Source/Developer/TraceServices/Private/Model/Memory.cpp b/Engine/Source/Developer/TraceServices/Private/Model/Memory.cpp
--- a/Engine/Source/Developer/TraceServices/Private/Model/Memory.cpp
+++ b/Engine/Source/Developer/TraceServices/Private/Model/Memory.cpp
@@ -7,6 +7,7 @@ namespace TraceServices
 
 FMemoryProvider::FMemoryProvider(IAnalysisSession& InSession)
 	: Session(InSession)
+	, TagDescsPool(Session.GetLinearAllocator(), 1024)
 	, TagsSerial(0)
 {
 	TagDescs.Reserve(256);
@@ -26,7 +27,11 @@ void FMemoryProvider::AddEventSpec(FMemoryTagId TagId, const TCHAR* Name, FMemor
 		return;
 	}
 
-	TagDescs.Add(TagId, FMemoryTagInfo{ TagId, FString(Name), ParentTagId, 0 });
+	// Do not remove or insert from TagDescsPool, cached references to this memory are stored and inserts / removes will
+	// result in dangerous memory moves
+	FMemoryTagInfo& NewMemoryInfo = TagDescsPool.EmplaceBack(FMemoryTagInfo{ TagId, FString(Name), ParentTagId, 0 });
+
+	TagDescs.Add(TagId, &NewMemoryInfo);
 	++TagsSerial;
 }
 
@@ -96,7 +101,10 @@ void FMemoryProvider::AddTagSnapshot(FMemoryTrackerId TrackerId, double Time, co
 		if (!TagSamples->TagPtr)
 		{
 			// Cache pointer to FMemoryTagInfo to avoid further lookups for this tag.
-			TagSamples->TagPtr = TagDescs.Find(TagId);
+			if (FMemoryTagInfo** TagInfo = TagDescs.Find(TagId))
+			{
+				TagSamples->TagPtr = *TagInfo;
+			}
 		}
 
 		if (TagSamples->TagPtr)
@@ -125,13 +133,17 @@ void FMemoryProvider::EnumerateTags(TFunctionRef<void(const FMemoryTagInfo&)> Ca
 {
 	for (auto& Tag : TagDescs)
 	{
-		Callback(Tag.Value);
+		Callback(*Tag.Value);
 	}
 }
 
 const FMemoryTagInfo* FMemoryProvider::GetTag(FMemoryTagId TagId) const
 {
-	return TagDescs.Find(TagId);
+	if (FMemoryTagInfo* const * TagInfo = TagDescs.Find(TagId))
+	{
+		return *TagInfo;
+	}
+	return nullptr;
 }
 
 uint32 FMemoryProvider::GetTrackerCount() const
diff --git a/Engine/Source/Developer/TraceServices/Private/Model/MemoryPrivate.h b/Engine/Source/Developer/TraceServices/Private/Model/MemoryPrivate.h
--- a/Engine/Source/Developer/TraceServices/Private/Model/MemoryPrivate.h
+++ b/Engine/Source/Developer/TraceServices/Private/Model/MemoryPrivate.h
@@ -68,7 +68,8 @@ private:
 	};
 
 	IAnalysisSession& Session;
-	TMap<FMemoryTagId, FMemoryTagInfo> TagDescs;
+	TMap<FMemoryTagId, FMemoryTagInfo*> TagDescs;
+	TPagedArray<FMemoryTagInfo> TagDescsPool;
 	TMap<FMemoryTrackerId, FMemoryTrackerInfo> TrackerDescs;
 	TArray<FTrackerData> Trackers;
 	uint32 TagsSerial;
