[Insights] Find previous functionality for the quick find widget #rb none

diff --git a/Engine/Source/Developer/TraceInsights/Private/Insights/ViewModels/QuickFind.h b/Engine/Source/Developer/TraceInsights/Private/Insights/ViewModels/QuickFind.h
--- a/Engine/Source/Developer/TraceInsights/Private/Insights/ViewModels/QuickFind.h
+++ b/Engine/Source/Developer/TraceInsights/Private/Insights/ViewModels/QuickFind.h
@@ -41,6 +41,16 @@ private:
 	FOnFindNextEvent OnFindNextEvent;
 
 	////////////////////////////////////////////////////////////////////////////////////////////////////
+	// OnFindPreviousEvent
+public:
+	/** The event to execute when the user selects the "Find Previous" option. */
+	DECLARE_MULTICAST_DELEGATE(FOnFindPreviousEvent);
+	FOnFindPreviousEvent& GetOnFindPreviousEvent() { return OnFindPreviousEvent; }
+
+private:
+	FOnFindPreviousEvent OnFindPreviousEvent;
+
+	////////////////////////////////////////////////////////////////////////////////////////////////////
 
 private:
 	TSharedPtr<FFilterConfigurator> FilterConfigurator;
diff --git a/Engine/Source/Developer/TraceInsights/Private/Insights/ViewModels/ThreadTimingTrack.cpp b/Engine/Source/Developer/TraceInsights/Private/Insights/ViewModels/ThreadTimingTrack.cpp
--- a/Engine/Source/Developer/TraceInsights/Private/Insights/ViewModels/ThreadTimingTrack.cpp
+++ b/Engine/Source/Developer/TraceInsights/Private/Insights/ViewModels/ThreadTimingTrack.cpp
@@ -1470,11 +1470,20 @@ bool FThreadTimingTrack::FindTimingProfilerEvent(const FTimingEventSearchParamet
 
 					TimingProfilerProvider.ReadTimeline(GetTimelineIndex(), [&InContext](const TraceServices::ITimingProfilerProvider::Timeline& Timeline)
 					{
-						Timeline.EnumerateEvents(InContext.GetParameters().StartTime, InContext.GetParameters().EndTime, [&InContext](double EventStartTime, double EventEndTime, uint32 EventDepth, const TraceServices::FTimingProfilerEvent& Event)
+						auto Callback = [&InContext](double EventStartTime, double EventEndTime, uint32 EventDepth, const TraceServices::FTimingProfilerEvent& Event)
 						{
 							InContext.Check(EventStartTime, EventEndTime, EventDepth, Event);
 							return InContext.ShouldContinueSearching() ? TraceServices::EEventEnumerate::Continue : TraceServices::EEventEnumerate::Stop;
-						});
+						};
+
+						if (InContext.GetParameters().SearchDirection == FTimingEventSearchParameters::ESearchDirection::Forward)
+						{
+							Timeline.EnumerateEvents(InContext.GetParameters().StartTime, InContext.GetParameters().EndTime, Callback);
+						}
+						else
+						{
+							Timeline.EnumerateEventsBackwards(InContext.GetParameters().EndTime, InContext.GetParameters().StartTime, Callback);
+						}
 					});
 				}
 			}
diff --git a/Engine/Source/Developer/TraceInsights/Private/Insights/Widgets/SQuickFind.cpp b/Engine/Source/Developer/TraceInsights/Private/Insights/Widgets/SQuickFind.cpp
--- a/Engine/Source/Developer/TraceInsights/Private/Insights/Widgets/SQuickFind.cpp
+++ b/Engine/Source/Developer/TraceInsights/Private/Insights/Widgets/SQuickFind.cpp
@@ -70,11 +70,23 @@ void SQuickFind::Construct(const FArguments& InArgs, TSharedPtr<FQuickFind> InQu
 		+ SVerticalBox::Slot()
 		.AutoHeight()
 		.Padding(0.0f, 5.0f, 0.0f, 5.0f)
-		.HAlign(EHorizontalAlignment::HAlign_Right)
+		.HAlign(EHorizontalAlignment::HAlign_Left)
 		.VAlign(EVerticalAlignment::VAlign_Bottom)
 		[
 			SNew(SHorizontalBox)
 
+			+ SHorizontalBox::Slot()
+			.HAlign(HAlign_Left)
+			.VAlign(VAlign_Center)
+			.AutoWidth()
+			.Padding(1.0f)
+			[
+				SNew(SButton)
+				.Text(LOCTEXT("FindPrev", "Find Previous"))
+				.ToolTipText(LOCTEXT("FindNextDesc", "Find the previous occurence that matches the search criteria."))
+				.OnClicked(this, &SQuickFind::FindPrevious_OnClicked)
+			]
+
 			+ SHorizontalBox::Slot()
 			.HAlign(HAlign_Left)
 			.VAlign(VAlign_Center)
@@ -84,7 +96,7 @@ void SQuickFind::Construct(const FArguments& InArgs, TSharedPtr<FQuickFind> InQu
 				SNew(SButton)
 				.Text(LOCTEXT("FindNext", "Find Next"))
 				.ToolTipText(LOCTEXT("FindNextDesc", "Find the next occurence that matches the search criteria."))
-				.OnClicked(this, &SQuickFind::FindNext_Onclicked)
+				.OnClicked(this, &SQuickFind::FindNext_OnClicked)
 			]
 		]
 	];
@@ -117,11 +129,19 @@ END_SLATE_FUNCTION_BUILD_OPTIMIZATION
 
 ////////////////////////////////////////////////////////////////////////////////////////////////////
 
-FReply SQuickFind::FindNext_Onclicked()
+FReply SQuickFind::FindNext_OnClicked()
 {
 	QuickFindViewModel->GetFilterConfigurator()->GetRootNode()->ProcessFilter();
 	QuickFindViewModel->GetOnFindNextEvent().Broadcast();
 
+	return FReply::Handled();
+}////////////////////////////////////////////////////////////////////////////////////////////////////
+
+FReply SQuickFind::FindPrevious_OnClicked()
+{
+	QuickFindViewModel->GetFilterConfigurator()->GetRootNode()->ProcessFilter();
+	QuickFindViewModel->GetOnFindPreviousEvent().Broadcast();
+
 	return FReply::Handled();
 }
 
@@ -170,6 +190,11 @@ TSharedRef<SDockTab> SQuickFind::SpawnTab(const FSpawnTabArgs& Args)
 	const TSharedRef<SDockTab> DockTab = SNew(SDockTab)
 		.TabRole(ETabRole::NomadTab);
 
+	if (!PendingWidget.IsValid())
+	{
+		return DockTab;
+	}
+
 	DockTab->SetContent(PendingWidget.ToSharedRef());
 	PendingWidget->SetParentTab(DockTab);
 
diff --git a/Engine/Source/Developer/TraceInsights/Private/Insights/Widgets/SQuickFind.h b/Engine/Source/Developer/TraceInsights/Private/Insights/Widgets/SQuickFind.h
--- a/Engine/Source/Developer/TraceInsights/Private/Insights/Widgets/SQuickFind.h
+++ b/Engine/Source/Developer/TraceInsights/Private/Insights/Widgets/SQuickFind.h
@@ -53,7 +53,9 @@ public:
 private:
 	void InitCommandList();
 
-	FReply FindNext_Onclicked();
+	FReply FindNext_OnClicked();
+
+	FReply FindPrevious_OnClicked();
 
 	FReply Cancel_OnClicked();
 
diff --git a/Engine/Source/Developer/TraceInsights/Private/Insights/Widgets/STimingView.cpp b/Engine/Source/Developer/TraceInsights/Private/Insights/Widgets/STimingView.cpp
--- a/Engine/Source/Developer/TraceInsights/Private/Insights/Widgets/STimingView.cpp
+++ b/Engine/Source/Developer/TraceInsights/Private/Insights/Widgets/STimingView.cpp
@@ -4199,6 +4199,7 @@ void STimingView::QuickFind_Execute()
 
 		QuickFindVm = MakeShared<FQuickFind>(FilterConfigurator);
 		QuickFindVm->GetOnFindNextEvent().AddSP(this, &STimingView::FindNextEvent);
+		QuickFindVm->GetOnFindPreviousEvent().AddSP(this, &STimingView::FindPrevEvent);
 	}
 
 	SQuickFind::CreateAndOpenQuickFilterWidget(QuickFindVm);
@@ -4312,5 +4313,47 @@ void STimingView::FindNextEvent()
 
 ////////////////////////////////////////////////////////////////////////////////////////////////////
 
+void STimingView::FindPrevEvent()
+{
+	TSharedPtr<const ITimingEvent> BestMatchEvent;
+	double StartTime = SelectedEvent.IsValid() ? SelectedEvent->GetStartTime() : std::numeric_limits<double>::max();
+
+	auto EventFilter = [StartTime](double EventStartTime, double EventEndTime, uint32 EventDepth)
+	{
+		return EventStartTime < StartTime;
+	};
+	FTimingEventSearchParameters Params(std::numeric_limits<double>::lowest(), StartTime, ETimingEventSearchFlags::StopAtFirstMatch, EventFilter);
+	Params.FilterExecutor = QuickFindVm->GetFilterConfigurator();
+	Params.SearchDirection = FTimingEventSearchParameters::ESearchDirection::Backward;
+
+	for (auto& Entry : AllTracks)
+	{
+		if (!Entry.Value->IsVisible())
+		{
+			continue;
+		}
+
+		Params.StartTime = BestMatchEvent.IsValid() ? BestMatchEvent->GetStartTime() : std::numeric_limits<double>::lowest();
+
+		TSharedPtr<const ITimingEvent> FoundEvent = Entry.Value->SearchEvent(Params);
+		if (FoundEvent.IsValid())
+		{
+			if (!BestMatchEvent.IsValid() || BestMatchEvent->GetStartTime() < FoundEvent->GetStartTime())
+			{
+				BestMatchEvent = FoundEvent;
+			}
+		}
+	}
+
+	if (BestMatchEvent)
+	{
+		SelectedEvent = BestMatchEvent;
+		BringIntoView(SelectedEvent->GetStartTime(), SelectedEvent->GetEndTime());
+		OnSelectedTimingEventChanged();
+	}
+}
+
+////////////////////////////////////////////////////////////////////////////////////////////////////
+
 #undef ACTIVATE_BENCHMARK
 #undef LOCTEXT_NAMESPACE
diff --git a/Engine/Source/Developer/TraceInsights/Private/Insights/Widgets/STimingView.h b/Engine/Source/Developer/TraceInsights/Private/Insights/Widgets/STimingView.h
--- a/Engine/Source/Developer/TraceInsights/Private/Insights/Widgets/STimingView.h
+++ b/Engine/Source/Developer/TraceInsights/Private/Insights/Widgets/STimingView.h
@@ -446,6 +446,7 @@ protected:
 	void SetTrackPosY(TSharedPtr<FBaseTimingTrack>& TrackPtr, float TopOffsetY) const;
 	
 	void FindNextEvent();
+	void FindPrevEvent();
 
 protected:
 	/** The track's viewport. Encapsulates info about position and scale. */
diff --git a/Engine/Source/Developer/TraceInsights/Public/Insights/ViewModels/TimingEventSearch.h b/Engine/Source/Developer/TraceInsights/Public/Insights/ViewModels/TimingEventSearch.h
--- a/Engine/Source/Developer/TraceInsights/Public/Insights/ViewModels/TimingEventSearch.h
+++ b/Engine/Source/Developer/TraceInsights/Public/Insights/ViewModels/TimingEventSearch.h
@@ -61,6 +61,12 @@ public:
 	// Predicate called when we get a match
 	typedef TFunctionRef<void(double /*StartTime*/, double /*EndTime*/, uint32 /*Depth*/)> EventMatchedPredicate;
 
+	enum class ESearchDirection : uint32
+	{
+		Forward = 0,
+		Backward = 1,
+	};
+
 	FTimingEventSearchParameters(double InStartTime, double InEndTime, ETimingEventSearchFlags Flags, EventFilterPredicate InEventFilter = NoFilter, EventMatchedPredicate InEventMatched = NoMatch)
 		: EventFilter(InEventFilter)
 		, EventMatched(InEventMatched)
@@ -95,6 +101,8 @@ public:
 	ETimingEventSearchFlags Flags;
 
 	TSharedPtr<Insights::IFilterExecutor> FilterExecutor;
+
+	ESearchDirection SearchDirection = ESearchDirection::Forward;
 };
 
 ////////////////////////////////////////////////////////////////////////////////////////////////////
diff --git a/Engine/Source/Developer/TraceServices/Public/Model/MonotonicTimeline.h b/Engine/Source/Developer/TraceServices/Public/Model/MonotonicTimeline.h
--- a/Engine/Source/Developer/TraceServices/Public/Model/MonotonicTimeline.h
+++ b/Engine/Source/Developer/TraceServices/Public/Model/MonotonicTimeline.h
@@ -341,6 +341,239 @@ public:
 		}
 	}
 
+	virtual void EnumerateEventsBackwardsDownSampled(double IntervalEnd, double IntervalStart, double Resolution, typename ITimeline<EventType>::EventCallback Callback) const override
+	{
+		int32 DetailLevelIndex = SettingsType::DetailLevelsCount - 1;
+		for (; DetailLevelIndex > 0; --DetailLevelIndex)
+		{
+			if (DetailLevels[DetailLevelIndex].Resolution <= Resolution)
+			{
+				break;
+			}
+		}
+
+		const FDetailLevel& DetailLevel = DetailLevels[DetailLevelIndex];
+		if (DetailLevel.ScopeEntries.Num() == 0)
+		{
+			return;
+		}
+
+		uint64 LastScopePageIndex = Algo::UpperBoundBy(DetailLevel.ScopeEntries, IntervalEnd, [](const FEventScopeEntryPage& Page)
+			{
+				return Page.BeginTime;
+			});
+
+		struct FEnumerationStackEntry
+		{
+			double EndTime;
+			EventType Event;
+		};
+
+		// By default, we start from the very end of the session.
+		auto ScopeEntryIterator = DetailLevel.ScopeEntries.GetIteratorFromItem(DetailLevel.ScopeEntries.Num() - 1);
+		auto EventsIterator = DetailLevel.Events.GetIteratorFromItem(DetailLevel.Events.Num() - 1);
+		FEnumerationStackEntry EventStack[SettingsType::MaxDepth];
+		int32 CurrentStackDepth = 0;
+		const FEventScopeEntry* ScopeEntry = nullptr;
+
+		if (LastScopePageIndex > 0 && LastScopePageIndex < DetailLevel.ScopeEntries.NumPages())
+		{
+			// If we have a page we can start from, start enumerating backwards from the begining of that page. 
+			ScopeEntryIterator = DetailLevel.ScopeEntries.GetIteratorFromPage(LastScopePageIndex);
+			const FEventScopeEntryPage* ScopePage = ScopeEntryIterator.GetCurrentPage();
+
+			EventsIterator = DetailLevel.Events.GetIteratorFromItem(ScopePage->BeginEventIndex);
+			CurrentStackDepth = ScopePage->InitialStackCount;
+			for (int32 InitialStackIndex = 0; InitialStackIndex < CurrentStackDepth; ++InitialStackIndex)
+			{
+				FEnumerationStackEntry& EnumerationStackEntry = EventStack[InitialStackIndex];
+				const FEventStackEntry& EventStackEntry = ScopePage->InitialStack[InitialStackIndex];
+				EnumerationStackEntry.EndTime = EventStackEntry.EndTime;
+				if (EnumerationStackEntry.EndTime < 0)
+				{
+					// We need to search for the EndTime of the event using pages. 
+					auto PageIterator = ScopeEntryIterator;
+					while (const FEventScopeEntryPage* Page = PageIterator.NextPage())
+					{
+						if (Page->InitialStack[InitialStackIndex].EndTime > 0)
+						{
+							EnumerationStackEntry.EndTime = Page->InitialStack[InitialStackIndex].EndTime;
+							break;
+						}
+					}
+					if (EnumerationStackEntry.EndTime < 0)
+					{
+						EnumerationStackEntry.EndTime = DetailLevel.InsertionState.LastTime;
+					}
+				}
+				EnumerationStackEntry.Event = DetailLevel.GetEvent(EventStackEntry.EventIndex);
+			}
+			// We start enumerating from the previous page.
+			ScopeEntry = ScopeEntryIterator.PrevItem();
+			EventsIterator.PrevItem();
+
+		}
+		else
+		{
+			// If we start from the end of the session, we use InsertionState as the initial stack.
+			CurrentStackDepth = DetailLevel.InsertionState.CurrentDepth;
+			for (int32 InitialStackIndex = 0; InitialStackIndex < CurrentStackDepth; ++InitialStackIndex)
+			{
+				FEnumerationStackEntry& EnumerationStackEntry = EventStack[InitialStackIndex];
+				const FEventStackEntry& EventStackEntry = DetailLevel.InsertionState.EventStack[InitialStackIndex];
+				EnumerationStackEntry.EndTime = DetailLevel.InsertionState.LastTime;
+				EnumerationStackEntry.Event = DetailLevel.GetEvent(EventStackEntry.EventIndex);
+			}
+			ScopeEntry = ScopeEntryIterator.GetCurrentItem();
+		}
+
+		// Enumerate backwards until we reach IntervalEnd, without calling the Callback because thess events are not in the provided interval.
+		const EventType* Event = EventsIterator.GetCurrentItem();
+		while (ScopeEntry && FMath::Abs(ScopeEntry->Time) > IntervalEnd)
+		{
+			if (ScopeEntry->Time < 0.0)
+			{
+				check(CurrentStackDepth > 0);
+				--CurrentStackDepth;
+				Event = EventsIterator.PrevItem();
+			}
+			else
+			{
+				check(CurrentStackDepth < SettingsType::MaxDepth);
+				FEnumerationStackEntry& StackEntry = EventStack[CurrentStackDepth++];
+				StackEntry.EndTime = ScopeEntry->Time;
+			}
+			ScopeEntry = ScopeEntryIterator.PrevItem();
+		}
+
+		// Call the callback for the events that are open at IntervalEnd.
+		for (int32 StackIndex = 0; StackIndex < CurrentStackDepth; ++StackIndex)
+		{
+			FEnumerationStackEntry& StackEntry = EventStack[StackIndex];
+			if (Callback(true, StackEntry.EndTime, StackEntry.Event) == EEventEnumerate::Stop)
+			{
+				return;
+			}
+		}
+
+		// Enumerate backwards between IntervalEnd and IntervalStart.
+		while (ScopeEntry && FMath::Abs(ScopeEntry->Time) >= IntervalStart)
+		{
+			if (ScopeEntry->Time < 0.0)
+			{
+				check(CurrentStackDepth > 0);
+				FEnumerationStackEntry& StackEntry = EventStack[--CurrentStackDepth];
+				StackEntry.Event = *Event;
+				if (Callback(false, -ScopeEntry->Time, StackEntry.Event) == EEventEnumerate::Stop)
+				{
+					return;
+				}
+				Event = EventsIterator.PrevItem();
+			}
+			else
+			{
+				check(CurrentStackDepth < SettingsType::MaxDepth);
+				FEnumerationStackEntry& StackEntry = EventStack[CurrentStackDepth++];
+				StackEntry.EndTime = ScopeEntry->Time;
+
+				if (Callback(true, ScopeEntry->Time, StackEntry.Event) == EEventEnumerate::Stop)
+				{
+					return;
+				}
+			}
+			ScopeEntry = ScopeEntryIterator.PrevItem();
+		}
+
+		// Find the StartTime of the events that are open at IntervalStart.
+		bool bSearchStartTimeUsingPages = false;
+		uint64 LastPageIndex = ScopeEntryIterator.GetCurrentPageIndex();
+		uint32 ExitDepth = 0;
+		while (CurrentStackDepth > 0 && ScopeEntry)
+		{
+			if (ScopeEntryIterator.GetCurrentPageIndex() != LastPageIndex)
+			{
+				bSearchStartTimeUsingPages = true;
+				break;
+			}
+			if (ScopeEntry->Time < 0.0)
+			{
+				if (ExitDepth == 0)
+				{
+					FEnumerationStackEntry& StackEntry = EventStack[--CurrentStackDepth];
+					if (Callback(false, -ScopeEntry->Time, *EventsIterator.GetCurrentItem()) == EEventEnumerate::Stop)
+					{
+						return;
+					}
+				}
+				else
+				{
+					--ExitDepth;
+				}
+				EventsIterator.PrevItem();
+			}
+			else
+			{
+				++ExitDepth;
+			}
+
+			LastPageIndex = ScopeEntryIterator.GetCurrentPageIndex();
+			ScopeEntry = ScopeEntryIterator.PrevItem();
+		}
+
+		if (bSearchStartTimeUsingPages)
+		{
+			do
+			{
+				const FEventScopeEntryPage* CurrentScopePage = ScopeEntryIterator.GetCurrentPage();
+				check(CurrentStackDepth <= CurrentScopePage->InitialStackCount);
+
+				while (CurrentStackDepth > 0)
+				{
+					--CurrentStackDepth;
+					EventType CurrentEvent = DetailLevel.GetEvent(CurrentScopePage->InitialStack[CurrentStackDepth].EventIndex);
+					double StartTime = DetailLevel.ScopeEntries[CurrentScopePage->InitialStack[CurrentStackDepth].EnterScopeIndex].Time;
+					if (Callback(false, StartTime, CurrentEvent) == EEventEnumerate::Stop)
+					{
+						return;
+					}
+				}
+			} while (ScopeEntryIterator.PrevPage());
+		}
+	}
+
+	virtual void EnumerateEventsBackwardsDownSampled(double IntervalEnd, double IntervalStart, double Resolution, typename ITimeline<EventType>::EventRangeCallback Callback) const override
+	{
+		struct FStackEntry
+		{
+			double EndTime;
+			EventType Event;
+		};
+		FStackEntry EventStack[SettingsType::MaxDepth];
+		uint32 CurrentDepth = 0;
+
+		EnumerateEventsBackwardsDownSampled(IntervalEnd, IntervalStart, Resolution, [&EventStack, &CurrentDepth, Callback](bool IsEnter, double Time, const EventType& Event)
+		{
+			if (IsEnter)
+			{
+				FStackEntry& StackEntry = EventStack[CurrentDepth];
+				StackEntry.Event = Event;
+				StackEntry.EndTime = Time;
+				++CurrentDepth;
+			}
+			else
+			{
+				FStackEntry& StackEntry = EventStack[--CurrentDepth];
+				EEventEnumerate Ret = Callback(Time, StackEntry.EndTime, CurrentDepth, Event);
+				if (Ret != EEventEnumerate::Continue)
+				{
+					return Ret;
+				}
+			}
+
+			return EEventEnumerate::Continue;
+		});
+	}
+
 	virtual void EnumerateEventsDownSampledAsync(const typename ITimeline<EventType>::EnumerateAsyncParams& EnumerateAsyncParams) const override
 	{
 		if (EnumerateAsyncParams.IntervalEnd < EnumerateAsyncParams.IntervalStart)
@@ -482,6 +715,16 @@ public:
 		EnumerateEventsDownSampled(IntervalStart, IntervalEnd, 0.0, Callback);
 	}
 
+	virtual void EnumerateEventsBackwards(double IntervalEnd, double IntervalStart, typename ITimeline<EventType>::EventCallback Callback) const override
+	{
+		EnumerateEventsBackwardsDownSampled(IntervalEnd, IntervalStart, 0.0, Callback);
+	}
+
+	virtual void EnumerateEventsBackwards(double IntervalEnd, double IntervalStart, typename ITimeline<EventType>::EventRangeCallback Callback) const override
+	{
+		EnumerateEventsBackwardsDownSampled(IntervalEnd, IntervalStart, 0.0, Callback);
+	}
+
 	virtual bool GetEventInfo(double InTime, double DeltaTime, int32 Depth, typename ITimeline<InEventType>::FTimelineEventInfo& EventInfo) const override
 	{
 		if (Depth >= SettingsType::MaxDepth || Depth < 0)
diff --git a/Engine/Source/Developer/TraceServices/Public/TraceServices/Containers/Timelines.h b/Engine/Source/Developer/TraceServices/Public/TraceServices/Containers/Timelines.h
--- a/Engine/Source/Developer/TraceServices/Public/TraceServices/Containers/Timelines.h
+++ b/Engine/Source/Developer/TraceServices/Public/TraceServices/Containers/Timelines.h
@@ -55,9 +55,18 @@ public:
 	virtual double GetEndTime() const = 0;
 	virtual void EnumerateEventsDownSampled(double IntervalStart, double IntervalEnd, double Resolution, EventCallback Callback) const = 0;
 	virtual void EnumerateEventsDownSampled(double IntervalStart, double IntervalEnd, double Resolution, EventRangeCallback Callback) const = 0;
+
+	virtual void EnumerateEventsBackwardsDownSampled(double IntervalEnd, double IntervalStart, double Resolution, EventCallback Callback) const {};
+	virtual void EnumerateEventsBackwardsDownSampled(double IntervalEnd, double IntervalStart, double Resolution, EventRangeCallback Callback) const {};
+
 	virtual void EnumerateEventsDownSampledAsync(const EnumerateAsyncParams& EnumerateAsyncParams) const {};
+
 	virtual void EnumerateEvents(double IntervalStart, double IntervalEnd, EventCallback Callback) const = 0;
 	virtual void EnumerateEvents(double IntervalStart, double IntervalEnd, EventRangeCallback Callback) const = 0;
+
+	virtual void EnumerateEventsBackwards(double IntervalEnd, double IntervalStart, EventCallback Callback) const {};
+	virtual void EnumerateEventsBackwards(double IntervalEnd, double IntervalStart, EventRangeCallback Callback) const {};
+
 	virtual int32 GetDepthAt(double Time) const { return 0; };
 	
 	/**
