Change Interval timeline to iterate over pages with unended intervals

diff --git a/Engine/Source/Developer/TraceServices/Public/Model/IntervalTimeline.h b/Engine/Source/Developer/TraceServices/Public/Model/IntervalTimeline.h
--- a/Engine/Source/Developer/TraceServices/Public/Model/IntervalTimeline.h
+++ b/Engine/Source/Developer/TraceServices/Public/Model/IntervalTimeline.h
@@ -55,7 +55,7 @@ public:
 		while (Event)
 		{
 			const FEventPage* Page = EventIterator.GetCurrentPage();
-			if (Page->EndTime < IntervalStart || IntervalEnd < Page->BeginTime)
+			if ((Page->OpenEvents == 0 && Page->EndTime < IntervalStart) || IntervalEnd < Page->BeginTime)
 			{
 				if (!EventIterator.NextPage())
 				{
@@ -108,6 +108,7 @@ public:
 		LastPage->EndEventIndex = Index + 1;
 
 		++ModCount;
+		++LastPage->OpenEvents;
 
 		return Index;
 	}
@@ -123,6 +124,8 @@ public:
 		Page->EndTime = FMath::Max(Page->EndTime, EndTime);
 
 		++ModCount;
+		check(Page->OpenEvents>0);
+		--Page->OpenEvents;
 
 		return EventInternal.Event;
 	}
@@ -157,6 +160,7 @@ private:
 		uint64 EndEventIndex = 0;
 		FEventInternal* Items = nullptr;
 		uint64 Count = 0;
+		uint64 OpenEvents = 0;
 	};
 
 	TPagedArray<FEventInternal, FEventPage> Events;
