Static analysis fix and fix for crash in context switches analysis. #rb none #jira UE-122035

diff --git a/Engine/Source/Developer/TraceServices/Private/Model/ContextSwitches.cpp b/Engine/Source/Developer/TraceServices/Private/Model/ContextSwitches.cpp
--- a/Engine/Source/Developer/TraceServices/Private/Model/ContextSwitches.cpp
+++ b/Engine/Source/Developer/TraceServices/Private/Model/ContextSwitches.cpp
@@ -318,6 +318,8 @@ void FContextSwitchesProvider::Add(uint32 SystemThreadId, double Start, double E
 
 void FContextSwitchesProvider::AddThreadInfo(uint32 ThreadId, uint32 SystemThreadId)
 {
+	Session.WriteAccessCheck();
+
 	if (SystemThreadId == 0)
 	{
 		// At the start of a session some threads might be received with a SystemThreadId of 0.
@@ -327,15 +329,10 @@ void FContextSwitchesProvider::AddThreadInfo(uint32 ThreadId, uint32 SystemThrea
 	uint32* OldSystemThreadId = TraceToSystemThreadIdMap.Find(ThreadId);
 	if (OldSystemThreadId)
 	{
-		ensure(*OldSystemThreadId == SystemThreadId);
+		ensure((*OldSystemThreadId) == SystemThreadId);
 	}
-	TraceToSystemThreadIdMap.Add(ThreadId, SystemThreadId);
 
-	uint32* OldTraceThreadId = SystemToTraceThreadIdMap.Find(SystemThreadId);
-	if (OldTraceThreadId)
-	{
-		ensure(*OldTraceThreadId == ThreadId);
-	}
+	TraceToSystemThreadIdMap.Add(ThreadId, SystemThreadId);
 	SystemToTraceThreadIdMap.Add(SystemThreadId, ThreadId);
 }
 
