Fix static analysis error

diff --git a/Engine/Source/Developer/TraceServices/Private/Model/ContextSwitches.cpp b/Engine/Source/Developer/TraceServices/Private/Model/ContextSwitches.cpp
--- a/Engine/Source/Developer/TraceServices/Private/Model/ContextSwitches.cpp
+++ b/Engine/Source/Developer/TraceServices/Private/Model/ContextSwitches.cpp
@@ -326,10 +326,11 @@ void FContextSwitchesProvider::AddThreadInfo(uint32 ThreadId, uint32 SystemThrea
 		return;
 	}
 
-	uint32* OldSystemThreadId = TraceToSystemThreadIdMap.Find(ThreadId);
-	if (OldSystemThreadId)
+	uint32* OldSystemThreadIdPtr = TraceToSystemThreadIdMap.Find(ThreadId);
+	if (OldSystemThreadIdPtr)
 	{
-		ensure((*OldSystemThreadId) == SystemThreadId);
+		uint32 OldSystemThreadId = *OldSystemThreadIdPtr;
+		ensure(OldSystemThreadId == SystemThreadId);
 	}
 
 	TraceToSystemThreadIdMap.Add(ThreadId, SystemThreadId);
