Gave some uninitialised FTaskInfo members some default values. There are some paths in the provider that create task info objects but do no set these fields. Code that displays tool tips checks the validity of these members, but they're uninitialised and end up segfaulting.

diff --git a/Engine/Source/Developer/TraceServices/Public/TraceServices/Model/TasksProfiler.h b/Engine/Source/Developer/TraceServices/Public/TraceServices/Model/TasksProfiler.h
--- a/Engine/Source/Developer/TraceServices/Public/TraceServices/Model/TasksProfiler.h
+++ b/Engine/Source/Developer/TraceServices/Public/TraceServices/Model/TasksProfiler.h
@@ -16,9 +16,9 @@ namespace TraceServices
 		TaskTrace::FId Id = TaskTrace::InvalidId;
 
 		// this members are filled during task launch
-		const TCHAR* DebugName; // user provided or automatically generated (in format `full_path(lineno)`)
-		bool bTracked; // if the task is awaitable, not awaitable tasks will never be "completed", they are done when their execution is finished
-		int32 ThreadToExecuteOn; // ENamedThreads::Type - combines the info about thread to execute on (one of named threads or a worker
+		const TCHAR* DebugName = nullptr; // user provided or automatically generated (in format `full_path(lineno)`)
+		bool bTracked = false; // if the task is awaitable, not awaitable tasks will never be "completed", they are done when their execution is finished
+		int32 ThreadToExecuteOn = ENamedThreads::AnyThread; // ENamedThreads::Type - combines the info about thread to execute on (one of named threads or a worker
 		// thread - ENamedThreads::AnyThread), task priority, thread priority (in case of a worker thread), and queue index (in case of a named
 		// thread - main queue or local queue). See ENamedThreads to split the info into separate fields
 
