[Insights] Fixed missing include ( -DisableUnity -NoPCH compile error).

diff --git a/Engine/Source/Developer/TraceServices/Public/TraceServices/Model/TasksProfiler.h b/Engine/Source/Developer/TraceServices/Public/TraceServices/Model/TasksProfiler.h
--- a/Engine/Source/Developer/TraceServices/Public/TraceServices/Model/TasksProfiler.h
+++ b/Engine/Source/Developer/TraceServices/Public/TraceServices/Model/TasksProfiler.h
@@ -3,6 +3,7 @@
 #pragma once
 
 #include "TraceServices/Model/AnalysisSession.h"
+#include "Async/TaskGraphInterfaces.h"
 #include "Async/TaskTrace.h"
 #include "Containers/Array.h"
 
