[Insights]
  - Cpu Profiler: Fixed "<unknown>" timers being incorrectly merged into a single timer (when multiple SpecId are used before their corresponding EventSpec events are processed).
  - Cpu Profiler: Added "<noname>" and "<invalid>" names when EventSpec event does not declare correctly a cpu timer.
  - Timing view: Made all the special timer names (ex. "<unknown>", "<noname>", "<invalid>", etc.) to be colored in black.

diff --git a/Engine/Source/Developer/TraceInsights/Private/Insights/ViewModels/TimingEvent.cpp b/Engine/Source/Developer/TraceInsights/Private/Insights/ViewModels/TimingEvent.cpp
--- a/Engine/Source/Developer/TraceInsights/Private/Insights/ViewModels/TimingEvent.cpp
+++ b/Engine/Source/Developer/TraceInsights/Private/Insights/ViewModels/TimingEvent.cpp
@@ -26,7 +26,8 @@ uint32 FTimingEvent::ComputeEventColor(uint32 Id)
 uint32 FTimingEvent::ComputeEventColor(const TCHAR* Str)
 {
 	uint32 Color = 0;
-	if (Str != nullptr)
+	// Let the special event names (ex. "<unknown>", "<noname>", "<invalid>", etc.) to be colored in black.
+	if (Str != nullptr && *Str != TEXT('<'))
 	{
 		for (const TCHAR* c = Str; *c; ++c)
 		{
diff --git a/Engine/Source/Developer/TraceServices/Private/Analyzers/CpuProfilerTraceAnalysis.cpp b/Engine/Source/Developer/TraceServices/Private/Analyzers/CpuProfilerTraceAnalysis.cpp
--- a/Engine/Source/Developer/TraceServices/Private/Analyzers/CpuProfilerTraceAnalysis.cpp
+++ b/Engine/Source/Developer/TraceServices/Private/Analyzers/CpuProfilerTraceAnalysis.cpp
@@ -45,25 +45,36 @@ bool FCpuProfilerAnalyzer::OnEvent(uint16 RouteId, EStyle Style, const FOnEventC
 	case RouteId_EventSpec:
 	{
 		uint32 SpecId = EventData.GetValue<uint32>("Id");
-		
+
+		const TCHAR* TimerName = nullptr;
 		FString Name;
 		if (EventData.GetString("Name", Name))
 		{
-			const TCHAR* NamePtr = Session.StoreString(*Name);
-			DefineScope(SpecId, NamePtr);
+			TimerName = *Name;
 		}
 		else
 		{
 			uint8 CharSize = EventData.GetValue<uint8>("CharSize");
 			if (CharSize == sizeof(ANSICHAR))
 			{
-				DefineScope(SpecId, Session.StoreString(StringCast<TCHAR>(reinterpret_cast<const ANSICHAR*>(EventData.GetAttachment())).Get()));
+				TimerName = StringCast<TCHAR>(reinterpret_cast<const ANSICHAR*>(EventData.GetAttachment())).Get();
 			}
 			else if (CharSize == 0 || CharSize == sizeof(TCHAR)) // 0 for backwards compatibility
 			{
-				DefineScope(SpecId, Session.StoreString(reinterpret_cast<const TCHAR*>(EventData.GetAttachment())));
+				TimerName = reinterpret_cast<const TCHAR*>(EventData.GetAttachment());
+			}
+			else
+			{
+				Name = FString::Printf(TEXT("<invalid %u>"), SpecId);
+				TimerName = *Name;
 			}
 		}
+		if (TimerName[0] == 0)
+		{
+			Name = FString::Printf(TEXT("<noname %u>"), SpecId);
+			TimerName = *Name;
+		}
+		DefineTimer(SpecId, Session.StoreString(TimerName));
 		break;
 	}
 	case RouteId_EndThread:
@@ -187,7 +198,12 @@ uint64 FCpuProfilerAnalyzer::ProcessBuffer(const FEventTime& EventTime, uint32 T
 			uint32 TimerId;
 			if (!FindIt)
 			{
-				TimerId = SpecIdToTimerIdMap.Add(SpecId, TimingProfilerProvider.AddCpuTimer(TEXT("<unknown>")));
+				// Adds a timer with an "unknown" name.
+				const TCHAR* TimerName = Session.StoreString(*FString::Printf(TEXT("<unknown %u>"), SpecId));
+				// The name might be updated when an EventSpec event is received (for this SpecId),
+				// so we do not want to merge by name all the unknown timers.
+				constexpr bool bMergeByName = false;
+				TimerId = DefineTimer(SpecId, TimerName, bMergeByName);
 			}
 			else
 			{
@@ -256,18 +272,23 @@ void FCpuProfilerAnalyzer::OnCpuScopeEnter(const FOnEventContext& Context)
 
 	uint32 SpecId = Context.EventData.GetTypeInfo().GetId();
 	SpecId = ~SpecId; // to keep out of the way of normal spec IDs.
+
+	uint32 TimerId;
 	uint32* TimerIdIter = SpecIdToTimerIdMap.Find(SpecId);
-	if (TimerIdIter == nullptr)
+	if (TimerIdIter)
+	{
+		TimerId = *TimerIdIter;
+	}
+	else
 	{
 		FString ScopeName;
 		ScopeName += Context.EventData.GetTypeInfo().GetName();
-		DefineScope(SpecId, Session.StoreString(*ScopeName));
-		TimerIdIter = SpecIdToTimerIdMap.Find(SpecId);
+		TimerId = DefineTimer(SpecId, Session.StoreString(*ScopeName));
 	}
 
 	TArray<uint8> CborData;
 	Context.EventData.SerializeToCbor(CborData);
-	uint32 TimerId = TimingProfilerProvider.AddMetadata(*TimerIdIter, MoveTemp(CborData));
+	TimerId = TimingProfilerProvider.AddMetadata(TimerId, MoveTemp(CborData));
 
 	uint64 Cycle = Context.EventTime.AsCycle64();
 	ThreadState.PendingEvents.Add({Cycle, TimerId});
@@ -287,14 +308,14 @@ void FCpuProfilerAnalyzer::OnCpuScopeLeave(const FOnEventContext& Context)
 	ThreadState.PendingEvents.Add({~Cycle});
 }
 
-void FCpuProfilerAnalyzer::DefineScope(uint32 SpecId, const TCHAR* Name)
+uint32 FCpuProfilerAnalyzer::DefineTimer(uint32 SpecId, const TCHAR* Name, bool bMergeByName)
 {
 	// The cpu scoped events (timers) will be merged by name.
 	// Ex.: If there are multiple timers defined in code with same name,
 	//      those will appear in Insights as a single timer.
 
 	// Check if a timer with same name was already defined.
-	uint32* FindTimerIdByName = ScopeNameToTimerIdMap.Find(Name);
+	uint32* FindTimerIdByName = bMergeByName ? ScopeNameToTimerIdMap.Find(Name) : nullptr;
 	if (FindTimerIdByName)
 	{
 		// Yes, a timer with same name was already defined.
@@ -308,16 +329,18 @@ void FCpuProfilerAnalyzer::DefineScope(uint32 SpecId, const TCHAR* Name)
 			// In this case, we do not remap the SpecId to the previously defined timer with same name.
 			// This is becasue the two timers are already used in timelines.
 			// So we will continue to use separate timers, even if those have same name.
+			return *FindTimerId;
 		}
 		else
 		{
 			// Map this SpecId to the previously defined timer with same name.
 			SpecIdToTimerIdMap.Add(SpecId, *FindTimerIdByName);
+			return *FindTimerIdByName;
 		}
 	}
 	else
 	{
-		// No, a timer with same name was not defined.
+		// No, a timer with same name was not defined (or we do not want to merge by name).
 		// Check if SpecId is already mapped to timer.
 		uint32* FindTimerId = SpecIdToTimerIdMap.Find(SpecId);
 		if (FindTimerId)
@@ -325,8 +348,12 @@ void FCpuProfilerAnalyzer::DefineScope(uint32 SpecId, const TCHAR* Name)
 			// Yes, SpecId was already mapped to a timer (ex. as an <unknown> timer).
 			// Update name for mapped timer.
 			TimingProfilerProvider.SetTimerName(*FindTimerId, Name);
-			// Map the name to the timer.
-			ScopeNameToTimerIdMap.Add(Name, *FindTimerId);
+			if (bMergeByName)
+			{
+				// Map the name to the timer.
+				ScopeNameToTimerIdMap.Add(Name, *FindTimerId);
+			}
+			return *FindTimerId;
 		}
 		else
 		{
@@ -334,8 +361,12 @@ void FCpuProfilerAnalyzer::DefineScope(uint32 SpecId, const TCHAR* Name)
 			uint32 NewTimerId = TimingProfilerProvider.AddCpuTimer(Name);
 			// Map the SpecId to the timer.
 			SpecIdToTimerIdMap.Add(SpecId, NewTimerId);
-			// Map the name to the timer.
-			ScopeNameToTimerIdMap.Add(Name, NewTimerId);
+			if (bMergeByName)
+			{
+				// Map the name to the timer.
+				ScopeNameToTimerIdMap.Add(Name, NewTimerId);
+			}
+			return NewTimerId;
 		}
 	}
 }
diff --git a/Engine/Source/Developer/TraceServices/Private/Analyzers/CpuProfilerTraceAnalysis.h b/Engine/Source/Developer/TraceServices/Private/Analyzers/CpuProfilerTraceAnalysis.h
--- a/Engine/Source/Developer/TraceServices/Private/Analyzers/CpuProfilerTraceAnalysis.h
+++ b/Engine/Source/Developer/TraceServices/Private/Analyzers/CpuProfilerTraceAnalysis.h
@@ -44,7 +44,7 @@ private:
 
 	void OnCpuScopeEnter(const FOnEventContext& Context);
 	void OnCpuScopeLeave(const FOnEventContext& Context);
-	void DefineScope(uint32 SpecId, const TCHAR* ScopeName);
+	uint32 DefineTimer(uint32 SpecId, const TCHAR* ScopeName, bool bMergeByName = true); // returns the TimerId
 	FThreadState& GetThreadState(uint32 ThreadId);
 	uint64 ProcessBuffer(const FEventTime& EventTime, uint32 ThreadId, const uint8* BufferPtr, uint32 BufferSize);
 
