Analyzer support for additional memory spaces

Changes to analysis wrt to the changelist "Memory tracing additional heaps".

diff --git a/Engine/Source/Developer/TraceInsights/Private/Insights/MemoryProfiler/ViewModels/MemAllocGroupingByHeap.cpp b/Engine/Source/Developer/TraceInsights/Private/Insights/MemoryProfiler/ViewModels/MemAllocGroupingByHeap.cpp
new file mode 100644
--- /dev/null
+++ b/Engine/Source/Developer/TraceInsights/Private/Insights/MemoryProfiler/ViewModels/MemAllocGroupingByHeap.cpp
@@ -0,0 +1,93 @@
+﻿// Copyright Epic Games, Inc. All Rights Reserved.
+
+#include "MemAllocGroupingByHeap.h"
+#include "CallstackFormatting.h"
+#include "Internationalization/Internationalization.h"
+
+// Insights
+#include "Insights/MemoryProfiler/ViewModels/MemAllocNode.h"
+
+#define LOCTEXT_NAMESPACE "Insights::FMemAllocGroupingByHeap"
+
+namespace Insights
+{
+	
+INSIGHTS_IMPLEMENT_RTTI(FMemAllocGroupingByHeap)
+	
+////////////////////////////////////////////////////////////////////////////////////////////////////
+
+FMemAllocGroupingByHeap::FMemAllocGroupingByHeap(const TraceServices::IAllocationsProvider& InAllocProvider)
+	: FTreeNodeGrouping(
+		LOCTEXT("Grouping_ByHeap_ShortName", "Heap"),
+		LOCTEXT("Grouping_ByHeap_TitleName", "By Heap"),
+		LOCTEXT("Grouping_Heap_Desc", "Creates a tree based on heap."),
+		TEXT("Profiler.FiltersAndPresets.GroupNameIcon"),
+		nullptr)
+	, AllocProvider(InAllocProvider)
+{
+}
+
+////////////////////////////////////////////////////////////////////////////////////////////////////
+
+FMemAllocGroupingByHeap::~FMemAllocGroupingByHeap()
+{
+}
+
+////////////////////////////////////////////////////////////////////////////////////////////////////
+
+FTableTreeNodePtr MakeGroupNodeHierarchy(const TraceServices::IAllocationsProvider::FHeapSpec& Spec, TWeakPtr<FTable>& InParentTable, TArray<FTableTreeNodePtr>& NodeTable)
+{
+	auto Node = MakeShared<FTableTreeNode>(FName(Spec.Name), InParentTable);
+	NodeTable[Spec.Id] = Node;
+	for (TraceServices::IAllocationsProvider::FHeapSpec* ChildSpec : Spec.Children)
+	{
+		auto ChildNode = MakeGroupNodeHierarchy(*ChildSpec, InParentTable, NodeTable);
+		Node->AddChildAndSetGroupPtr(ChildNode);
+	}
+	return Node;
+}
+
+////////////////////////////////////////////////////////////////////////////////////////////////////
+
+void FMemAllocGroupingByHeap::GroupNodes(const TArray<FTableTreeNodePtr>& Nodes, FTableTreeNode& ParentGroup,
+	TWeakPtr<FTable> InParentTable, std::atomic<bool>& bCancelGrouping) const
+{
+	ParentGroup.ClearChildren();
+
+	// Build heap hierarchy
+	TArray<FTableTreeNodePtr> HeapNodes;
+	HeapNodes.AddZeroed(256);
+	
+	AllocProvider.EnumerateRootHeaps([&](uint8 Id, const TraceServices::IAllocationsProvider::FHeapSpec& Spec)
+	{
+		ParentGroup.AddChildAndSetGroupPtr(MakeGroupNodeHierarchy(Spec, InParentTable, HeapNodes));
+	});
+
+	// Add allocations nodes to heaps
+	for (FTableTreeNodePtr NodePtr : Nodes)
+	{
+		if (bCancelGrouping)
+		{
+			return;
+		}
+		
+		const FMemAllocNode& MemAllocNode = static_cast<const FMemAllocNode&>(*NodePtr);
+		const FMemoryAlloc* Alloc = MemAllocNode.GetMemAlloc();
+
+		if (Alloc)
+		{
+			//todo: Calculating the real HeapId when the allocation is first added to the provider was too expensive, so deferring that operation to here could make sense.
+			//const uint8 HeapId = AllocProvider.GetParentBlock(Alloc->GetAddress());
+			const uint8 HeapId = Alloc->GetRootHeap();
+			if (FTableTreeNodePtr GroupPtr = HeapNodes[HeapId])
+			{
+				GroupPtr->AddChildAndSetGroupPtr(NodePtr);
+			}
+		}
+	}
+}
+	
+////////////////////////////////////////////////////////////////////////////////////////////////////
+} // namespace Insights
+
+#undef LOCTEXT_NAMESPACE
\ No newline at end of file
diff --git a/Engine/Source/Developer/TraceInsights/Private/Insights/MemoryProfiler/ViewModels/MemAllocGroupingByHeap.h b/Engine/Source/Developer/TraceInsights/Private/Insights/MemoryProfiler/ViewModels/MemAllocGroupingByHeap.h
new file mode 100644
--- /dev/null
+++ b/Engine/Source/Developer/TraceInsights/Private/Insights/MemoryProfiler/ViewModels/MemAllocGroupingByHeap.h
@@ -0,0 +1,37 @@
+﻿// Copyright Epic Games, Inc. All Rights Reserved.
+
+#pragma once
+
+#include "CoreMinimal.h"
+
+// Insights
+#include "Insights/Table/ViewModels/BaseTreeNode.h"
+#include "Insights/Table/ViewModels/TreeNodeGrouping.h"
+#include "TraceServices/Model/AllocationsProvider.h"
+
+namespace TraceServices
+{
+	struct FStackFrame;
+}
+
+namespace Insights
+{
+
+////////////////////////////////////////////////////////////////////////////////////////////////////
+
+class FMemAllocGroupingByHeap : public FTreeNodeGrouping
+{
+	INSIGHTS_DECLARE_RTTI(FMemAllocGroupingByHeap, FTreeNodeGrouping)
+
+public:
+	FMemAllocGroupingByHeap(const TraceServices::IAllocationsProvider& AllocProvider);
+	virtual ~FMemAllocGroupingByHeap() override;
+
+	virtual void GroupNodes(const TArray<FTableTreeNodePtr>& Nodes, FTableTreeNode& ParentGroup, TWeakPtr<FTable> InParentTable, std::atomic<bool>& bCancelGrouping) const override;
+private:
+	const TraceServices::IAllocationsProvider& AllocProvider;
+};
+
+////////////////////////////////////////////////////////////////////////////////////////////////////
+
+} // namespace Insights
diff --git a/Engine/Source/Developer/TraceInsights/Private/Insights/MemoryProfiler/ViewModels/MemoryAlloc.cpp b/Engine/Source/Developer/TraceInsights/Private/Insights/MemoryProfiler/ViewModels/MemoryAlloc.cpp
--- a/Engine/Source/Developer/TraceInsights/Private/Insights/MemoryProfiler/ViewModels/MemoryAlloc.cpp
+++ b/Engine/Source/Developer/TraceInsights/Private/Insights/MemoryProfiler/ViewModels/MemoryAlloc.cpp
@@ -8,13 +8,15 @@ namespace Insights
 
 ////////////////////////////////////////////////////////////////////////////////////////////////////
 
-FMemoryAlloc::FMemoryAlloc(double InStartTime, double InEndTime, uint64 InAddress, uint64 InSize, const TCHAR* InTag, const TraceServices::FCallstack* InCallstack)
+FMemoryAlloc::FMemoryAlloc(double InStartTime, double InEndTime, uint64 InAddress, uint64 InSize, const TCHAR* InTag, const TraceServices::FCallstack* InCallstack, HeapId InRootHeap)
 : StartTime(InStartTime)
 , EndTime(InEndTime)
 , Address(InAddress)
 , Size(InSize)
 , Tag(InTag)
 , Callstack(InCallstack)
+, RootHeap(InRootHeap)
+, bIsBlock(false)
 {
 }
 
diff --git a/Engine/Source/Developer/TraceInsights/Private/Insights/MemoryProfiler/ViewModels/MemoryAlloc.h b/Engine/Source/Developer/TraceInsights/Private/Insights/MemoryProfiler/ViewModels/MemoryAlloc.h
--- a/Engine/Source/Developer/TraceInsights/Private/Insights/MemoryProfiler/ViewModels/MemoryAlloc.h
+++ b/Engine/Source/Developer/TraceInsights/Private/Insights/MemoryProfiler/ViewModels/MemoryAlloc.h
@@ -3,6 +3,8 @@
 #pragma once
 
 #include "CoreMinimal.h"
+
+#include "TraceServices/Model/AllocationsProvider.h"
 #include "TraceServices/Model/Callstack.h"
 
 namespace Insights
@@ -15,7 +17,7 @@ class FMemoryAlloc
 	friend class SMemAllocTableTreeView;
 
 public:
-	FMemoryAlloc(double InStartTime, double InEndTime, uint64 InAddress, uint64 InSize, const TCHAR* InTag, const TraceServices::FCallstack* InCallstack);
+	FMemoryAlloc(double InStartTime, double InEndTime, uint64 InAddress, uint64 InSize, const TCHAR* InTag, const TraceServices::FCallstack* InCallstack, HeapId InRootHeap);
 	~FMemoryAlloc();
 
 	double GetStartTime() const { return StartTime; }
@@ -27,6 +29,7 @@ public:
 	const TCHAR* GetTag() const { return Tag; }
 	const TraceServices::FCallstack* GetCallstack() const { return Callstack; }
 	FText GetFullCallstack() const;
+	HeapId GetRootHeap() const { return RootHeap; };
 
 private:
 	double StartTime;
@@ -35,6 +38,8 @@ private:
 	uint64 Size;
 	const TCHAR* Tag;
 	const TraceServices::FCallstack* Callstack;
+	HeapId RootHeap;
+	bool bIsBlock;
 };
 
 ////////////////////////////////////////////////////////////////////////////////////////////////////
diff --git a/Engine/Source/Developer/TraceInsights/Private/Insights/MemoryProfiler/Widgets/SMemAllocTableTreeView.cpp b/Engine/Source/Developer/TraceInsights/Private/Insights/MemoryProfiler/Widgets/SMemAllocTableTreeView.cpp
--- a/Engine/Source/Developer/TraceInsights/Private/Insights/MemoryProfiler/Widgets/SMemAllocTableTreeView.cpp
+++ b/Engine/Source/Developer/TraceInsights/Private/Insights/MemoryProfiler/Widgets/SMemAllocTableTreeView.cpp
@@ -31,6 +31,8 @@
 
 #include <limits>
 
+#include "Insights/MemoryProfiler/ViewModels/MemAllocGroupingByHeap.h"
+
 #define LOCTEXT_NAMESPACE "SMemAllocTableTreeView"
 
 using namespace TraceServices;
@@ -138,14 +140,26 @@ void SMemAllocTableTreeView::RebuildTree(bool bResync)
 				}
 				TableTreeNodes.Reserve(TotalAllocCount);
 
-				FName BaseNodeName(TEXT("alloc"));
+				uint32 HeapAllocCount(0);
+				const FName BaseNodeName(TEXT("alloc"));
+				const FName BaseHeapName(TEXT("heap"));
 				for (int32 AllocIndex = TableTreeNodes.Num(); AllocIndex < TotalAllocCount; ++AllocIndex)
 				{
-					FName NodeName(BaseNodeName, AllocIndex + 1);
+					const FMemoryAlloc* Alloc = MemAllocTable->GetMemAlloc(AllocIndex);
+					FName NodeName(Alloc->bIsBlock ? BaseHeapName : BaseNodeName, AllocIndex + 1);
 					FMemAllocNodePtr NodePtr = MakeShared<FMemAllocNode>(NodeName, MemAllocTable, AllocIndex);
+
+					// Until we have an UX story around heap allocations
+					// remove them from the list
+					if (Alloc->bIsBlock)
+					{
+						++HeapAllocCount;
+						continue;
+					}
+					
 					TableTreeNodes.Add(NodePtr);
 				}
-				ensure(TableTreeNodes.Num() == TotalAllocCount);
+				ensure(TableTreeNodes.Num() == TotalAllocCount - HeapAllocCount);
 			}
 		}
 	}
@@ -341,6 +355,8 @@ void SMemAllocTableTreeView::UpdateQuery(TraceServices::IAllocationsProvider::EQ
 					Alloc.Size = Allocation->GetSize();
 					Alloc.Tag = Provider.GetTagName(Allocation->GetTag());
 					Alloc.Callstack = Allocation->GetCallstack();
+					Alloc.RootHeap = Allocation->GetRootHeap();
+					Alloc.bIsBlock = Allocation->IsHeap();
 					check(Alloc.Callstack != nullptr);
 				}
 
@@ -477,8 +493,18 @@ TSharedPtr<SWidget> SMemAllocTableTreeView::ConstructToolbar()
 			.Text(LOCTEXT("ByInvertedCallstackBtn_Text", "By Inverted Callstack"))
 			.ToolTipText(LOCTEXT("ByInvertedCallstackBtn_Tooltip", "Inverted Callstack Breakdown View\nConfigure the tree view to show a breakdown of allocations by inverted callstack."))
 			.OnClicked(this, &SMemAllocTableTreeView::OnCallstackViewClicked, true)
-		]
+		]	
 
+		+ SHorizontalBox::Slot()
+		.AutoWidth()
+		.Padding(4.0f, 0.0f, 0.0f, 0.0f)
+		[
+			SNew(SButton)
+			.Text(LOCTEXT("ByHeapTypeBtn_Text", "By Heap"))
+			.ToolTipText(LOCTEXT("ByHeapBtn_Tooltip", "Heap type breakdown."))
+			.OnClicked(this, &SMemAllocTableTreeView::OnHeapViewClicked)
+		]
+	
 		+ SHorizontalBox::Slot()
 		.AutoWidth()
 		.Padding(4.0f, 0.0f, 0.0f, 0.0f)
@@ -777,6 +803,46 @@ FReply SMemAllocTableTreeView::OnCallstackViewClicked(bool bIsInverted)
 
 ////////////////////////////////////////////////////////////////////////////////////////////////////
 
+FReply SMemAllocTableTreeView::OnHeapViewClicked()
+{
+	ColumnBeingSorted = FMemAllocTableColumns::SizeColumnId;
+	ColumnSortMode = EColumnSortMode::Type::Descending;
+	UpdateCurrentSortingByColumn();
+
+	PreChangeGroupings();
+
+	CurrentGroupings.Reset();
+
+	TSharedPtr<FTreeNodeGrouping>* CallstackGrouping = AvailableGroupings.FindByPredicate(
+		[](TSharedPtr<FTreeNodeGrouping>& Grouping)
+		{
+			return Grouping->Is<FMemAllocGroupingByHeap>();
+		});
+	if (CallstackGrouping)
+	{
+		CurrentGroupings.Add(*CallstackGrouping);
+	}
+
+	PostChangeGroupings();
+
+	FColumnConfig Preset[] =
+	{
+		{ FTable::GetHierarchyColumnId(),           true,  400.0f },
+		{ FMemAllocTableColumns::StartTimeColumnId, false, 0.0f },
+		{ FMemAllocTableColumns::EndTimeColumnId,   false, 0.0f },
+		{ FMemAllocTableColumns::DurationColumnId,  false, 0.0f },
+		{ FMemAllocTableColumns::AddressColumnId,   false, 0.0f },
+		{ FMemAllocTableColumns::CountColumnId,     true,  100.0f },
+		{ FMemAllocTableColumns::SizeColumnId,      true,  100.0f },
+		{ FMemAllocTableColumns::TagColumnId,       true,  200.0f },
+		{ FMemAllocTableColumns::FunctionColumnId,  true,  200.0f },
+	};
+	ApplyColumnConfig(TArrayView<FColumnConfig>(Preset, UE_ARRAY_COUNT(Preset)));
+
+	return FReply::Handled();
+}
+////////////////////////////////////////////////////////////////////////////////////////////////////
+
 void SMemAllocTableTreeView::ApplyColumnConfig(const TArrayView<FColumnConfig>& Preset)
 {
 	for (const TSharedRef<FTableColumn>& ColumnRef : Table->GetColumns())
@@ -908,6 +974,9 @@ void SMemAllocTableTreeView::InternalCreateGroupings()
 
 	AvailableGroupings.Insert(MakeShared<FMemAllocGroupingByCallstack>(false, bIsCallstackGroupingByFunction), Index++);
 	AvailableGroupings.Insert(MakeShared<FMemAllocGroupingByCallstack>(true, bIsCallstackGroupingByFunction), Index++);
+	
+	const TraceServices::IAllocationsProvider* AllocationsProvider = TraceServices::ReadAllocationsProvider(*Session.Get());
+	AvailableGroupings.Insert(MakeShared<FMemAllocGroupingByHeap>(*AllocationsProvider), Index++);
 }
 
 ////////////////////////////////////////////////////////////////////////////////////////////////////
diff --git a/Engine/Source/Developer/TraceInsights/Private/Insights/MemoryProfiler/Widgets/SMemAllocTableTreeView.h b/Engine/Source/Developer/TraceInsights/Private/Insights/MemoryProfiler/Widgets/SMemAllocTableTreeView.h
--- a/Engine/Source/Developer/TraceInsights/Private/Insights/MemoryProfiler/Widgets/SMemAllocTableTreeView.h
+++ b/Engine/Source/Developer/TraceInsights/Private/Insights/MemoryProfiler/Widgets/SMemAllocTableTreeView.h
@@ -110,6 +110,7 @@ private:
 	FReply OnTagViewClicked();
 	FReply OnMemoryPageClicked();
 	FReply OnCallstackViewClicked(bool bIsInverted);
+	FReply OnHeapViewClicked();
 
 	FText GetSymbolResolutionStatus() const;
 	FText GetSymbolResolutionTooltip() const;
diff --git a/Engine/Source/Developer/TraceServices/Private/Analyzers/AllocationsAnalysis.cpp b/Engine/Source/Developer/TraceServices/Private/Analyzers/AllocationsAnalysis.cpp
--- a/Engine/Source/Developer/TraceServices/Private/Analyzers/AllocationsAnalysis.cpp
+++ b/Engine/Source/Developer/TraceServices/Private/Analyzers/AllocationsAnalysis.cpp
@@ -2,6 +2,7 @@
 
 #include "AllocationsAnalysis.h"
 #include "Model/AllocationsProvider.h"
+#include "ProfilingDebugging/MemoryTrace.h"
 #include "TraceServices/Model/AnalysisSession.h"
 
 namespace TraceServices
@@ -32,14 +33,17 @@ void FAllocationsAnalyzer::OnAnalysisBegin(const FOnAnalysisContext& Context)
 	FInterfaceBuilder& Builder = Context.InterfaceBuilder;
 
 	Builder.RouteEvent(RouteId_Init,         "Memory", "Init");
-	Builder.RouteEvent(RouteId_CoreAdd,      "Memory", "CoreAdd");
-	Builder.RouteEvent(RouteId_CoreRemove,   "Memory", "CoreRemove");
 	Builder.RouteEvent(RouteId_Alloc,        "Memory", "Alloc");
+	Builder.RouteEvent(RouteId_AllocSystem,  "Memory", "AllocSystem");
+	Builder.RouteEvent(RouteId_AllocVideo,   "Memory", "AllocVideo");
 	Builder.RouteEvent(RouteId_Free,         "Memory", "Free");
 	Builder.RouteEvent(RouteId_ReallocAlloc, "Memory", "ReallocAlloc");
 	Builder.RouteEvent(RouteId_ReallocFree,  "Memory", "ReallocFree");
 	Builder.RouteEvent(RouteId_Marker,       "Memory", "Marker");
 	Builder.RouteEvent(RouteId_TagSpec,      "Memory", "TagSpec");
+	Builder.RouteEvent(RouteId_HeapSpec,	 "Memory", "HeapSpec");
+	Builder.RouteEvent(RouteId_HeapMarkAlloc,	 "Memory", "HeapMarkAlloc");
+	Builder.RouteEvent(RouteId_HeapUnmarkAlloc, "Memory", "HeapUnmarkAlloc");
 
 	Builder.RouteLoggerEvents(RouteId_MemScope, "Memory", true);
 }
@@ -68,70 +72,123 @@ void FAllocationsAnalyzer::OnAnalysisEnd()
 bool FAllocationsAnalyzer::OnEvent(uint16 RouteId, EStyle Style, const FOnEventContext& Context)
 {
 	const auto& EventData = Context.EventData;
+	HeapId RootHeap = 0;
+	
 	switch (RouteId)
 	{
 		case RouteId_Init:
 		{
 			const double Time = GetCurrentTime();
 			const uint8 MinAlignment = EventData.GetValue<uint8>("MinAlignment");
-			const uint8 SizeShift = EventData.GetValue<uint8>("SizeShift");
-			const uint8 SummarySizeShift = EventData.GetValue<uint8>("SummarySizeShift");
+			SizeShift = EventData.GetValue<uint8>("SizeShift");
 
 			BaseCycle = EventData.GetValue<uint64>("BaseCycle", 0);
 			MarkerPeriod = EventData.GetValue<uint32>("MarkerPeriod");
 
 			FAllocationsProvider::FEditScopeLock _(AllocationsProvider);
-			AllocationsProvider.EditInit(Time, MinAlignment, SizeShift, SummarySizeShift);
+			AllocationsProvider.EditInit(Time, MinAlignment);
 			break;
 		}
-		case RouteId_CoreAdd:
+		case RouteId_HeapSpec:
 		{
-			const double Time = GetCurrentTime();
-			const uint64 Owner = EventData.GetValue<uint64>("Owner");
-			const uint64 Base = EventData.GetValue<uint64>("Base");
-			const uint32 Size = EventData.GetValue<uint32>("Size");
+			const HeapId Id = EventData.GetValue<uint16>("Id");
+			const HeapId ParentId = EventData.GetValue<uint16>("ParentId");
+			const EMemoryTraceHeapFlags Flags = EventData.GetValue<EMemoryTraceHeapFlags>("Flags");
+			FStringView Name;
+			EventData.GetString("Name", Name);
 
 			FAllocationsProvider::FEditScopeLock _(AllocationsProvider);
-			AllocationsProvider.EditAddCore(Time, Owner, Base, Size);
+			AllocationsProvider.EditHeapSpec(Id, ParentId, Name, Flags);
 			break;
 		}
-		case RouteId_CoreRemove:
+		case RouteId_AllocSystem:
+		case RouteId_AllocVideo:
+		case RouteId_ReallocAllocSystem:
 		{
-			const double Time = GetCurrentTime();
-			const uint64 Owner = EventData.GetValue<uint64>("Owner");
-			const uint64 Base = EventData.GetValue<uint64>("Base");
-			const uint32 Size = EventData.GetValue<uint32>("Size");
-
-			FAllocationsProvider::FEditScopeLock _(AllocationsProvider);
-			AllocationsProvider.EditRemoveCore(Time, Owner, Base, Size);
-			break;
+			RootHeap = RouteId == RouteId_AllocVideo ? EMemoryTraceRootHeap::VideoMemory : EMemoryTraceRootHeap::SystemMemory;
 		}
+		// intentional fallthrough
 		case RouteId_Alloc:
 		case RouteId_ReallocAlloc:
 		{
 			// TODO: Can we have a struct mapping over the EventData?
 			//       something like: const auto& Ev = (const AllocEvent&)EventData.Get();
 			const double Time = GetCurrentTime();
-			const uint64 Owner = EventData.GetValue<uint64>("Owner");
-			const uint64 Address = EventData.GetValue<uint64>("Address");
-			const uint32 Size = EventData.GetValue<uint32>("Size");
+			uint64 Owner = EventData.GetValue<uint32>("CallstackId");
+			if (!Owner)
+			{
+				// Legeacy format of sending the hash value
+				Owner = EventData.GetValue<uint64>("Owner");
+			}
+			uint64 Address = EventData.GetValue<uint64>("Address");
+			RootHeap = EventData.GetValue<uint8>("RootHeap", RootHeap);
+			uint64 SizeUpper = EventData.GetValue<uint32>("Size");
+			const uint8 SizeLowerMask = ((1 << SizeShift) - 1);
+			const uint8 AlignmentMask = ~SizeLowerMask;
+			uint64 Size = 0;
+			uint32 Alignment = 0;
 			const uint8 Alignment_SizeLower = EventData.GetValue<uint8>("Alignment_SizeLower");
-
+			if (Alignment_SizeLower)
+			{
+				Size = SizeUpper << SizeShift | static_cast<uint64>(Alignment_SizeLower & SizeLowerMask);
+				Alignment = Alignment_SizeLower & AlignmentMask;
+			}
+			else
+			{
+				const uint8 AlignmentPow2_SizeLower = EventData.GetValue<uint8>("AlignmentPow2_SizeLower");
+				Size = SizeUpper << SizeShift | static_cast<uint64>(AlignmentPow2_SizeLower & SizeLowerMask);
+				Alignment = 1 << (AlignmentPow2_SizeLower >> SizeShift);
+			}
 			const uint32 ThreadId = Context.ThreadInfo.GetSystemId();
-			const uint8 Tracker = 0; // We only care about the default tracker for now.
 
 			FAllocationsProvider::FEditScopeLock _(AllocationsProvider);
-			AllocationsProvider.EditAlloc(Time, Owner, Address, Size, Alignment_SizeLower, ThreadId, Tracker);
+			AllocationsProvider.EditAlloc(Time, Owner, Address, Size, Alignment, ThreadId, 0, RootHeap);
 			break;
 		}
+
+		case RouteId_FreeSystem:
+		case RouteId_FreeVideo:
+		{
+			RootHeap = RouteId == RouteId_FreeSystem ? EMemoryTraceRootHeap::SystemMemory : EMemoryTraceRootHeap::VideoMemory;
+		}
+		// intentional fallthrough
 		case RouteId_Free:
 		case RouteId_ReallocFree:
 		{
 			const double Time = GetCurrentTime();
-			const uint64 Address = EventData.GetValue<uint64>("Address");
+			uint64 Address = EventData.GetValue<uint64>("Address");
+			if (Address == 0) //@todo: Check that this isn't intentional for reallocs?
+			{
+				constexpr uint32 HeapShift = 60;
+				constexpr uint64 RootHeapMask = uint64(0xF) << HeapShift;
+				const uint64 Address_RootHeap = EventData.GetValue<uint64>("Address_RootHeap");
+				Address = Address_RootHeap & ~RootHeapMask;
+				RootHeap = (Address_RootHeap & RootHeapMask) >> HeapShift;	
+			}
 
 			FAllocationsProvider::FEditScopeLock _(AllocationsProvider);
-			AllocationsProvider.EditFree(Time, Address);
+			AllocationsProvider.EditFree(Time, Address, RootHeap);
+			break;
+		}
+		case RouteId_HeapMarkAlloc:
+		{
+			const double Time = GetCurrentTime();
+			const uint64 Address = EventData.GetValue<uint64>("Address");
+			const HeapId Heap = EventData.GetValue<uint16>("Heap", 0);
+			const EMemoryTraceHeapAllocationFlags Flags = EventData.GetValue<EMemoryTraceHeapAllocationFlags>("Flags");
+			
+			FAllocationsProvider::FEditScopeLock _(AllocationsProvider);
+			AllocationsProvider.EditMarkAllocationAsHeap(Time, Address, Heap, Flags);
+			break;
+		}
+		case RouteId_HeapUnmarkAlloc:
+		{
+			const double Time = GetCurrentTime();
+			const uint64 Address = EventData.GetValue<uint64>("Address");
+			const HeapId Heap = EventData.GetValue<uint16>("Heap", 0);
+			
+			FAllocationsProvider::FEditScopeLock _(AllocationsProvider);
+			AllocationsProvider.EditUnmarkAllocationAsHeap(Time, Address, Heap);
 			break;
 		}
 		case RouteId_Marker:
@@ -185,7 +242,7 @@ bool FAllocationsAnalyzer::OnEvent(uint16 RouteId, EStyle Style, const FOnEventC
 					FAllocationsProvider::FEditScopeLock _(AllocationsProvider);
 					AllocationsProvider.EditPushTag(ThreadId, Tracker, Tag);
 				}
-				else // "MemoryScopeRealloc"
+				else // "MemoryScopeRealloc or MemoryScopePtr"
 				{
 					const uint64 Ptr = Context.EventData.GetValue<uint64>("Ptr");
 					FAllocationsProvider::FEditScopeLock _(AllocationsProvider);
diff --git a/Engine/Source/Developer/TraceServices/Private/Analyzers/AllocationsAnalysis.h b/Engine/Source/Developer/TraceServices/Private/Analyzers/AllocationsAnalysis.h
--- a/Engine/Source/Developer/TraceServices/Private/Analyzers/AllocationsAnalysis.h
+++ b/Engine/Source/Developer/TraceServices/Private/Analyzers/AllocationsAnalysis.h
@@ -17,15 +17,21 @@ private:
 	enum : uint16
 	{
 		RouteId_Init,
-		RouteId_CoreAdd,
-		RouteId_CoreRemove,
 		RouteId_Alloc,
+		RouteId_AllocSystem,
+		RouteId_AllocVideo,
 		RouteId_Free,
+		RouteId_FreeSystem,
+		RouteId_FreeVideo,
 		RouteId_ReallocAlloc,
+		RouteId_ReallocAllocSystem,
 		RouteId_ReallocFree,
 		RouteId_Marker,
 		RouteId_TagSpec,
 		RouteId_MemScope,
+		RouteId_HeapSpec,
+		RouteId_HeapMarkAlloc,
+		RouteId_HeapUnmarkAlloc,
 	};
 
 public:
@@ -43,6 +49,7 @@ private:
 	uint32 MarkerPeriod;
 	uint64 LastMarkerCycle;
 	double LastMarkerSeconds;
+	uint8 SizeShift;
 };
 
 } // namespace TraceServices
diff --git a/Engine/Source/Developer/TraceServices/Private/Analyzers/CallstacksAnalysis.cpp b/Engine/Source/Developer/TraceServices/Private/Analyzers/CallstacksAnalysis.cpp
--- a/Engine/Source/Developer/TraceServices/Private/Analyzers/CallstacksAnalysis.cpp
+++ b/Engine/Source/Developer/TraceServices/Private/Analyzers/CallstacksAnalysis.cpp
@@ -28,9 +28,15 @@ bool FCallstacksAnalyzer::OnEvent(uint16 RouteId, EStyle Style, const FOnEventCo
 	switch(RouteId)
 	{
 		case RouteId_Callstack:
-			const uint64 Id = Context.EventData.GetValue<uint64>("Id");
 			const TArrayReader<uint64>& Frames = Context.EventData.GetArray<uint64>("Frames");
-			Provider->AddCallstack(Id, Frames.GetData(), uint8(Frames.Num()));
+			if (const uint64 Hash = Context.EventData.GetValue<uint64>("Id"))
+			{
+				Provider->AddCallstack(Hash, Frames.GetData(), uint8(Frames.Num()));
+			}
+			else if (const uint32 Id = Context.EventData.GetValue<uint32>("CallstackId"))
+			{
+				Provider->AddCallstack(Id, Frames.GetData(), uint8(Frames.Num()));
+			}
 			break;
 	}
 	return true;
diff --git a/Engine/Source/Developer/TraceServices/Private/Model/AllocationsProvider.cpp b/Engine/Source/Developer/TraceServices/Private/Model/AllocationsProvider.cpp
--- a/Engine/Source/Developer/TraceServices/Private/Model/AllocationsProvider.cpp
+++ b/Engine/Source/Developer/TraceServices/Private/Model/AllocationsProvider.cpp
@@ -3,9 +3,11 @@
 #include "AllocationsProvider.h"
 
 #include "AllocationsQuery.h"
+#include "SbTree.h"
+#include "Algo/ForEach.h"
 #include "Common/Utils.h"
 #include "Containers/ArrayView.h"
-#include "SbTree.h"
+#include "ProfilingDebugging/MemoryTrace.h"
 #include "TraceServices/Containers/Allocators.h"
 #include "TraceServices/Model/Callstack.h"
 
@@ -274,6 +276,21 @@ uint32 IAllocationsProvider::FAllocation::GetTag() const
 	return Inner->Tag;
 }
 
+////////////////////////////////////////////////////////////////////////////////////////////////////
+
+HeapId IAllocationsProvider::FAllocation::GetRootHeap() const
+{
+	const auto* Inner = (const FAllocationItem*)this;
+	return Inner->RootHeap;
+}
+
+////////////////////////////////////////////////////////////////////////////////////////////////////
+bool IAllocationsProvider::FAllocation::IsHeap() const
+{
+	const auto* Inner = (const FAllocationItem*)this;
+	return Inner->IsHeap();
+}
+
 ////////////////////////////////////////////////////////////////////////////////////////////////////
 // IAllocationsProvider::FAllocations
 ////////////////////////////////////////////////////////////////////////////////////////////////////
@@ -486,6 +503,22 @@ void FShortLivingAllocs::Enumerate(TFunctionRef<void(const FAllocationItem& Allo
 	}
 }
 
+////////////////////////////////////////////////////////////////////////////////////////////////////
+FAllocationItem* FShortLivingAllocs::FindRange(const uint64 Address) const
+{
+	//todo: Can we do better than iterating all the nodes?
+	FNode* Node = LastAddedAllocNode;
+	while (Node != nullptr)
+	{
+		if (Node->Alloc->IsContained(Address))
+		{
+			return Node->Alloc;
+		}
+		Node = Node->Prev;
+	}
+	return nullptr;
+}
+
 ////////////////////////////////////////////////////////////////////////////////////////////////////
 // FLiveAllocCollection
 ////////////////////////////////////////////////////////////////////////////////////////////////////
@@ -542,6 +575,37 @@ FAllocationItem* FLiveAllocCollection::FindRef(uint64 Address)
 	return nullptr;
 }
 
+////////////////////////////////////////////////////////////////////////////////////////////////////
+///
+FAllocationItem* FLiveAllocCollection::FindParentHeap(uint64 Address)
+{
+#if INSIGHTS_USE_LAST_ALLOC
+	if (LastAlloc && LastAlloc->IsContained(Address))
+	{
+		return LastAlloc;
+	}
+#endif
+
+	for (const auto& Alloc : LongLivingAllocs)
+	{
+		if (Alloc.Value->IsContained(Address))
+		{
+			return Alloc.Value;
+		}
+	}
+	
+#if INSIGHTS_USE_SHORT_LIVING_ALLOCS
+	FAllocationItem* FoundShortLivingAlloc = ShortLivingAllocs.FindRange(Address);
+	if (FoundShortLivingAlloc)
+	{
+		INSIGHTS_SLOW_CHECK(FoundShortLivingAlloc->IsContained(Address));
+		return FoundShortLivingAlloc;
+	}
+#endif
+
+	return nullptr;
+}
+
 ////////////////////////////////////////////////////////////////////////////////////////////////////
 
 FAllocationItem* FLiveAllocCollection::AddNewChecked(uint64 Address)
@@ -651,7 +715,6 @@ void FLiveAllocCollection::Enumerate(TFunctionRef<void(const FAllocationItem& Al
 
 FAllocationsProvider::FAllocationsProvider(IAnalysisSession& InSession)
 	: Session(InSession)
-	, SbTree(nullptr)
 	, Timeline(Session.GetLinearAllocator(), 1024)
 	, MinTotalAllocatedMemoryTimeline(Session.GetLinearAllocator(), 1024)
 	, MaxTotalAllocatedMemoryTimeline(Session.GetLinearAllocator(), 1024)
@@ -660,24 +723,35 @@ FAllocationsProvider::FAllocationsProvider(IAnalysisSession& InSession)
 	, AllocEventsTimeline(Session.GetLinearAllocator(), 1024)
 	, FreeEventsTimeline(Session.GetLinearAllocator(), 1024)
 {
-	const uint32 ColumnShift = 17; // 1<<17 = 128K
-	SbTree = new FSbTree(Session.GetLinearAllocator(), ColumnShift);
+	
+	FMemory::Memset(EventIndex, 0, sizeof(EventIndex));
+	FMemory::Memset(SbTree, 0, sizeof(SbTree));
+	FMemory::Memset(LiveAllocs, 0, sizeof(LiveAllocs));
+	HeapSpecs.AddZeroed(256);
 }
 
 ////////////////////////////////////////////////////////////////////////////////////////////////////
 
 FAllocationsProvider::~FAllocationsProvider()
 {
-	if (SbTree)
+	for (uint8 RootHeapIdx = 0; RootHeapIdx < MaxRootHeaps; ++RootHeapIdx)
 	{
-		delete SbTree;
-		SbTree = nullptr;
+		if (SbTree[RootHeapIdx])
+		{
+			delete SbTree[RootHeapIdx];
+			SbTree[RootHeapIdx] = nullptr;
+		}
+		if (LiveAllocs[RootHeapIdx])
+		{
+			delete LiveAllocs[RootHeapIdx];
+			LiveAllocs[RootHeapIdx] = nullptr;
+		}
 	}
 }
 
 ////////////////////////////////////////////////////////////////////////////////////////////////////
 
-void FAllocationsProvider::EditInit(double InTime, uint8 InMinAlignment, uint8 InSizeShift, uint8 InSummarySizeShift)
+void FAllocationsProvider::EditInit(double InTime, uint8 InMinAlignment)
 {
 	Lock.WriteAccessCheck();
 
@@ -689,8 +763,6 @@ void FAllocationsProvider::EditInit(double InTime, uint8 InMinAlignment, uint8 I
 
 	InitTime = InTime;
 	MinAlignment = InMinAlignment;
-	SizeShift = InSizeShift;
-	SummarySizeShift = InSummarySizeShift;
 
 	bInitialized = true;
 
@@ -698,36 +770,43 @@ void FAllocationsProvider::EditInit(double InTime, uint8 InMinAlignment, uint8 I
 }
 
 ////////////////////////////////////////////////////////////////////////////////////////////////////
-
-void FAllocationsProvider::EditAddCore(double Time, uint64 Owner, uint64 Base, uint32 Size)
+void FAllocationsProvider::EditHeapSpec(HeapId Id, HeapId ParentId, const FStringView& Name, EMemoryTraceHeapFlags Flags)
 {
 	Lock.WriteAccessCheck();
 
-	if (!bInitialized)
+	if (Id < MaxRootHeaps)
 	{
-		return;
+		check(SbTree[Id] == nullptr && LiveAllocs[Id] == nullptr); //Already announced
+		constexpr uint32 ColumnShift = 17; // 1<<17 = 128K
+		SbTree[Id] = new FSbTree(Session.GetLinearAllocator(), ColumnShift);
+		LiveAllocs[Id] = new FLiveAllocCollection();
+		
+		FHeapSpec& RootHeapSpec = HeapSpecs[Id];
+		RootHeapSpec.Name = Session.StoreString(Name);
+		RootHeapSpec.Parent = nullptr;
+		RootHeapSpec.Id = Id;
+		RootHeapSpec.Flags = Flags;
 	}
-
-	//TODO
-}
-
-////////////////////////////////////////////////////////////////////////////////////////////////////
-
-void FAllocationsProvider::EditRemoveCore(double Time, uint64 Owner, uint64 Base, uint32 Size)
-{
-	Lock.WriteAccessCheck();
-
-	if (!bInitialized)
+	else
 	{
-		return;
-	}
+		FHeapSpec& ParentSpec = HeapSpecs[ParentId];
+		if (ParentSpec.Name == nullptr)
+		{
+			UE_LOG(LogTraceServices, Warning, TEXT("Parent heap id (%u) used before it was announced for heap %s."), ParentId, *FString(Name));
+		}
+		FHeapSpec& NewSpec = HeapSpecs[Id];
+		NewSpec.Name = Session.StoreString(Name);
+		NewSpec.Parent = &ParentSpec;
+		NewSpec.Id = Id;
+		NewSpec.Flags = Flags;
 
-	//TODO
+		ParentSpec.Children.Add(&NewSpec);
+	}
 }
 
 ////////////////////////////////////////////////////////////////////////////////////////////////////
 
-void FAllocationsProvider::EditAlloc(double Time, uint64 Owner, uint64 Address, uint32 InSize, uint8 InAlignmentAndSizeLower, uint32 ThreadId, uint8 Tracker)
+void FAllocationsProvider::EditAlloc(double Time, uint64 Owner, uint64 Address, uint64 InSize, uint32 InAlignment, uint32 ThreadId, uint8 Tracker, HeapId RootHeap)
 {
 	Lock.WriteAccessCheck();
 
@@ -741,13 +820,11 @@ void FAllocationsProvider::EditAlloc(double Time, uint64 Owner, uint64 Address,
 		return;
 	}
 
-	SbTree->SetTimeForEvent(EventIndex, Time);
+	check(RootHeap < MaxRootHeaps && HeapSpecs[RootHeap].Name != nullptr);
+	SbTree[RootHeap]->SetTimeForEvent(EventIndex[RootHeap], Time);
 
 	AdvanceTimelines(Time);
 
-	const uint8 SizeLowerMask = ((1 << SizeShift) - 1);
-	const uint8 AlignmentMask = ~SizeLowerMask;
-	const uint64 Size = (static_cast<uint64>(InSize) << SizeShift) | static_cast<uint64>(InAlignmentAndSizeLower & SizeLowerMask);
 
 	const uint32 Tag = TagTracker.GetCurrentTag(ThreadId, Tracker);
 
@@ -770,27 +847,29 @@ void FAllocationsProvider::EditAlloc(double Time, uint64 Owner, uint64 Address,
 
 	if (!AllocationPtr)
 	{
-		AllocationPtr = LiveAllocs.AddNewChecked(Address);
+		AllocationPtr = LiveAllocs[RootHeap]->AddNewChecked(Address);
 		INSIGHTS_SLOW_CHECK(Allocation.Address == Address)
 
 		FAllocationItem& Allocation = *AllocationPtr;
-		Allocation.StartEventIndex = EventIndex;
+		Allocation.StartEventIndex = EventIndex[RootHeap];
 		Allocation.EndEventIndex = (uint32)-1;
 		Allocation.StartTime = Time;
 		Allocation.EndTime = std::numeric_limits<double>::infinity();
 		Allocation.Owner = Owner;
 		//Allocation.Address = Address;
-		Allocation.SizeAndAlignment = FAllocationItem::PackSizeAndAlignment(Size, InAlignmentAndSizeLower & AlignmentMask);
+		Allocation.SizeAndAlignment = FAllocationItem::PackSizeAndAlignment(InSize, InAlignment);
 		Allocation.Callstack = nullptr;
 		Allocation.Tag = Tag;
-		Allocation.Reserved1 = 0;
+		Allocation.Flags = EMemoryTraceHeapAllocationFlags::None;
+		Allocation.RootHeap = RootHeap;
 
-		UpdateHistogramByAllocSize(Size);
+		UpdateHistogramByAllocSize(InSize);
 
 		// Update stats for the current timeline sample.
-		TotalAllocatedMemory += Size;
+		TotalAllocatedMemory += InSize;
+		++TotalLiveAllocations;
 		SampleMaxTotalAllocatedMemory = FMath::Max(SampleMaxTotalAllocatedMemory, TotalAllocatedMemory);
-		SampleMaxLiveAllocations = FMath::Max(SampleMaxLiveAllocations, (uint32)LiveAllocs.Num());
+		SampleMaxLiveAllocations = FMath::Max(SampleMaxLiveAllocations, TotalLiveAllocations); 
 		++SampleAllocEvents;
 	}
 #if INSIGHTS_VALIDATE_ALLOC_EVENTS
@@ -805,12 +884,12 @@ void FAllocationsProvider::EditAlloc(double Time, uint64 Owner, uint64 Address,
 #endif
 
 	++AllocCount;
-	++EventIndex;
+	++EventIndex[RootHeap];
 }
 
 ////////////////////////////////////////////////////////////////////////////////////////////////////
 
-void FAllocationsProvider::EditFree(double Time, uint64 Address)
+void FAllocationsProvider::EditFree(double Time, uint64 Address, HeapId RootHeap)
 {
 	Lock.WriteAccessCheck();
 
@@ -824,7 +903,8 @@ void FAllocationsProvider::EditFree(double Time, uint64 Address)
 		return;
 	}
 
-	SbTree->SetTimeForEvent(EventIndex, Time);
+	check(RootHeap < MaxRootHeaps && HeapSpecs[RootHeap].Name != nullptr);
+	SbTree[RootHeap]->SetTimeForEvent(EventIndex[RootHeap], Time);
 
 	AdvanceTimelines(Time);
 
@@ -839,24 +919,28 @@ void FAllocationsProvider::EditFree(double Time, uint64 Address)
 	}
 #endif // INSIGHTS_DEBUG_WATCH
 
-	FAllocationItem* AllocationPtr = LiveAllocs.Remove(Address); // we take ownership of AllocationPtr
+	FAllocationItem* AllocationPtr = LiveAllocs[RootHeap]->Remove(Address); // we take ownership of AllocationPtr
 	if (AllocationPtr)
 	{
-		check(EventIndex > AllocationPtr->StartEventIndex);
-		AllocationPtr->EndEventIndex = EventIndex;
+		check(EventIndex[RootHeap] > AllocationPtr->StartEventIndex);
+		AllocationPtr->EndEventIndex = EventIndex[RootHeap];
 		AllocationPtr->EndTime = Time;
 
 		const uint64 OldSize = AllocationPtr->GetSize();
 
-		SbTree->AddAlloc(AllocationPtr); // SbTree takes ownership of AllocationPtr
+		SbTree[RootHeap]->AddAlloc(AllocationPtr); // SbTree takes ownership of AllocationPtr
 
 		uint32 EventDistance = AllocationPtr->EndEventIndex - AllocationPtr->StartEventIndex;
 		UpdateHistogramByEventDistance(EventDistance);
 
-		// Update stats for the current timeline sample.
-		TotalAllocatedMemory -= OldSize;
-		SampleMinTotalAllocatedMemory = FMath::Min(SampleMinTotalAllocatedMemory, TotalAllocatedMemory);
-		SampleMinLiveAllocations = FMath::Min(SampleMinLiveAllocations, (uint32)LiveAllocs.Num());
+		// Update stats for the current timeline sample. (Heap allocations are already excluded)
+		if (!AllocationPtr->IsHeap())
+		{
+			TotalAllocatedMemory -= OldSize;
+			--TotalLiveAllocations;
+			SampleMinTotalAllocatedMemory = FMath::Min(SampleMinTotalAllocatedMemory, TotalAllocatedMemory);
+			SampleMinLiveAllocations = FMath::Min(SampleMinLiveAllocations, TotalLiveAllocations);
+		}
 		++SampleFreeEvents;
 	}
 	else
@@ -869,7 +953,65 @@ void FAllocationsProvider::EditFree(double Time, uint64 Address)
 	}
 
 	++FreeCount;
-	++EventIndex;
+	++EventIndex[RootHeap];
+}
+
+////////////////////////////////////////////////////////////////////////////////////////////////////
+
+void FAllocationsProvider::EditUnmarkAllocationAsHeap(double Time, uint64 Address, HeapId Heap)
+{
+	Lock.WriteAccessCheck();
+	
+	HeapId RootHeap = FindRootHeap(Heap);
+	
+	// Find allocation in Live allocs and mark as a heap allocation
+	FAllocationItem* Alloc = LiveAllocs[RootHeap]->FindRef(Address);
+	if (!Alloc)
+	{
+		UE_LOG(LogTraceServices, Error, TEXT("[MemAlloc] Could not find matching allocation for address %0x%llX while unmarking as heap."), Address);
+		return;
+	}
+
+	const uint64 Size = Alloc->GetSize();
+	const uint32 Alignment = Alloc->GetAlignment();
+	const uint32 Owner = Alloc->Owner;
+	const uint32 ThreadId = 0;
+	const uint8 Tracker = 1;
+
+	// We cannot just unmark the allocation as heap, there is no timestamp support, instead fake a free
+	// and allocation. Make sure the new allocation retains the tag from the original.
+	EditPushRealloc(ThreadId, Tracker, Address);
+	EditFree(Time, Address, RootHeap);
+	EditAlloc(Time, Address, Owner, Size, Alignment, 0, 1, RootHeap);
+	EditPopRealloc(ThreadId, Tracker);
+}
+
+////////////////////////////////////////////////////////////////////////////////////////////////////
+
+void FAllocationsProvider::EditMarkAllocationAsHeap(double Time, uint64 Address, HeapId Heap, EMemoryTraceHeapAllocationFlags Flags)
+{
+	Lock.WriteAccessCheck();
+	
+	HeapId RootHeap = FindRootHeap(Heap);
+	
+	// Find allocation in Live allocs and mark as a heap allocation
+	FAllocationItem* Alloc = LiveAllocs[RootHeap]->FindRef(Address);
+	if (Alloc)
+	{
+		check(EnumHasAnyFlags(Flags, EMemoryTraceHeapAllocationFlags::Heap));
+		Alloc->Flags = Flags;
+		Alloc->RootHeap = RootHeap;
+
+		// Remove this allocation from the total
+		TotalAllocatedMemory -= Alloc->GetSize();
+		--TotalLiveAllocations;
+		SampleMinTotalAllocatedMemory = FMath::Min(SampleMinTotalAllocatedMemory, TotalAllocatedMemory);
+		SampleMinLiveAllocations = FMath::Min(SampleMinLiveAllocations, TotalLiveAllocations);
+	}
+	else
+	{
+		UE_LOG(LogTraceServices, Error, TEXT("[MemAlloc] Could not find matching allocation for address %0x%llX while marking as heap."), Address);
+	}
 }
 
 ////////////////////////////////////////////////////////////////////////////////////////////////////
@@ -934,11 +1076,10 @@ void FAllocationsProvider::AdvanceTimelines(double Time)
 
 		// Start a new sample.
 		SampleStartTimestamp = Time;
-		const uint32 NumLiveAllocs = (uint32)LiveAllocs.Num();
 		SampleMinTotalAllocatedMemory = TotalAllocatedMemory;
 		SampleMaxTotalAllocatedMemory = TotalAllocatedMemory;
-		SampleMinLiveAllocations = NumLiveAllocs;
-		SampleMaxLiveAllocations = NumLiveAllocs;
+		SampleMinLiveAllocations = TotalLiveAllocations;
+		SampleMaxLiveAllocations = TotalLiveAllocations;
 		SampleAllocEvents = 0;
 		SampleFreeEvents = 0;
 
@@ -949,8 +1090,8 @@ void FAllocationsProvider::AdvanceTimelines(double Time)
 			Timeline.EmplaceBack(SampleEndTimestamp);
 			MinTotalAllocatedMemoryTimeline.EmplaceBack(TotalAllocatedMemory);
 			MaxTotalAllocatedMemoryTimeline.EmplaceBack(TotalAllocatedMemory);
-			MinLiveAllocationsTimeline.EmplaceBack(NumLiveAllocs);
-			MaxLiveAllocationsTimeline.EmplaceBack(NumLiveAllocs);
+			MinLiveAllocationsTimeline.EmplaceBack(TotalLiveAllocations);
+			MaxLiveAllocationsTimeline.EmplaceBack(TotalLiveAllocations);
 			AllocEventsTimeline.EmplaceBack(0);
 			FreeEventsTimeline.EmplaceBack(0);
 		}
@@ -961,14 +1102,24 @@ void FAllocationsProvider::AdvanceTimelines(double Time)
 
 ////////////////////////////////////////////////////////////////////////////////////////////////////
 
-void FAllocationsProvider::EditPushRealloc(uint32 ThreadId, uint8 Tracker, uint64 Ptr)
+HeapId FAllocationsProvider::FindRootHeap(HeapId Heap) const
 {
-	EditAccessCheck();
-	int32 Tag(0); // If ptr is not found use "Untagged"
-	if (FAllocationItem* Alloc = LiveAllocs.FindRef(Ptr))
+	const FHeapSpec* RootHeap = &HeapSpecs[Heap];
+	while (RootHeap->Parent != nullptr)
 	{
-		Tag = Alloc->Tag;
+		RootHeap = RootHeap->Parent;
 	}
+	return RootHeap->Id;
+}
+
+////////////////////////////////////////////////////////////////////////////////////////////////////
+
+void FAllocationsProvider::EditPushRealloc(uint32 ThreadId, uint8 Tracker, uint64 Ptr)
+{
+	EditAccessCheck();
+	// todo: Currently only system root heap is affected by reallocs so limit search
+	FAllocationItem* Alloc = LiveAllocs[EMemoryTraceRootHeap::SystemMemory]->FindRef(Ptr);
+	const int32 Tag = Alloc ? Alloc->Tag : 0; // If ptr is not found use "Untagged"
 	TagTracker.PushRealloc(ThreadId, Tracker, Tag);
 }
 
@@ -1017,7 +1168,7 @@ void FAllocationsProvider::EditOnAnalysisCompleted(double Time)
 	{
 		AdvanceTimelines(Time + 10 * DefaultTimelineSampleGranularity);
 
-		const uint32 LiveAllocsTotalCount = (uint32)LiveAllocs.Num();
+		const uint32 LiveAllocsTotalCount = TotalLiveAllocations;
 		LiveAllocs.Empty();
 
 		// Update stats for the last timeline sample (reset to zero).
@@ -1035,7 +1186,21 @@ void FAllocationsProvider::EditOnAnalysisCompleted(double Time)
 	DebugPrint();
 #endif
 
-	SbTree->Validate();
+	for (const FSbTree* Tree : SbTree)
+	{
+		if (Tree)
+		{
+			Tree->Validate();
+		}
+	}
+}
+
+////////////////////////////////////////////////////////////////////////////////////////////////////
+
+void FAllocationsProvider::EnumerateRootHeaps(TFunctionRef<void(HeapId Id, const FHeapSpec&)> Callback) const
+{
+	Algo::ForEachIf(HeapSpecs, [](const FHeapSpec& Spec) { return Spec.Parent == nullptr && Spec.Name != nullptr; },
+	                [&](const FHeapSpec& Spec) { Callback(Spec.Id, Spec); });
 }
 
 ////////////////////////////////////////////////////////////////////////////////////////////////////
@@ -1297,7 +1462,13 @@ void FAllocationsProvider::EnumerateFreeEventsTimeline(int32 StartIndex, int32 E
 
 void FAllocationsProvider::DebugPrint() const
 {
-	SbTree->DebugPrint();
+	for(const FSbTree* Tree : SbTree)
+	{
+		if (Tree)
+		{
+			Tree->DebugPrint();
+		}
+	}
 }
 
 ////////////////////////////////////////////////////////////////////////////////////////////////////
@@ -1326,13 +1497,32 @@ const IAllocationsProvider::FQueryStatus FAllocationsProvider::PollQuery(FQueryH
 }
 
 ////////////////////////////////////////////////////////////////////////////////////////////////////
+void FAllocationsProvider::EnumerateLiveAllocs(TFunctionRef<void(const FAllocationItem& Alloc)> Callback) const
+{
+	ReadAccessCheck();
+	for (const FLiveAllocCollection* Allocs : LiveAllocs)
+	{
+		if (Allocs)
+		{
+			Allocs->Enumerate(Callback);
+		}
+	}
+}
 
+////////////////////////////////////////////////////////////////////////////////////////////////////
 FName GetAllocationsProviderName()
 {
 	static FName Name(TEXT("AllocationsProvider"));
 	return Name;
 }
 
+////////////////////////////////////////////////////////////////////////////////////////////////////
+uint32 FAllocationsProvider::GetNumLiveAllocs() const
+{
+	ReadAccessCheck();
+	return TotalLiveAllocations;
+}
+
 ////////////////////////////////////////////////////////////////////////////////////////////////////
 
 const IAllocationsProvider* ReadAllocationsProvider(const IAnalysisSession& Session)
diff --git a/Engine/Source/Developer/TraceServices/Private/Model/AllocationsProvider.h b/Engine/Source/Developer/TraceServices/Private/Model/AllocationsProvider.h
--- a/Engine/Source/Developer/TraceServices/Private/Model/AllocationsProvider.h
+++ b/Engine/Source/Developer/TraceServices/Private/Model/AllocationsProvider.h
@@ -5,6 +5,7 @@
 #include "Common/PagedArray.h"
 #include "Containers/Map.h"
 #include "HAL/CriticalSection.h"
+#include "ProfilingDebugging/MemoryTrace.h"
 #include "TraceServices/Model/Callstack.h"
 #include "TraceServices/Model/AllocationsProvider.h"
 
@@ -87,8 +88,10 @@ struct FAllocationItem
 	static uint32 UnpackAlignment(uint64 SizeAndAlignment) { return static_cast<uint32>(SizeAndAlignment >> AlignmentShift); }
 	static uint64 PackSizeAndAlignment(uint64 Size, uint8 Alignment) { return (Size | (static_cast<uint64>(Alignment) << AlignmentShift)); }
 
+	FORCEINLINE bool IsContained(uint64 InAddress) const { return Address >= InAddress && Address < (Address + GetSize()); }
 	uint64 GetSize() const { return UnpackSize(SizeAndAlignment); }
 	uint32 GetAlignment() const { return UnpackAlignment(SizeAndAlignment); }
+	bool IsHeap() const { return EnumHasAnyFlags(Flags, EMemoryTraceHeapAllocationFlags::Heap); }
 
 	uint32 StartEventIndex;
 	uint32 EndEventIndex;
@@ -99,7 +102,8 @@ struct FAllocationItem
 	uint64 SizeAndAlignment; // (Alignment << AlignmentShift) | Size
 	mutable const FCallstack* Callstack;
 	uint32 Tag;
-	uint32 Reserved1;
+	uint8 RootHeap;
+	EMemoryTraceHeapAllocationFlags Flags;
 };
 
 static_assert(sizeof(FAllocationItem) == 64, "struct FAllocationItem needs packing");
@@ -150,6 +154,9 @@ public:
 
 	void Enumerate(TFunctionRef<void(const FAllocationItem& Alloc)> Callback) const;
 
+	// Finds an allocation that contains the address specified or null if not found.
+	FAllocationItem* FindRange(uint64 Address) const;
+
 private:
 	TMap<uint64, FNode*> AddressMap; // map for short living allocations: Address -> FNode*
 	FNode* AllNodes = nullptr; // preallocated array of nodes
@@ -167,11 +174,14 @@ public:
 	FLiveAllocCollection();
 	~FLiveAllocCollection();
 
-	int32 Num() const { return TotalAllocCount; }
-	int32 PeakCount() const { return MaxAllocCount; }
-
+	uint32 Num() const { return TotalAllocCount; }
+	uint32 PeakCount() const { return MaxAllocCount; }
+	
 	FORCEINLINE FAllocationItem* FindRef(uint64 Address);
 
+	// Finds the block containing the address or nullptr if the address was not found.
+	FORCEINLINE FAllocationItem* FindParentHeap(uint64 Address);
+
 	// The collection keeps ownership of FAllocationItem* until Remove is called.
 	// Returns the new added allocation.
 	FORCEINLINE FAllocationItem* AddNewChecked(uint64 Address);
@@ -187,8 +197,8 @@ private:
 	FShortLivingAllocs ShortLivingAllocs; // short living allocs
 	TMap<uint64, FAllocationItem*> LongLivingAllocs; // long living allocations
 
-	int32 TotalAllocCount = 0;
-	int32 MaxAllocCount = 0; // debug stats
+	uint32 TotalAllocCount = 0;
+	uint32 MaxAllocCount = 0; // debug stats
 };
 
 ////////////////////////////////////////////////////////////////////////////////////////////////////
@@ -216,6 +226,7 @@ public:
 	virtual bool IsInitialized() const override { ReadAccessCheck(); return bInitialized; }
 
 	virtual uint32 GetTimelineNumPoints() const override { ReadAccessCheck(); return Timeline.Num(); }
+	virtual void EnumerateRootHeaps(TFunctionRef<void(HeapId Id, const FHeapSpec&)> Callback) const override;
 	virtual void GetTimelineIndexRange(double StartTime, double EndTime, int32& StartIndex, int32& EndIndex) const override;
 	virtual void EnumerateMinTotalAllocatedMemoryTimeline(int32 StartIndex, int32 EndIndex, TFunctionRef<void(double Time, double Duration, uint64 Value)> Callback) const override;
 	virtual void EnumerateMaxTotalAllocatedMemoryTimeline(int32 StartIndex, int32 EndIndex, TFunctionRef<void(double Time, double Duration, uint64 Value)> Callback) const override;
@@ -228,10 +239,10 @@ public:
 	virtual void CancelQuery(FQueryHandle Query) const override;
 	virtual const FQueryStatus PollQuery(FQueryHandle Query) const override;
 
-	const FSbTree* GetSbTree() const { ReadAccessCheck(); return SbTree; }
+	const FSbTree* GetSbTreeUnchecked(HeapId Heap) const { ReadAccessCheck(); return SbTree[Heap]; }
 
-	void EnumerateLiveAllocs(TFunctionRef<void(const FAllocationItem& Alloc)> Callback) const { ReadAccessCheck(); return LiveAllocs.Enumerate(Callback); }
-	uint32 GetNumLiveAllocs() const { ReadAccessCheck(); return LiveAllocs.Num(); }
+	void EnumerateLiveAllocs(TFunctionRef<void(const FAllocationItem& Alloc)> Callback) const;
+	uint32 GetNumLiveAllocs() const;
 
 	bool HasReallocScope(uint32 ThreadId, uint8 Tracker) const { ReadAccessCheck(); return TagTracker.HasReallocScope(ThreadId, Tracker); }
 
@@ -242,14 +253,15 @@ public:
 	//////////////////////////////////////////////////
 	// Edit operations
 
-	void EditInit(double Time, uint8 MinAlignment, uint8 SizeShift, uint8 SummarySizeShift);
-
-	void EditAddCore(double Time, uint64 Owner, uint64 Base, uint32 Size);
-	void EditRemoveCore(double Time, uint64 Owner, uint64 Base, uint32 Size);
+	void EditInit(double Time, uint8 MinAlignment);
 
-	void EditAlloc(double Time, uint64 Owner, uint64 Address, uint32 Size, uint8 AlignmentAndSizeLower, uint32 ThreadId, uint8 Tracker);
-	void EditFree(double Time, uint64 Address);
+	void EditHeapSpec(HeapId Id, HeapId ParentId, const FStringView& Name, EMemoryTraceHeapFlags Flags);
 
+	void EditAlloc(double Time, uint64 Owner, uint64 Address, uint64 Size, uint32 Alignment, uint32 ThreadId, uint8 Tracker, HeapId RootHeap);
+	void EditFree(double Time, uint64 Address, HeapId RootHeap);
+	void EditMarkAllocationAsHeap(double Time, uint64 Address, HeapId Heap, EMemoryTraceHeapAllocationFlags Flags);
+	void EditUnmarkAllocationAsHeap(double Time, uint64 Address, HeapId Heap);
+	
 	void EditAddTagSpec(int32 Tag, uint32 ParentTag, const TCHAR* Display) { EditAccessCheck(); TagTracker.AddTagSpec(Tag, ParentTag, Display); }
 	void EditPushTag(uint32 ThreadId, uint8 Tracker, uint32 Tag)           { EditAccessCheck(); TagTracker.PushTag(ThreadId, Tracker, Tag); }
 	void EditPopTag(uint32 ThreadId, uint8 Tracker)                        { EditAccessCheck(); TagTracker.PopTag(ThreadId, Tracker); }
@@ -265,12 +277,16 @@ private:
 	void UpdateHistogramByAllocSize(uint64 Size);
 	void UpdateHistogramByEventDistance(uint32 EventDistance);
 	void AdvanceTimelines(double Time);
+	HeapId FindRootHeap(HeapId Heap) const;
 
 private:
 	IAnalysisSession& Session;
 
 	mutable FAllocationsProviderLock Lock;
 
+	//Number of supported root heaps
+	constexpr static uint8 MaxRootHeaps = 16;
+	
 	double InitTime = 0;
 	uint8 MinAlignment = 0;
 	uint8 SizeShift = 0;
@@ -279,12 +295,12 @@ private:
 
 	FTagTracker TagTracker;
 
-	uint32 EventIndex = 0;
+	uint32 EventIndex[MaxRootHeaps];
 
 	uint64 AllocCount = 0;
 	uint64 FreeCount = 0;
 
-	FLiveAllocCollection LiveAllocs;
+	FLiveAllocCollection* LiveAllocs[MaxRootHeaps];
 
 	uint64 AllocErrors = 0;
 	uint64 FreeErrors = 0;
@@ -295,9 +311,10 @@ private:
 	uint32 MaxEventDistance = 0;
 	uint32 EventDistanceHistogramPow2[33] = { 0 };
 
-	FSbTree* SbTree;
+	FSbTree* SbTree[MaxRootHeaps];
 
 	uint64 TotalAllocatedMemory = 0;
+	uint32 TotalLiveAllocations = 0;
 
 	double SampleStartTimestamp = 0.0;
 	double SampleEndTimestamp = 0.0;
@@ -315,6 +332,8 @@ private:
 	TPagedArray<uint32> MaxLiveAllocationsTimeline;
 	TPagedArray<uint32> AllocEventsTimeline;
 	TPagedArray<uint32> FreeEventsTimeline;
+	
+	TArray<FHeapSpec> HeapSpecs;
 };
 
 ////////////////////////////////////////////////////////////////////////////////////////////////////
diff --git a/Engine/Source/Developer/TraceServices/Private/Model/AllocationsQuery.cpp b/Engine/Source/Developer/TraceServices/Private/Model/AllocationsQuery.cpp
--- a/Engine/Source/Developer/TraceServices/Private/Model/AllocationsQuery.cpp
+++ b/Engine/Source/Developer/TraceServices/Private/Model/AllocationsQuery.cpp
@@ -107,72 +107,80 @@ void FAllocationsQuery::Run()
 
 	AllocationsProvider.BeginRead();
 
-	const FSbTree* SbTree = AllocationsProvider.GetSbTree();
-	if (SbTree)
-	{
-		UE_LOG(LogTraceServices, Log, TEXT("[MemAlloc] Processing %d live allocs..."), AllocationsProvider.GetNumLiveAllocs());
+	UE_LOG(LogTraceServices, Log, TEXT("[MemAlloc] Processing %d live allocs..."), AllocationsProvider.GetNumLiveAllocs());
 
-		FAllocationsImpl* LiveAllocsResult = new FAllocationsImpl();
-		QueryLiveAllocs(LiveAllocsResult->Items);
+	FAllocationsImpl* LiveAllocsResult = new FAllocationsImpl();
+	QueryLiveAllocs(LiveAllocsResult->Items);
 
-		const uint32 NumLiveAllocs = static_cast<uint32>(LiveAllocsResult->Items.Num());
-		if (NumLiveAllocs != 0)
-		{
-			UE_LOG(LogTraceServices, Log, TEXT("[MemAlloc] Enqueue %u live allocs..."), NumLiveAllocs);
-			TotalAllocationCount += NumLiveAllocs;
+	const uint32 NumLiveAllocs = static_cast<uint32>(LiveAllocsResult->Items.Num());
+	if (NumLiveAllocs != 0)
+	{
+		UE_LOG(LogTraceServices, Log, TEXT("[MemAlloc] Enqueue %u live allocs..."), NumLiveAllocs);
+		TotalAllocationCount += NumLiveAllocs;
 
-			QueryCallstacks(LiveAllocsResult);
-			Results.Enqueue(LiveAllocsResult);
-		}
-		else
-		{
-			delete LiveAllocsResult;
-		}
+		QueryCallstacks(LiveAllocsResult);
+		Results.Enqueue(LiveAllocsResult);
+	}
+	else
+	{
+		delete LiveAllocsResult;
+	}
 
-		UE_LOG(LogTraceServices, Log, TEXT("[MemAlloc] Detecting cells..."));
+	UE_LOG(LogTraceServices, Log, TEXT("[MemAlloc] Detecting cells..."));
 
-		TArray<const FSbTreeCell*> Cells;
+	TArray<HeapId> Heaps;
+	TArray<const FSbTreeCell*> Cells;
 
+	AllocationsProvider.EnumerateRootHeaps([&Heaps](HeapId Heap, const IAllocationsProvider::FHeapSpec& Spec)
+	{
+		Heaps.Add(Heap);
+	});
+
+	for (HeapId Heap : Heaps)
+	{
 		if (!IsCanceling)
 		{
+			const FSbTree* SbTree = AllocationsProvider.GetSbTreeUnchecked(Heap);
 			SbTree->Query(Cells, Params);
-			CellCount += Cells.Num();
 		}
+	}
+		
+	CellCount += Cells.Num();
 
-		UE_LOG(LogTraceServices, Log, TEXT("[MemAlloc] %d cells to process"), Cells.Num());
-		uint32 CellIndex = 0;
+	UE_LOG(LogTraceServices, Log, TEXT("[MemAlloc] %d cells to process"), Cells.Num());
+	uint32 CellIndex = 0;
 
-		for (const FSbTreeCell* Cell : Cells)
+	for (const FSbTreeCell* Cell : Cells)
+	{
+		if (IsCanceling)
 		{
-			if (IsCanceling)
-			{
-				break;
-			}
+			break;
+		}
 
-			++CellIndex;
-			UE_LOG(LogTraceServices, Log, TEXT("[MemAlloc] Processing cell %u (%u allocs)..."), CellIndex, Cell->GetAllocCount());
+		++CellIndex;
+		UE_LOG(LogTraceServices, Log, TEXT("[MemAlloc] Processing cell %u (%u allocs)..."), CellIndex, Cell->GetAllocCount());
 
-			FAllocationsImpl* CellResult = new FAllocationsImpl();
-			Cell->Query(CellResult->Items, Params);
+		FAllocationsImpl* CellResult = new FAllocationsImpl();
+			
+		Cell->Query(CellResult->Items, Params);
 
-			const uint32 NumAllocs = static_cast<uint32>(CellResult->Items.Num());
-			if (NumAllocs != 0)
-			{
-				UE_LOG(LogTraceServices, Log, TEXT("[MemAlloc] Enqueue %u allocs..."), NumAllocs);
-				TotalAllocationCount += NumAllocs;
+		const uint32 NumAllocs = static_cast<uint32>(CellResult->Items.Num());
+		if (NumAllocs != 0)
+		{
+			UE_LOG(LogTraceServices, Log, TEXT("[MemAlloc] Enqueue %u allocs..."), NumAllocs);
+			TotalAllocationCount += NumAllocs;
 
-				QueryCallstacks(CellResult);
-				Results.Enqueue(CellResult);
-			}
-			else
-			{
-				delete CellResult;
-			}
+			QueryCallstacks(CellResult);
+			Results.Enqueue(CellResult);
+		}
+		else
+		{
+			delete CellResult;
 		}
-
-		UE_LOG(LogTraceServices, Log, TEXT("[MemAlloc] Done"));
 	}
 
+	UE_LOG(LogTraceServices, Log, TEXT("[MemAlloc] Done"));
+
 	AllocationsProvider.EndRead();
 
 	IsWorking = false;
diff --git a/Engine/Source/Developer/TraceServices/Private/Model/CallstacksProvider.cpp b/Engine/Source/Developer/TraceServices/Private/Model/CallstacksProvider.cpp
--- a/Engine/Source/Developer/TraceServices/Private/Model/CallstacksProvider.cpp
+++ b/Engine/Source/Developer/TraceServices/Private/Model/CallstacksProvider.cpp
@@ -93,6 +93,14 @@ void FCallstacksProvider::AddCallstack(uint64 InCallstackId, const uint64* InFra
 	}
 }
 
+/////////////////////////////////////////////////////////////////////
+void FCallstacksProvider::AddCallstack(uint32 InCallstackId, const uint64* InFrames, uint8 InFrameCount)
+{
+	// In order to maintain backward compatibility to older format where we sent the runtime hash value
+	// we keep the callstack entry key as uint64.
+	AddCallstack(uint64(InCallstackId), InFrames, InFrameCount);
+}
+
 /////////////////////////////////////////////////////////////////////
 const FCallstack* FCallstacksProvider::GetCallstack(uint64 CallstackId) const
 {
diff --git a/Engine/Source/Developer/TraceServices/Private/Model/CallstacksProvider.h b/Engine/Source/Developer/TraceServices/Private/Model/CallstacksProvider.h
--- a/Engine/Source/Developer/TraceServices/Private/Model/CallstacksProvider.h
+++ b/Engine/Source/Developer/TraceServices/Private/Model/CallstacksProvider.h
@@ -23,6 +23,7 @@ public:
 	const FCallstack*	GetCallstack(uint64 CallstackId) const override;
 	void				GetCallstacks(const TArrayView<uint64>& CallstackIds, FCallstack const** OutCallstacks) const override;
 	void				AddCallstack(uint64 CallstackId, const uint64* Frames, uint8 FrameCount);
+	void				AddCallstack(uint32 CallstackId, const uint64* Frames, uint8 FrameCount);
 
 private:
 	enum
diff --git a/Engine/Source/Developer/TraceServices/Public/TraceServices/Model/AllocationsProvider.h b/Engine/Source/Developer/TraceServices/Public/TraceServices/Model/AllocationsProvider.h
--- a/Engine/Source/Developer/TraceServices/Public/TraceServices/Model/AllocationsProvider.h
+++ b/Engine/Source/Developer/TraceServices/Public/TraceServices/Model/AllocationsProvider.h
@@ -5,6 +5,7 @@
 #include "AnalysisSession.h"
 #include "CoreTypes.h"
 #include "Templates/UniquePtr.h"
+#include "ProfilingDebugging/MemoryTrace.h"
 
 #include <new>
 
@@ -92,6 +93,8 @@ public:
 		uint32 GetAlignment() const;
 		const struct FCallstack* GetCallstack() const;
 		uint32 GetTag() const;
+		HeapId GetRootHeap() const;
+		bool IsHeap() const;
 	};
 
 	class TRACESERVICES_API FAllocations
@@ -120,6 +123,15 @@ public:
 		mutable UPTRINT Handle;
 	};
 
+	struct TRACESERVICES_API FHeapSpec
+	{
+		HeapId Id;
+		FHeapSpec* Parent;
+		TArray<FHeapSpec*> Children;
+		const TCHAR* Name;
+		EMemoryTraceHeapFlags Flags;
+	};
+
 	typedef UPTRINT FQueryHandle;
 
 public:
@@ -139,6 +151,8 @@ public:
 	// Index values are in range { -1, 0, .. , N-1, N }, where N = GetTimelineNumPoints().
 	virtual void GetTimelineIndexRange(double StartTime, double EndTime, int32& StartIndex, int32& EndIndex) const = 0;
 
+	virtual void EnumerateRootHeaps(TFunctionRef<void(HeapId Id, const FHeapSpec&)> Callback) const = 0;
+
 	// Enumerates the Min Total Allocated Memory timeline points in the inclusive index interval [StartIndex, EndIndex].
 	virtual void EnumerateMinTotalAllocatedMemoryTimeline(int32 StartIndex, int32 EndIndex, TFunctionRef<void(double Time, double Duration, uint64 Value)> Callback) const = 0;
 
