Fix backwards compatibility with older memory traces.

#rnx

diff --git a/Engine/Source/Developer/TraceServices/Private/Model/AllocationsProvider.cpp b/Engine/Source/Developer/TraceServices/Private/Model/AllocationsProvider.cpp
--- a/Engine/Source/Developer/TraceServices/Private/Model/AllocationsProvider.cpp
+++ b/Engine/Source/Developer/TraceServices/Private/Model/AllocationsProvider.cpp
@@ -728,6 +728,9 @@ FAllocationsProvider::FAllocationsProvider(IAnalysisSession& InSession)
 	FMemory::Memset(SbTree, 0, sizeof(SbTree));
 	FMemory::Memset(LiveAllocs, 0, sizeof(LiveAllocs));
 	HeapSpecs.AddZeroed(256);
+
+	// Create system root heap structures for backwards compatibility (before heap description events)
+	AddHeapSpec(EMemoryTraceRootHeap::SystemMemory, 0, TEXT("System memory"), EMemoryTraceHeapFlags::Root);
 }
 
 ////////////////////////////////////////////////////////////////////////////////////////////////////
@@ -770,10 +773,9 @@ void FAllocationsProvider::EditInit(double InTime, uint8 InMinAlignment)
 }
 
 ////////////////////////////////////////////////////////////////////////////////////////////////////
-void FAllocationsProvider::EditHeapSpec(HeapId Id, HeapId ParentId, const FStringView& Name, EMemoryTraceHeapFlags Flags)
-{
-	Lock.WriteAccessCheck();
 
+void FAllocationsProvider::AddHeapSpec(HeapId Id, HeapId ParentId, const FStringView& Name, EMemoryTraceHeapFlags Flags)
+{
 	if (Id < MaxRootHeaps)
 	{
 		check(SbTree[Id] == nullptr && LiveAllocs[Id] == nullptr); //Already announced
@@ -806,6 +808,15 @@ void FAllocationsProvider::EditHeapSpec(HeapId Id, HeapId ParentId, const FStrin
 
 ////////////////////////////////////////////////////////////////////////////////////////////////////
 
+void FAllocationsProvider::EditHeapSpec(HeapId Id, HeapId ParentId, const FStringView& Name, EMemoryTraceHeapFlags Flags)
+{
+	Lock.WriteAccessCheck();
+
+	AddHeapSpec(Id, ParentId, Name, Flags);
+}
+
+////////////////////////////////////////////////////////////////////////////////////////////////////
+
 void FAllocationsProvider::EditAlloc(double Time, uint64 Owner, uint64 Address, uint64 InSize, uint32 InAlignment, uint32 ThreadId, uint8 Tracker, HeapId RootHeap)
 {
 	Lock.WriteAccessCheck();
@@ -820,7 +831,13 @@ void FAllocationsProvider::EditAlloc(double Time, uint64 Owner, uint64 Address,
 		return;
 	}
 
-	check(RootHeap < MaxRootHeaps && HeapSpecs[RootHeap].Name != nullptr);
+	check(RootHeap < MaxRootHeaps);
+	if (SbTree[RootHeap] == nullptr)
+	{
+		UE_LOG(LogTraceServices, Warning, TEXT("Invalid root heap id (%d). This heap has not yet been registered."), RootHeap);
+		return;	
+	}
+	
 	SbTree[RootHeap]->SetTimeForEvent(EventIndex[RootHeap], Time);
 
 	AdvanceTimelines(Time);
@@ -903,7 +920,13 @@ void FAllocationsProvider::EditFree(double Time, uint64 Address, HeapId RootHeap
 		return;
 	}
 
-	check(RootHeap < MaxRootHeaps && HeapSpecs[RootHeap].Name != nullptr);
+	check(RootHeap < MaxRootHeaps);
+	if (SbTree[RootHeap] == nullptr)
+	{
+		UE_LOG(LogTraceServices, Warning, TEXT("Invalid root heap id (%d). This heap has not yet been registered."), RootHeap);
+		return;	
+	}
+	
 	SbTree[RootHeap]->SetTimeForEvent(EventIndex[RootHeap], Time);
 
 	AdvanceTimelines(Time);
@@ -1118,9 +1141,13 @@ void FAllocationsProvider::EditPushRealloc(uint32 ThreadId, uint8 Tracker, uint6
 {
 	EditAccessCheck();
 	// todo: Currently only system root heap is affected by reallocs so limit search
-	FAllocationItem* Alloc = LiveAllocs[EMemoryTraceRootHeap::SystemMemory]->FindRef(Ptr);
-	const int32 Tag = Alloc ? Alloc->Tag : 0; // If ptr is not found use "Untagged"
-	TagTracker.PushRealloc(ThreadId, Tracker, Tag);
+	FLiveAllocCollection* Allocs = LiveAllocs[EMemoryTraceRootHeap::SystemMemory];
+	if (Allocs)
+	{
+		FAllocationItem* Alloc = Allocs->FindRef(Ptr);
+		const int32 Tag = Alloc ? Alloc->Tag : 0; // If ptr is not found use "Untagged"
+		TagTracker.PushRealloc(ThreadId, Tracker, Tag);
+	}
 }
 
 ////////////////////////////////////////////////////////////////////////////////////////////////////
diff --git a/Engine/Source/Developer/TraceServices/Private/Model/AllocationsProvider.h b/Engine/Source/Developer/TraceServices/Private/Model/AllocationsProvider.h
--- a/Engine/Source/Developer/TraceServices/Private/Model/AllocationsProvider.h
+++ b/Engine/Source/Developer/TraceServices/Private/Model/AllocationsProvider.h
@@ -277,6 +277,7 @@ private:
 	void UpdateHistogramByAllocSize(uint64 Size);
 	void UpdateHistogramByEventDistance(uint32 EventDistance);
 	void AdvanceTimelines(double Time);
+	void AddHeapSpec(HeapId Id, HeapId ParentId, const FStringView& Name, EMemoryTraceHeapFlags Flags);
 	HeapId FindRootHeap(HeapId Heap) const;
 
 private:
diff --git a/Engine/Source/Runtime/Core/Private/ProfilingDebugging/MemoryAllocationTrace.cpp b/Engine/Source/Runtime/Core/Private/ProfilingDebugging/MemoryAllocationTrace.cpp
--- a/Engine/Source/Runtime/Core/Private/ProfilingDebugging/MemoryAllocationTrace.cpp
+++ b/Engine/Source/Runtime/Core/Private/ProfilingDebugging/MemoryAllocationTrace.cpp
@@ -114,7 +114,7 @@ void FAllocationTrace::Initialize() const
 
 	const HeapId SystemRootHeap = RootHeapSpec(TEXT("System memory"));
 	check(SystemRootHeap == EMemoryTraceRootHeap::SystemMemory);
-	const HeapId VideoRootHeap = RootHeapSpec(TEXT("VideoMemory"));
+	const HeapId VideoRootHeap = RootHeapSpec(TEXT("Video memory"));
 	check(VideoRootHeap == EMemoryTraceRootHeap::VideoMemory);
 
 	static_assert((1 << SizeShift) - 1 <= MIN_ALIGNMENT, "Not enough bits to pack size fields");
