Assertion should exlude system root heap which is automatically added.

#rnx

diff --git a/Engine/Source/Developer/TraceServices/Private/Model/AllocationsProvider.cpp b/Engine/Source/Developer/TraceServices/Private/Model/AllocationsProvider.cpp
--- a/Engine/Source/Developer/TraceServices/Private/Model/AllocationsProvider.cpp
+++ b/Engine/Source/Developer/TraceServices/Private/Model/AllocationsProvider.cpp
@@ -778,7 +778,7 @@ void FAllocationsProvider::AddHeapSpec(HeapId Id, HeapId ParentId, const FString
 {
 	if (Id < MaxRootHeaps)
 	{
-		check(SbTree[Id] == nullptr && LiveAllocs[Id] == nullptr); //Already announced
+		check(Id==0 || (SbTree[Id] == nullptr && LiveAllocs[Id] == nullptr)); //Already announced
 		constexpr uint32 ColumnShift = 17; // 1<<17 = 128K
 		SbTree[Id] = new FSbTree(Session.GetLinearAllocator(), ColumnShift);
 		LiveAllocs[Id] = new FLiveAllocCollection();
