Fix static analysis warning

diff --git a/Engine/Source/Developer/TraceServices/Public/Model/AsyncEnumerateTask.h b/Engine/Source/Developer/TraceServices/Public/Model/AsyncEnumerateTask.h
--- a/Engine/Source/Developer/TraceServices/Public/Model/AsyncEnumerateTask.h
+++ b/Engine/Source/Developer/TraceServices/Public/Model/AsyncEnumerateTask.h
@@ -136,7 +136,6 @@ public:
 				{
 					--CurrentStackDepth;
 					InEventType CurrentEvent = DetailLevel.GetEvent(CurrentScopePage->InitialStack[CurrentStackDepth].EventIndex);
-					double StartTime = DetailLevel.GetScopeEntryTime(CurrentScopePage->InitialStack[CurrentStackDepth].EnterScopeIndex);
 					if (Callback(false, CurrentScopePage->InitialStack[CurrentStackDepth].EndTime, CurrentEvent) == EEventEnumerate::Stop)
 					{
 						return;
