TasksInsights: removed TaskCreated check as invalid. While not confirmed, it's assumed that the check fails when a task is created and launched from different threads and "task launched" event arrives first, which is OK and should be handled correctly

diff --git a/Engine/Source/Developer/TraceServices/Private/Model/TasksProfiler.cpp b/Engine/Source/Developer/TraceServices/Private/Model/TasksProfiler.cpp
--- a/Engine/Source/Developer/TraceServices/Private/Model/TasksProfiler.cpp
+++ b/Engine/Source/Developer/TraceServices/Private/Model/TasksProfiler.cpp
@@ -117,8 +117,6 @@ namespace TraceServices
 			return;
 		}
 
-		checkf(Task->CreatedTimestamp == FTaskInfo::InvalidTimestamp, TEXT("%d"), TaskId);
-
 		Task->Id = TaskId;
 		Task->CreatedTimestamp = Timestamp;
 		Task->CreatedThreadId = ThreadId;
