[Insights] Fixed usage of UE_LOG.

diff --git a/Engine/Source/Developer/TraceServices/Private/Model/AllocationsProvider.cpp b/Engine/Source/Developer/TraceServices/Private/Model/AllocationsProvider.cpp
--- a/Engine/Source/Developer/TraceServices/Private/Model/AllocationsProvider.cpp
+++ b/Engine/Source/Developer/TraceServices/Private/Model/AllocationsProvider.cpp
@@ -865,7 +865,7 @@ void FAllocationsProvider::EditAlloc(double Time, uint64 Owner, uint64 Address,
 	if (!AllocationPtr)
 	{
 		AllocationPtr = LiveAllocs[RootHeap]->AddNewChecked(Address);
-		INSIGHTS_SLOW_CHECK(Allocation.Address == Address)
+		INSIGHTS_SLOW_CHECK(AllocationPtr->Address == Address)
 
 		FAllocationItem& Allocation = *AllocationPtr;
 		Allocation.StartEventIndex = EventIndex[RootHeap];
@@ -991,7 +991,7 @@ void FAllocationsProvider::EditUnmarkAllocationAsHeap(double Time, uint64 Addres
 	FAllocationItem* Alloc = LiveAllocs[RootHeap]->FindRef(Address);
 	if (!Alloc)
 	{
-		UE_LOG(LogTraceServices, Error, TEXT("[MemAlloc] Could not find matching allocation for address %0x%llX while unmarking as heap."), Address);
+		UE_LOG(LogTraceServices, Error, TEXT("[MemAlloc] Could not find matching allocation for address 0x%llX while unmarking as heap."), Address);
 		return;
 	}
 
@@ -1033,7 +1033,7 @@ void FAllocationsProvider::EditMarkAllocationAsHeap(double Time, uint64 Address,
 	}
 	else
 	{
-		UE_LOG(LogTraceServices, Error, TEXT("[MemAlloc] Could not find matching allocation for address %0x%llX while marking as heap."), Address);
+		UE_LOG(LogTraceServices, Error, TEXT("[MemAlloc] Could not find matching allocation for address 0x%llX while marking as heap."), Address);
 	}
 }
 
