Rename the TraceServices CsvProfiler.h/cpp CsvProfilerProvider.h/cpp to avoid confusion with the actual Csv profiler, and fix inconveniences with VAX Find in Files/Swap to Header

diff --git a/Engine/Source/Developer/TraceServices/Private/Analyzers/CsvProfilerTraceAnalysis.h b/Engine/Source/Developer/TraceServices/Private/Analyzers/CsvProfilerTraceAnalysis.h
--- a/Engine/Source/Developer/TraceServices/Private/Analyzers/CsvProfilerTraceAnalysis.h
+++ b/Engine/Source/Developer/TraceServices/Private/Analyzers/CsvProfilerTraceAnalysis.h
@@ -4,7 +4,7 @@
 
 #include "Trace/Analyzer.h"
 #include "Containers/Map.h"
-#include "Model/CsvProfilerPrivate.h"
+#include "Model/CsvProfilerProviderPrivate.h"
 #include "ProfilingDebugging/MiscTrace.h"
 
 namespace TraceServices
diff --git a/Engine/Source/Developer/TraceServices/Private/Model/CsvProfiler.cpp b/Engine/Source/Developer/TraceServices/Private/Model/CsvProfiler.cpp
deleted file mode 100644
diff --git a/Engine/Source/Developer/TraceServices/Private/Model/CsvProfilerPrivate.h b/Engine/Source/Developer/TraceServices/Private/Model/CsvProfilerPrivate.h
deleted file mode 100644
diff --git a/Engine/Source/Developer/TraceServices/Private/Model/CsvProfilerProvider.cpp b/Engine/Source/Developer/TraceServices/Private/Model/CsvProfilerProvider.cpp
new file mode 100644
--- /dev/null
+++ b/Engine/Source/Developer/TraceServices/Private/Model/CsvProfilerProvider.cpp
@@ -0,0 +1,253 @@
+// Copyright Epic Games, Inc. All Rights Reserved.
+
+#include "TraceServices/Model/CsvProfilerProvider.h"
+#include "CsvProfilerProviderPrivate.h"
+
+namespace TraceServices
+{
+
+uint64 FCsvProfilerProvider::FTableLayout::GetColumnCount() const
+{
+	return StatSeries.Num() + 1;
+}
+
+const TCHAR* FCsvProfilerProvider::FTableLayout::GetColumnName(uint64 ColumnIndex) const
+{
+	if (ColumnIndex == 0)
+	{
+		return TEXT("EVENTS");
+	}
+	else
+	{
+		return StatSeries[ColumnIndex - 1]->Name;
+	}
+}
+
+ETableColumnType FCsvProfilerProvider::FTableLayout::GetColumnType(uint64 ColumnIndex) const
+{
+	if (ColumnIndex == 0)
+	{
+		return TableColumnType_CString;
+	}
+	else
+	{
+		ECsvStatSeriesType SeriesType = StatSeries[ColumnIndex - 1]->Type;
+		if (SeriesType == CsvStatSeriesType_CustomStatInt)
+		{
+			return TableColumnType_Int;
+		}
+		else
+		{
+			return TableColumnType_Double;
+		}
+	}
+}
+
+bool FCsvProfilerProvider::FTableReader::IsValid() const
+{
+	return CurrentRowIndex < (Capture.EndFrame - Capture.StartFrame);
+}
+
+void FCsvProfilerProvider::FTableReader::NextRow()
+{
+	++CurrentRowIndex;
+}
+
+void FCsvProfilerProvider::FTableReader::SetRowIndex(uint64 RowIndex)
+{
+	CurrentRowIndex = RowIndex;
+}
+
+bool FCsvProfilerProvider::FTableReader::GetValueBool(uint64 ColumnIndex) const
+{
+	return false;
+}
+
+const FCsvProfilerProvider::FStatSeriesValue* FCsvProfilerProvider::FTableReader::GetValue(uint64 ColumnIndex) const
+{
+	if (!IsValid() || ColumnIndex == 0 || Capture.StatSeries.Num() <= ColumnIndex - 1)
+	{
+		return nullptr;
+	}
+	const FStatSeriesValue* FindIt = Capture.StatSeries[ColumnIndex - 1]->Values.Find(Capture.StartFrame + CurrentRowIndex);
+	return FindIt;
+}
+
+int64 FCsvProfilerProvider::FTableReader::GetValueInt(uint64 ColumnIndex) const
+{
+	const FStatSeriesValue* Value = GetValue(ColumnIndex);
+	if (!Value)
+	{
+		return 0;
+	}
+	return Value->Value.AsInt;
+}
+
+float FCsvProfilerProvider::FTableReader::GetValueFloat(uint64 ColumnIndex) const
+{
+	return static_cast<float>(GetValueDouble(ColumnIndex));
+}
+
+double FCsvProfilerProvider::FTableReader::GetValueDouble(uint64 ColumnIndex) const
+{
+	const FStatSeriesValue* Value = GetValue(ColumnIndex);
+	if (!Value)
+	{
+		return 0;
+	}
+	return Value->Value.AsDouble;
+}
+
+const TCHAR* FCsvProfilerProvider::FTableReader::GetValueCString(uint64 ColumnIndex) const
+{
+	if (ColumnIndex != 0)
+	{
+		return nullptr;
+	}
+	const FEvents* const* FindIt = Events.Find(Capture.StartFrame + CurrentRowIndex);
+	if (!FindIt)
+	{
+		return TEXT("");
+	}
+	return (*FindIt)->SemiColonSeparatedEvents;
+}
+
+uint64 FCsvProfilerProvider::FTable::GetRowCount() const
+{
+	return Capture.EndFrame - Capture.StartFrame;
+}
+
+IUntypedTableReader* FCsvProfilerProvider::FTable::CreateReader() const
+{
+	return new FCsvProfilerProvider::FTableReader(Capture, Events);
+}
+
+FCsvProfilerProvider::FCsvProfilerProvider(IAnalysisSession& InSession)
+	: Session(InSession)
+{
+
+}
+
+FCsvProfilerProvider::~FCsvProfilerProvider()
+{
+	for (FCapture* Capture : Captures)
+	{
+		delete Capture;
+	}
+	delete CurrentCapture;
+	for (FStatSeries* Series : StatSeries)
+	{
+		delete Series;
+	}
+}
+
+void FCsvProfilerProvider::EnumerateCaptures(TFunctionRef<void(const FCaptureInfo&)> Callback) const
+{
+	Session.ReadAccessCheck();
+
+	uint32 Index = 0;
+	for (const FCapture* Capture : Captures)
+	{
+		Callback({ Capture->Metadata, Capture->Filename, Index++, uint32(Capture->EndFrame - Capture->StartFrame) });
+	}
+}
+
+void FCsvProfilerProvider::StartCapture(const TCHAR* Filename, int32 FrameNumber)
+{
+	Session.WriteAccessCheck();
+
+	check(!CurrentCapture);
+	CurrentCapture = new FCapture(Events);
+	CurrentCapture->StartFrame = FrameNumber;
+	CurrentCapture->EndFrame = FrameNumber;
+	CurrentCapture->Filename = Filename;
+}
+
+void FCsvProfilerProvider::EndCapture(int32 FrameNumber)
+{
+	Session.WriteAccessCheck();
+
+	check(CurrentCapture);
+
+	CurrentCapture->EndFrame = FrameNumber;
+	check(CurrentCapture->EndFrame >= CurrentCapture->StartFrame);
+
+	CurrentCapture->Metadata = Metadata;
+
+	Captures.Add(CurrentCapture);
+	CurrentCapture = nullptr;
+}
+
+uint64 FCsvProfilerProvider::AddSeries(const TCHAR* Name, ECsvStatSeriesType Type)
+{
+	Session.WriteAccessCheck();
+
+	FStatSeries* Series = new FStatSeries();
+	Series->Name = Name;
+	Series->Type = Type;
+	int32 ColumnIndex = StatSeries.Num();
+	StatSeries.Add(Series);
+	return ColumnIndex;
+}
+
+FCsvProfilerProvider::FStatSeriesValue& FCsvProfilerProvider::GetValueRef(uint64 SeriesHandle, uint32 FrameNumber)
+{
+	FStatSeries* Series = StatSeries[SeriesHandle];
+	if (CurrentCapture && !Series->Captures.Contains(Captures.Num()))
+	{
+		CurrentCapture->StatSeries.Add(Series);
+		Series->Captures.Add(Captures.Num());
+	}
+	return Series->Values.FindOrAdd(FrameNumber);
+}
+
+void FCsvProfilerProvider::SetValue(uint64 SeriesHandle, uint32 FrameNumber, int64 Value)
+{
+	Session.WriteAccessCheck();
+
+	FStatSeriesValue& RowValue = GetValueRef(SeriesHandle, FrameNumber);
+	RowValue.Value.AsInt = Value;
+}
+
+void FCsvProfilerProvider::SetValue(uint64 SeriesHandle, uint32 FrameNumber, double Value)
+{
+	Session.WriteAccessCheck();
+
+	FStatSeriesValue& RowValue = GetValueRef(SeriesHandle, FrameNumber);
+	RowValue.Value.AsDouble = Value;
+}
+
+void FCsvProfilerProvider::AddEvent(uint32 FrameNumber, const TCHAR* Text)
+{
+	Session.WriteAccessCheck();
+
+	FEvents* FrameEvents;
+	FEvents** FindIt = Events.Find(FrameNumber);
+	if (!FindIt)
+	{
+		FrameEvents = new FEvents();
+		Events.Add(FrameNumber, FrameEvents);
+	}
+	else
+	{
+		FrameEvents = *FindIt;
+	}
+	if (FrameEvents->Events.Num())
+	{
+		FrameEvents->SemiColonSeparatedEvents = Session.StoreString(*(FString(FrameEvents->SemiColonSeparatedEvents) + ";" + Text));
+	}
+	else
+	{
+		FrameEvents->SemiColonSeparatedEvents = Text;
+	}
+	FrameEvents->Events.Add(Text);
+}
+
+void FCsvProfilerProvider::SetMetadata(const TCHAR* Key, const TCHAR* Value)
+{
+	Session.WriteAccessCheck();
+
+	Metadata.Add(Key, Value);
+}
+
+} // namespace TraceServices
diff --git a/Engine/Source/Developer/TraceServices/Private/Model/CsvProfilerProviderPrivate.h b/Engine/Source/Developer/TraceServices/Private/Model/CsvProfilerProviderPrivate.h
new file mode 100644
--- /dev/null
+++ b/Engine/Source/Developer/TraceServices/Private/Model/CsvProfilerProviderPrivate.h
@@ -0,0 +1,160 @@
+// Copyright Epic Games, Inc. All Rights Reserved.
+
+#pragma once
+
+#include "TraceServices/Model/CsvProfilerProvider.h"
+#include "TraceServices/Containers/Tables.h"
+#include "Containers/Map.h"
+
+namespace TraceServices
+{
+
+enum ECsvStatSeriesType
+{
+	CsvStatSeriesType_Timer,
+	CsvStatSeriesType_CustomStatInt,
+	CsvStatSeriesType_CustomStatFloat
+};
+
+class FCsvProfilerProvider
+	: public ICsvProfilerProvider
+{
+public:
+	explicit FCsvProfilerProvider(IAnalysisSession& InSession);
+	virtual ~FCsvProfilerProvider();
+	virtual const IUntypedTable& GetTable(uint32 CaptureIndex) const override { return Captures[CaptureIndex]->Table; }
+	virtual void EnumerateCaptures(TFunctionRef<void(const FCaptureInfo&)> Callback) const override;
+	void StartCapture(const TCHAR* Filename, int32 FrameNumber);
+	void EndCapture(int32 FrameNumber);
+	uint64 AddSeries(const TCHAR* Name, ECsvStatSeriesType Type);
+	void SetValue(uint64 SeriesHandle, uint32 FrameNumber, double Value);
+	void SetValue(uint64 SeriesHandle, uint32 FrameNumber, int64 Value);
+	void AddEvent(uint32 FrameNumber, const TCHAR* Text);
+	void SetMetadata(const TCHAR* Key, const TCHAR* Value);
+
+private:
+	struct FStatSeriesValue
+	{
+		FStatSeriesValue() { Value.AsInt = 0; }
+		union
+		{
+			int64 AsInt;
+			double AsDouble;
+		} Value;
+	};
+
+	struct FStatSeries
+	{
+		TMap<uint64, FStatSeriesValue> Values;
+		TSet<uint32> Captures;
+		const TCHAR* Name = nullptr;
+		ECsvStatSeriesType Type;
+	};
+
+	struct FEvents
+	{
+		TArray<const TCHAR*> Events;
+		const TCHAR* SemiColonSeparatedEvents = nullptr;
+	};
+
+	struct FCapture;
+
+	class FTableLayout
+		: public ITableLayout
+	{
+	public:
+		FTableLayout(const TArray<FStatSeries*>& InStatSeries)
+			: StatSeries(InStatSeries)
+		{
+
+		}
+
+		virtual uint64 GetColumnCount() const override;
+		virtual const TCHAR* GetColumnName(uint64 ColumnIndex) const override;
+		virtual ETableColumnType GetColumnType(uint64 ColumnIndex) const override;
+		virtual uint32 GetColumnDisplayHintFlags(uint64 ColumnIndex) const override
+		{
+			return 0;
+		}
+
+	private:
+		const TArray<FStatSeries*>& StatSeries;
+	};
+
+	class FTableReader
+		: public IUntypedTableReader
+	{
+	public:
+		FTableReader(const FCapture& InCapture, const TMap<uint32, FEvents*>& InEvents)
+			: Capture(InCapture)
+			, Events(InEvents)
+			, CurrentRowIndex(0)
+		{
+		}
+
+		virtual bool IsValid() const override;
+		virtual void NextRow() override;
+		virtual void SetRowIndex(uint64 RowIndex) override;
+		virtual bool GetValueBool(uint64 ColumnIndex) const override;
+		virtual int64 GetValueInt(uint64 ColumnIndex) const override;
+		virtual float GetValueFloat(uint64 ColumnIndex) const override;
+		virtual double GetValueDouble(uint64 ColumnIndex) const override;
+		virtual const TCHAR* GetValueCString(uint64 ColumnIndex) const override;
+
+	private:
+		const FStatSeriesValue* GetValue(uint64 ColumnIndex) const;
+
+		const FCapture& Capture;
+		const TMap<uint32, FEvents*>& Events;
+		uint64 CurrentRowIndex;
+	};
+
+	class FTable
+		: public IUntypedTable
+	{
+	public:
+		FTable(const FCapture& InCapture, const TMap<uint32, FEvents*>& InEvents)
+			: Layout(InCapture.StatSeries)
+			, Capture(InCapture)
+			, Events(InEvents)
+		{
+
+		}
+
+		virtual const ITableLayout& GetLayout() const override { return Layout; }
+		virtual uint64 GetRowCount() const override;
+		virtual IUntypedTableReader* CreateReader() const override;
+
+	private:
+		FTableLayout Layout;
+		const FCapture& Capture;
+		const TMap<uint32, FEvents*>& Events;
+	};
+
+	struct FCapture
+	{
+		FCapture(const TMap<uint32, FEvents*>& Events)
+			: Table(*this, Events)
+		{
+			
+		}
+
+		FTable Table;
+		const TCHAR* Filename;
+		uint32 StartFrame = 0;
+		uint32 EndFrame = 0;
+		TArray<FStatSeries*> StatSeries;
+		TMap<const TCHAR*, const TCHAR*> Metadata;
+	};
+
+	FStatSeriesValue& GetValueRef(uint64 SeriesHandle, uint32 FrameNumber);
+
+	IAnalysisSession& Session;
+	TArray<FStatSeries*> StatSeries;
+	TMap<uint32, FEvents*> Events;
+	TMap<const TCHAR*, const TCHAR*> Metadata;
+	TArray<FCapture*> Captures;
+	FCapture* CurrentCapture = nullptr;
+};
+
+} // namespace TraceServices
diff --git a/Engine/Source/Developer/TraceServices/Public/TraceServices/Model/CsvProfiler.h b/Engine/Source/Developer/TraceServices/Public/TraceServices/Model/CsvProfiler.h
deleted file mode 100644
diff --git a/Engine/Source/Developer/TraceServices/Public/TraceServices/Model/CsvProfilerProvider.h b/Engine/Source/Developer/TraceServices/Public/TraceServices/Model/CsvProfilerProvider.h
new file mode 100644
--- /dev/null
+++ b/Engine/Source/Developer/TraceServices/Public/TraceServices/Model/CsvProfilerProvider.h
@@ -0,0 +1,32 @@
+// Copyright Epic Games, Inc. All Rights Reserved.
+
+#pragma once
+
+#include "TraceServices/Model/AnalysisSession.h"
+#include "Templates/Function.h"
+
+namespace TraceServices
+{
+
+class IUntypedTable;
+
+struct FCaptureInfo
+{
+	const TMap<const TCHAR*, const TCHAR*>& Metadata;
+	const TCHAR* Filename = nullptr;
+	uint32 Id = uint32(-1);
+	uint32 FrameCount = 0;
+};
+
+class ICsvProfilerProvider
+	: public IProvider
+{
+public:
+	virtual ~ICsvProfilerProvider() = default;
+	virtual void EnumerateCaptures(TFunctionRef<void(const FCaptureInfo&)> Callback) const = 0;
+	virtual const IUntypedTable& GetTable(uint32 CaptureId) const = 0;
+};
+
+const ICsvProfilerProvider* ReadCsvProfilerProvider(const IAnalysisSession& Session);
+
+} // namespace TraceServices
