Update syms library to latest version, includes dwarf5 support

diff --git a/Engine/Source/Developer/TraceServices/Private/Model/SymslibResolver.cpp b/Engine/Source/Developer/TraceServices/Private/Model/SymslibResolver.cpp
--- a/Engine/Source/Developer/TraceServices/Private/Model/SymslibResolver.cpp
+++ b/Engine/Source/Developer/TraceServices/Private/Model/SymslibResolver.cpp
@@ -3,13 +3,16 @@
 #if PLATFORM_WINDOWS
 
 #include "SymslibResolver.h"
+#include "Algo/ForEach.h"
 #include "Algo/Sort.h"
 #include "Async/MappedFileHandle.h"
 #include "Async/ParallelFor.h"
 #include "Containers/StringView.h"
 #include "GenericPlatform/GenericPlatformFile.h"
 #include "HAL/PlatformFileManager.h"
+#include "HAL/PlatformTLS.h"
 #include "Logging/LogMacros.h"
+#include "Misc/CString.h"
 #include "Misc/Paths.h"
 #include "Misc/PathViews.h"
 #include "Misc/ScopeLock.h"
@@ -18,8 +21,6 @@
 #include "TraceServices/Model/AnalysisSession.h"
 #include <atomic>
 
-#include "symslib.h"
-
 /////////////////////////////////////////////////////////////////////
 DEFINE_LOG_CATEGORY_STATIC(LogSymslib, Log, All);
 
@@ -32,6 +33,144 @@ static const TCHAR* GUnknownModuleTextSymsLib = TEXT("Unknown");
 
 /////////////////////////////////////////////////////////////////////
 
+
+namespace
+{
+	struct FAutoMappedFile
+	{
+		FString FilePath;
+
+		TUniquePtr<IMappedFileHandle> Handle;
+		TUniquePtr<IMappedFileRegion> Region;
+
+		bool Load(const TCHAR* FileName)
+		{
+			Handle.Reset(FPlatformFileManager::Get().GetPlatformFile().OpenMapped(FileName));
+			Region.Reset(Handle.IsValid() ? Handle->MapRegion(0, Handle->GetFileSize()) : nullptr);
+			return Region.IsValid();
+		}
+
+		inline void* GetData() const
+		{
+			return Region.IsValid() ? (void*)Region->GetMappedPtr() : nullptr;
+		}
+
+		inline uint64 Num() const
+		{
+			return Region.IsValid() ? Region->GetMappedSize() : 0;
+		}
+	};
+
+	static FString FindSymbolFile(const FString& Path)
+	{
+		IPlatformFile* PlatformFile = &FPlatformFileManager::Get().GetPlatformFile();
+
+		// Extract only filename part in case Path is absolute path
+		FString FileName = FPaths::GetCleanFilename(Path);
+
+		FString SymbolPath = FPlatformMisc::GetEnvironmentVariable(L"UE_INSIGHTS_SYMBOL_PATH");
+
+		FString SymbolPathPart;
+		while (SymbolPath.Split(TEXT(";"), &SymbolPathPart, &SymbolPath))
+		{
+			FString Result = FPaths::Combine(SymbolPathPart, FileName);
+			if (PlatformFile->FileExists(*Result))
+			{
+				// File found in one of symbol paths
+				return Result;
+			}
+		}
+
+		FString Result = FPaths::Combine(SymbolPath, FileName);
+		if (PlatformFile->FileExists(*Result))
+		{
+			// File found in symbol path
+			return Result;
+		}
+
+		if (PlatformFile->FileExists(*Path))
+		{
+			// File found by absolute path
+			return Path;
+		}
+
+		// File not found
+		return FString();
+	}
+
+	static SYMS_String8 SymsLoader(void* User, SYMS_Arena* Arena, SYMS_String8 FileName)
+	{
+		IPlatformFile* PlatformFile = &FPlatformFileManager::Get().GetPlatformFile();
+
+		FString FilePath = ANSI_TO_TCHAR(reinterpret_cast<char*>(FileName.str));
+
+		TArray<FAutoMappedFile>* Files = static_cast<TArray<FAutoMappedFile>*>(User);
+		FAutoMappedFile& File = Files->AddDefaulted_GetRef();
+
+		if (Files->Num() == 1)
+		{
+			// we're loading main executable file
+			FString SymbolFilePath = FindSymbolFile(FilePath);
+			if (SymbolFilePath.IsEmpty())
+			{
+				UE_LOG(LogSymslib, Warning, TEXT("File '%s' not found"), *FilePath);
+				return syms_str8(nullptr, 0);
+			}
+
+			if (!File.Load(*SymbolFilePath))
+			{
+				UE_LOG(LogSymslib, Warning, TEXT("Failed to open '%s'"), *SymbolFilePath);
+			}
+
+			File.FilePath = SymbolFilePath;
+		}
+		else
+		{
+			// we're loading secondary file, like .pdb for .exe
+
+			// main .exe file folder, so we can lookup relative paths
+			FString ModuleBasePath = FPaths::GetPath((*Files)[0].FilePath);
+
+			bool bDebugFilePathFound = false;
+
+			// First try path itself if it is absolute
+			FString DebugFilePath = FilePath;
+			if (!FPathViews::IsRelativePath(*DebugFilePath))
+			{
+				if (PlatformFile->FileExists(*DebugFilePath))
+				{
+					bDebugFilePathFound = true;
+				}
+			}
+
+			// Otherwise try relative path to module file
+			if (!bDebugFilePathFound)
+			{
+				DebugFilePath = FPaths::Combine(ModuleBasePath, FPaths::GetCleanFilename(DebugFilePath));
+				if (PlatformFile->FileExists(*DebugFilePath))
+				{
+					bDebugFilePathFound = true;
+				}
+			}
+
+			// Use path if file was found
+			if (bDebugFilePathFound)
+			{
+				if (!File.Load(*DebugFilePath))
+				{
+					UE_LOG(LogSymslib, Warning, TEXT("Failed to open '%s'"), *DebugFilePath);
+				}
+
+				File.FilePath = DebugFilePath;
+			}
+		}
+
+		return syms_str8(reinterpret_cast<SYMS_U8*>(File.GetData()), File.Num());
+	}
+}
+
+/////////////////////////////////////////////////////////////////////
+
 class FSymbolStringAllocator
 {
 public:
@@ -100,7 +239,6 @@ void FSymslibResolver::QueueModuleLoad(const FStringView& ModulePath, uint64 Bas
 	Entry->Name = Session.StoreString(ModuleName); //Allows safe sharing to UI
 	Entry->Path = Session.StoreString(ModulePath);
 	Entry->Status = EModuleStatus::Pending;
-	Entry->Instance = nullptr;
 	Entry->ImageId = TArrayView<const uint8>(ImageId, ImageIdSize);
 
 	++ModulesDiscovered;
@@ -149,27 +287,14 @@ void FSymslibResolver::OnAnalysisComplete()
 		}
 		while (OutstandingTasks > 0);
 
-		// No tasks can be in running at this point.
-		{
-			FReadScopeLock _(FileHandlesAndRegionsLock);
-			for (FMappedFileAndRegion& File : FileHandlesAndRegions)
-			{
-				// Release file region before releasing the file handle.
-				delete File.Region;
-				delete File.Handle;
-			}
-		}
-
 		// Release memory used by syms library
 		{
 			FReadScopeLock _(ModulesLock);
 			for (uint32 ModuleIndex = 0; ModuleIndex < Modules.Num(); ++ModuleIndex)
 			{
 				FModuleEntry& Module = Modules[ModuleIndex];
-				if (Module.Instance)
-				{
-					syms_quit(Module.Instance);
-				}
+				Algo::ForEach(Module.Instance.Arenas, syms_arena_release);
+				Module.Instance.Arenas.Empty();
 			}
 		}
 		
@@ -300,188 +425,162 @@ void FSymslibResolver::LoadModuleTracked(FModuleEntry* Module)
 	}
 }
 
-FString FSymslibResolver::FindSymbolFile(const FString& Path)
+FSymslibResolver::EModuleStatus FSymslibResolver::LoadModule(FModuleEntry* Module)
 {
 	IPlatformFile* PlatformFile = &FPlatformFileManager::Get().GetPlatformFile();
 
-	// Extract only filename part in case Path is absolute path
-	FString FileName = FPaths::GetCleanFilename(Path);
+	// temporary memory used for loading
+	SYMS_Group* Group = syms_group_alloc();
 
-	FString SymbolPath = FPlatformMisc::GetEnvironmentVariable(L"UE_INSIGHTS_SYMBOL_PATH");
+	SYMS_FileInfOptions opts = { 0 };
+	opts.disable_fallback = 1;
 
-	FString SymbolPathPart;
-	while (SymbolPath.Split(TEXT(";"), &SymbolPathPart, &SymbolPath))
-	{
-		FString Result = FPaths::Combine(SymbolPathPart, FileName);
-		if (PlatformFile->FileExists(*Result))
-		{
-			// File found in one of symbol paths
-			return Result;
-		}
-	}
+	TArray<FAutoMappedFile> Files;
+	SYMS_FileLoadCtx Context = { &SymsLoader, &Files };
+	SYMS_FileInfResult FileInf = syms_group_infer_from_file(Group->arena, Context, syms_str8_cstring(TCHAR_TO_ANSI(Module->Path)), &opts);
 
-	FString Result = FPaths::Combine(SymbolPath, FileName);
-	if (PlatformFile->FileExists(*Result))
+	SYMS_DbgAccel* Dbg = FileInf.data_parsed.dbg;
+	if (Dbg == nullptr || Dbg->format == SYMS_FileFormat_Null)
 	{
-		// File found in symbol path
-		return Result;
+		syms_group_release(Group);
+		UE_LOG(LogSymslib, Display, TEXT("No debug information loaded for '%s'"), Module->Name);
+		return EModuleStatus::Failed;
 	}
 
-	if (PlatformFile->FileExists(*Path))
+	// check debug data mismatch to captured module
+	if (!Module->ImageId.IsEmpty() && Dbg->format == SYMS_FileFormat_PDB)
 	{
-		// File found by absolute path
-		return Path;
+		// for Pdbs checksum is a 16 byte guid and 4 byte unsigned integer for age, but usually age is not used for matching debug file to exe
+		static_assert(sizeof(FGuid) == 16, "Expected 16 byte FGuid");
+		check(Module->ImageId.Num() == 20);
+		FGuid* ModuleGuid = (FGuid*)Module->ImageId.GetData();
+
+		SYMS_ExtMatchKey MatchKey = syms_ext_match_key_from_dbg(FileInf.data_parsed.dbg_data, Dbg);
+		FGuid* PdbGuid = (FGuid*)MatchKey.v;
+
+		if (*ModuleGuid != *PdbGuid)
+		{
+			syms_group_release(Group);
+			UE_LOG(LogSymslib, Warning, TEXT("Symbols for '%s' does not match traced binary."), Module->Name);
+			return EModuleStatus::VersionMismatch;
+		}
 	}
 
-	// File not found
-	return FString();
-}
+	FSymsInstance* Instance = &Module->Instance;
 
-FSymslibResolver::EModuleStatus FSymslibResolver::LoadModule(FModuleEntry* Module)
-{
-	IPlatformFile* PlatformFile = &FPlatformFileManager::Get().GetPlatformFile();
+	// initialize group
+	SYMS_GroupInitParams GroupInit = {};
+	GroupInit.dbg_data = FileInf.data_parsed.dbg_data;
+	GroupInit.dbg = FileInf.data_parsed.dbg;
 
-	// Initialize the instance
-	Module->Instance = syms_init_ex(Module->Base);
-	SymsInstance* ModuleInstance = Module->Instance;
+	syms_set_lane(0);
+	syms_group_init(Group, &GroupInit);
 
-	if (!ModuleInstance)
-	{
-		UE_LOG(LogSymslib, Error, TEXT("Failed to initialize '%s'"), Module->Name);
-		return EModuleStatus::Failed;
-	}
+	// unit storage
+	SYMS_U64 UnitCount = syms_group_unit_count(Group);
+	Instance->Units.SetNum(UnitCount);
 
-	FString SymbolFilePath = FindSymbolFile(Module->Path);
-	if (SymbolFilePath.IsEmpty())
-	{
-		UE_LOG(LogSymslib, Warning, TEXT("File '%s' does not exist"), Module->Path);
-		return EModuleStatus::Failed;
-	}
+	// per-thread arena storage (at least one)
+	int32 WorkerThreadCount = FMath::Max(1, FTaskGraphInterface::Get().GetNumWorkerThreads());
+	Instance->Arenas.SetNum(WorkerThreadCount);
+
+	// how many symbols are loaded
+	std::atomic<uint32> SymbolCount = 0;
 
-	// Map the image file to memory
-	IMappedFileHandle* ImageFileHandle = PlatformFile->OpenMapped(*SymbolFilePath);
-	if (ImageFileHandle == nullptr)
+	// parse debug info in multiple threads
 	{
-		UE_LOG(LogSymslib, Warning, TEXT("Failed to open '%s'"), *SymbolFilePath);
-		return EModuleStatus::Failed;
-	}
-	IMappedFileRegion* ImageFileRegion = ImageFileHandle->MapRegion(0, ImageFileHandle->GetFileSize());
-	check(ImageFileRegion != nullptr);
+		uint32 LaneSlot = FPlatformTLS::AllocTlsSlot();
+
+		std::atomic<uint32> LaneCount = 0;
+		syms_group_begin_multilane(Group, WorkerThreadCount);
+		ParallelFor(UnitCount, [Instance, Group, LaneSlot, &LaneCount, &SymbolCount](uint32 Index)
+		{
+			SYMS_Arena* Arena;
+			uint32 LaneValue = uint32(reinterpret_cast<intptr_t>(FPlatformTLS::GetTlsValue(LaneSlot)));
+			if (LaneValue == 0)
+			{
+				// first time we are on this thread
+				LaneValue = ++LaneCount;
+				FPlatformTLS::SetTlsValue(LaneSlot, reinterpret_cast<void*>(intptr_t(LaneValue)));
+
+				// syms lane index is 0-based
+				uint32 LaneIndex = LaneValue - 1;
+				syms_set_lane(LaneIndex);
+				Arena = Instance->Arenas[LaneIndex] = syms_arena_alloc();
+			}
+			else
+			{
+				uint32 LaneIndex = LaneValue - 1;
+				syms_set_lane(LaneIndex);
+				Arena = Instance->Arenas[LaneIndex];
+			}
 
-	// Track all open handles
-	TrackFileHandlesAndRegions(ImageFileHandle, ImageFileRegion);
+			SYMS_ArenaTemp Scratch = syms_get_scratch(0, 0);
 
-	// Load the image
-	SymsErrorCode Result = syms_load_image(ModuleInstance, (void*)ImageFileRegion->GetMappedPtr(), ImageFileRegion->GetMappedSize(), 0);
+			SYMS_UnitID UnitID = static_cast<SYMS_UnitID>(Index + 1); // syms unit id's are 1-based
+			FSymsUnit* Unit = &Instance->Units[Index];
 
-	if (SYMS_RESULT_FAIL(Result))
-	{
-		UE_LOG(LogSymslib, Warning, TEXT("Failed to load '%s'"), *SymbolFilePath);
-		return EModuleStatus::Failed;
-	}
+			SYMS_SpatialMap1D* ProcSpatialMap = syms_group_proc_map_from_uid(Group, UnitID);
+			Unit->ProcMap = syms_spatial_map_1d_copy(Arena, ProcSpatialMap);
 
-	// Iterate and map all the debug files
-	const FString ModuleBasePath = FPaths::GetPath(SymbolFilePath);
-	TArray<SymsFile> Files;
+			SYMS_String8Array* FileTable = syms_group_file_table_from_uid_with_fallbacks(Group, UnitID);
+			Unit->FileTable = syms_string_array_copy(Arena, 0, FileTable);
 
-	SymsDebugFileIter DebugFiles;
-	if (syms_debug_file_iter_init(&DebugFiles, ModuleInstance))
-	{
-		SymsString Path;
-		while (syms_debug_file_iter_next(&DebugFiles, &Path))
-		{
-			// Map the debug file
+			SYMS_LineParseOut* LineParse = syms_group_line_parse_from_uid(Group, UnitID);
+			Unit->LineTable = syms_line_table_with_indexes_from_parse(Arena, LineParse);
 
-			bool bDebugFilePathFound = false;
+			SYMS_SpatialMap1D* LineSpatialMap = syms_group_line_sequence_map_from_uid(Group, UnitID);
+			Unit->LineMap = syms_spatial_map_1d_copy(Arena, LineSpatialMap);
 
-			// First try path itself if it is absolute
-			FString DebugFilePath = ANSI_TO_TCHAR(Path.data);
-			if (!FPathViews::IsRelativePath(DebugFilePath))
-			{
-				if (PlatformFile->FileExists(*DebugFilePath))
-				{
-					bDebugFilePathFound = true;
-				}
-			}
+			SYMS_UnitAccel* UnitAccel = syms_group_unit_from_uid(Group, UnitID);
 
-			// Otherwise try relative path to symbol file path
-			if (!bDebugFilePathFound)
+			SYMS_IDMap ProcIdMap = syms_id_map_alloc(Scratch.arena, 4093);
+
+			SYMS_SymbolIDArray* ProcArray = syms_group_proc_sid_array_from_uid(Group, UnitID);
+			SYMS_U64 ProcCount = ProcArray->count;
+
+			uint32 UnitSymbolCount = 0;
+
+			FSymsSymbol* Symbols = syms_push_array(Arena, FSymsSymbol, ProcCount);
+			for (SYMS_U64 ProcIndex = 0; ProcIndex < ProcCount; ProcIndex++)
 			{
-				DebugFilePath = FPaths::Combine(ModuleBasePath, FPaths::GetCleanFilename(DebugFilePath));
-				if (PlatformFile->FileExists(*DebugFilePath))
+				SYMS_SymbolID SymbolID = ProcArray->ids[ProcIndex];
+
+				SYMS_SymbolKind Kind = syms_group_symbol_kind_from_sid(Group, UnitAccel, SymbolID);
+				if (Kind == SYMS_SymbolKind_Procedure)
 				{
-					bDebugFilePathFound = true;
+					UnitSymbolCount++;
 				}
-			}
 
-			// Use path if file was found
-			if (bDebugFilePathFound)
-			{
-				IMappedFileHandle* FileHandle = PlatformFile->OpenMapped(*DebugFilePath);
-				if (FileHandle != nullptr)
-				{
-					IMappedFileRegion* FileRegion = FileHandle->MapRegion(0, FileHandle->GetFileSize());
-					check(FileRegion != nullptr);
+				SYMS_String8 Name = syms_group_symbol_name_from_sid(Arena, Group, UnitAccel, SymbolID);
+				Symbols[ProcIndex].Name = reinterpret_cast<char*>(Name.str);
 
-					Files.Push(SymsFile{
-						Path.data,
-						(void*)FileRegion->GetMappedPtr(),
-						(uint64)FileRegion->GetMappedSize()
-					});
+				syms_id_map_insert(Scratch.arena, &ProcIdMap, SymbolID, &Symbols[ProcIndex]);
+			}
 
-					TrackFileHandlesAndRegions(FileHandle, FileRegion);
-				}
-				else
-				{
-					UE_LOG(LogSymslib, Warning, TEXT("Could not open the debug file '%s'"), *DebugFilePath);
-				}
+			SYMS_SpatialMap1D* ProcMap = &Unit->ProcMap;
+			for (SYMS_SpatialMap1DRange* Range = ProcMap->ranges, *EndRange = ProcMap->ranges + ProcMap->count; Range < EndRange; Range++)
+			{
+				void* SymbolPtr = syms_id_map_ptr_from_u64(&ProcIdMap, Range->val);
+				Range->val = SYMS_U64(reinterpret_cast<intptr_t>(SymbolPtr));
 			}
-		}
-	}
 
-	// Prepare to load the debug files
-	Result = syms_load_debug_info_ex(ModuleInstance, Files.GetData(), Files.Num(), SYMS_LOAD_DEBUG_INFO_FLAGS_DEFER_BUILD_MODULE);
+			syms_release_scratch(Scratch);
 
-	if (SYMS_RESULT_FAIL(Result))
-	{
-		UE_LOG(LogSymslib, Display, TEXT("No debug information for '%s'"), *SymbolFilePath);
-		return EModuleStatus::Failed;
-	}
+			SymbolCount += UnitSymbolCount;
+		});
+		syms_group_end_multilane(Group);
 
-	// Split up the work of actually loading debug info
-	const uint32 Count = syms_get_module_build_count(ModuleInstance);
-	// Each work item needs some memory
-	TArray<SymsArena*> Arenas;
-	for (uint32 Index = 0; Index < Count; ++Index)
-	{
-		Arenas.Add(syms_borrow_memory(ModuleInstance));
+		FPlatformTLS::FreeTlsSlot(LaneSlot);
 	}
-	ParallelFor(Count, [ModuleInstance, &Arenas](uint32 Index)
-	{
-		syms_build_module(ModuleInstance, (SymsModID)Index, Arenas[Index]);
-	});
 
-	// Check age of debug data
-	if (!Module->ImageId.IsEmpty() && ModuleInstance->img.type == SYMS_IMAGE_NT)
-	{
-		// For Pdbs checksum is a 16 byte guid and 8 byte unsigned integer for age
-		static_assert(sizeof(FGuid) == 16, "Expected 16 byte FGuid");
-		check(Module->ImageId.Num() == 20);
-		FGuid* ModuleGuid = (FGuid*) Module->ImageId.GetData();
-		uint32* ModuleAge = ((uint32*) Module->ImageId.GetData()) + 4;
-		SymsDebugInfoPdb* PdbInfo = (SymsDebugInfoPdb*)ModuleInstance->debug_info.impl;
-		FGuid* PdbGuid = (FGuid*)&PdbInfo->context.auth_guid;
-		const uint32 PdbAge = PdbInfo->context.age;
-		if (*ModuleGuid != *PdbGuid || *ModuleAge != PdbAge)
-		{
-			UE_LOG(LogSymslib, Warning, TEXT("Symbols for '%s' does not match traced binary."), Module->Name);
-			return EModuleStatus::VersionMismatch;
-		}
-	}
+	Instance->UnitMap = syms_spatial_map_1d_copy(Instance->Arenas[0], syms_group_unit_map(Group));
+	Instance->DefaultBase = syms_group_default_vbase(Group);
+
+	syms_group_release(Group);
+	Instance->Arenas.RemoveAll([](SYMS_Arena* arena) { return arena == nullptr; });
 
-	const uint32 SymbolCount = syms_get_proc_count(ModuleInstance); // get symbol count
-	const uint32 LineCount = syms_get_line_count(ModuleInstance); // get line count
-	UE_LOG(LogSymslib, Display, TEXT("Loaded symbols for '%s' at base 0x%016llx, %u symbols and %u lines."), Module->Name, Module->Base, SymbolCount, LineCount);
+	UE_LOG(LogSymslib, Display, TEXT("Loaded symbols for '%s' at base 0x%016llx, %u symbols."), Module->Name, Module->Base, SymbolCount.load());
 
 	return EModuleStatus::Loaded;
 }
@@ -527,29 +626,56 @@ bool FSymslibResolver::ResolveSymbol(uint64 Address, FResolvedSymbol* Target, FS
 		break;
 	}
 
-	SymsProc Proc;
-	SymsSourceFileMap File;
-
 	// Find procedure and source file for address
-	if (!syms_proc_from_va(Module->Instance, Address, &Proc) || !syms_va_to_src(Module->Instance, Address, &File))
+
+	FSymsInstance* Instance = &Module->Instance;
+	SYMS_U64 VirtualOffset = Address + Instance->DefaultBase - Module->Base;
+
+	FSymsSymbol* SymsSymbol = nullptr;
+	bool bFileValid = false;
+	SYMS_String8 FileName = {};
+	uint32 FileLine = 0;
+
+	SYMS_UnitID UnitID = syms_spatial_map_1d_value_from_point(&Instance->UnitMap, VirtualOffset);
+	if (UnitID > 0 && UnitID <= Instance->Units.Num())
 	{
-		UpdateResolvedSymbol(Target, ESymbolQueryResult::NotFound, Module->Name, GUnknownModuleTextSymsLib, GUnknownModuleTextSymsLib, 0); 
+		FSymsUnit* Unit = &Instance->Units[UnitID - 1];
+
+		SYMS_U64 Value = syms_spatial_map_1d_value_from_point(&Unit->ProcMap, VirtualOffset);
+		if (Value)
+		{
+			SymsSymbol = reinterpret_cast<FSymsSymbol*>(Value);
+
+			SYMS_U64 SeqNumber = syms_spatial_map_1d_value_from_point(&Unit->LineMap, VirtualOffset);
+			if (SeqNumber)
+			{
+				SYMS_Line Line = syms_line_from_sequence_voff(&Unit->LineTable, SeqNumber, VirtualOffset);
+				if (Line.src_coord.file_id > 0 && Line.src_coord.file_id <= Unit->FileTable.count)
+				{
+					FileName = Unit->FileTable.strings[Line.src_coord.file_id - 1];
+					FileLine = Line.src_coord.line;
+					bFileValid = true;
+				}
+			}
+		}
+	}
+
+	// this includes skipping symbols without name (empty string)
+	if (!SymsSymbol || !bFileValid || SymsSymbol && SymsSymbol->Name[0] == 0)
+	{
+		UpdateResolvedSymbol(Target, ESymbolQueryResult::NotFound, Module->Name, GUnknownModuleTextSymsLib, GUnknownModuleTextSymsLib, 0);
 		return false;
 	}
 
 	constexpr uint32 MaxStringSize = 1024;
 
-	// Read the symbol name
 	ANSICHAR SymbolName[MaxStringSize];
+	FCStringAnsi::Strncpy(SymbolName, SymsSymbol->Name, MaxStringSize);
 	SymbolName[MaxStringSize - 1] = 0;
-	syms_read_strref(Module->Instance, &Proc.name_ref, SymbolName, sizeof(SymbolName) - 1);
-	check(SymbolName[MaxStringSize - 1] == 0);
 
-	// Read the source location
 	ANSICHAR SourceFile[MaxStringSize];
+	FCStringAnsi::Strncpy(SourceFile, reinterpret_cast<char*>(FileName.str), MaxStringSize);
 	SourceFile[MaxStringSize - 1] = 0;
-	syms_read_strref(Module->Instance, &File.file.name, SourceFile, sizeof(SourceFile) - 1);
-	check(SourceFile[MaxStringSize - 1] == 0);
 
 	const TCHAR* SymbolNamePersistent =  StringAllocator.Store(ANSI_TO_TCHAR(SymbolName));
 	const TCHAR* SourceFilePersistent = StringAllocator.Store(ANSI_TO_TCHAR(SourceFile));
@@ -561,18 +687,12 @@ bool FSymslibResolver::ResolveSymbol(uint64 Address, FResolvedSymbol* Target, FS
 		Module->Name,
 		SymbolNamePersistent,
 		SourceFilePersistent,
-		File.line.ln
+		FileLine
 	);
 
 	return true;
 }
 
-void FSymslibResolver::TrackFileHandlesAndRegions(IMappedFileHandle* FileHandle, IMappedFileRegion* FileRegion)
-{
-	FWriteScopeLock _(FileHandlesAndRegionsLock);
-	FileHandlesAndRegions.Add(FMappedFileAndRegion{ FileHandle, FileRegion });
-}
-
 /////////////////////////////////////////////////////////////////////
 
 } // namespace TraceServices
diff --git a/Engine/Source/Developer/TraceServices/Private/Model/SymslibResolver.h b/Engine/Source/Developer/TraceServices/Private/Model/SymslibResolver.h
--- a/Engine/Source/Developer/TraceServices/Private/Model/SymslibResolver.h
+++ b/Engine/Source/Developer/TraceServices/Private/Model/SymslibResolver.h
@@ -36,13 +36,34 @@ private:
 		Failed
 	};
 
+	struct FSymsSymbol
+	{
+		const char* Name;
+	};
+
+	struct FSymsUnit
+	{
+		SYMS_SpatialMap1D ProcMap;
+		SYMS_String8Array FileTable;
+		SYMS_LineTable    LineTable;
+		SYMS_SpatialMap1D LineMap;
+	};
+
+	struct FSymsInstance
+	{
+		TArray<SYMS_Arena*> Arenas;
+		TArray<FSymsUnit>   Units;
+		SYMS_SpatialMap1D   UnitMap;
+		uint64              DefaultBase;
+	};
+
 	struct FModuleEntry
 	{
 		uint64 Base;
 		uint32 Size;
 		const TCHAR* Name;
 		const TCHAR* Path;
-		SymsInstance* Instance;
+		FSymsInstance Instance;
 		std::atomic<EModuleStatus> Status;
 		TArray<uint8> ImageId;
 	};
@@ -79,11 +100,9 @@ private:
 	FModuleEntry* GetModuleForAddress(uint64 Address);
 	static void UpdateResolvedSymbol(FResolvedSymbol* Symbol, ESymbolQueryResult Result, const TCHAR* Module, const TCHAR* Name, const TCHAR* File, uint16 Line);
 	void LoadModuleTracked(FModuleEntry* Module);
-	FString FindSymbolFile(const FString& Path);
 	EModuleStatus LoadModule(FModuleEntry* Module);
 	void ResolveSymbolTracked(uint64 Address, FResolvedSymbol* Target, class FSymbolStringAllocator& StringAllocator);
 	bool ResolveSymbol(uint64 Address, FResolvedSymbol* Target, class FSymbolStringAllocator& StringAllocator);
-	void TrackFileHandlesAndRegions(IMappedFileHandle* FileHandle, IMappedFileRegion* FileRegion);
 
 	FRWLock ModulesLock;
 	TPagedArray<FModuleEntry> Modules;
@@ -93,14 +112,6 @@ private:
 	std::atomic<uint32> TasksInFlight;
 	FGraphEventRef CleanupTask;
 	std::atomic<bool> CancelTasks;
-	
-	struct FMappedFileAndRegion
-	{
-		IMappedFileHandle* Handle;
-		IMappedFileRegion* Region;
-	};
-	TArray<FMappedFileAndRegion> FileHandlesAndRegions;
-	FRWLock FileHandlesAndRegionsLock;
 
 	std::atomic<uint32> ModulesDiscovered;
 	std::atomic<uint32> ModulesFailed;
diff --git a/Engine/Source/Runtime/SymsLib/SymsLib.Build.cs b/Engine/Source/Runtime/SymsLib/SymsLib.Build.cs
--- a/Engine/Source/Runtime/SymsLib/SymsLib.Build.cs
+++ b/Engine/Source/Runtime/SymsLib/SymsLib.Build.cs
@@ -9,9 +9,9 @@ public class SymsLib : ModuleRules
 		Type = ModuleType.External;
 		
 		string LibPathBase =  ModuleDirectory + "/lib";
-		string IncludePath = ModuleDirectory + "/include";
 
-		PublicSystemIncludePaths.Add(IncludePath);
+		PublicSystemIncludePaths.Add(ModuleDirectory);
+		PublicSystemIncludePaths.Add(ModuleDirectory + "/syms");
 
 		if (Target.Platform.IsInGroup(UnrealPlatformGroup.Windows))
 		{
diff --git a/Engine/Source/Runtime/SymsLib/syms/core/base/syms_base_overrides_check.h b/Engine/Source/Runtime/SymsLib/syms/core/base/syms_base_overrides_check.h
new file mode 100644
--- /dev/null
+++ b/Engine/Source/Runtime/SymsLib/syms/core/base/syms_base_overrides_check.h
@@ -0,0 +1,130 @@
+// Copyright Epic Games, Inc. All Rights Reserved.
+/* date = July 27th 2021 10:54 am */
+
+#ifndef SYMS_BASE_OVERRIDES_CHECK_H
+#define SYMS_BASE_OVERRIDES_CHECK_H
+
+// This will check every possible overridable that could be provided by the
+// user. All overrideables are required at this layer.
+
+
+// NOTE(allen): basic types
+
+#if !defined(SYMS_S8)
+# error SYMS: SYMS_S8 was not #define'd. It should be #define'd to the name of a signed 8-bit integer type (int8_t in the C Standard Library).
+#endif
+
+#if !defined(SYMS_S16)
+# error SYMS: SYMS_S16 was not #define'd. It should be #define'd to the name of a signed 16-bit integer type (int16_t in the C Standard Library).
+#endif
+
+#if !defined(SYMS_S32)
+# error SYMS: SYMS_S32 was not #define'd. It should be #define'd to the name of a signed 32-bit integer type (int32_t in the C Standard Library).
+#endif
+
+#if !defined(SYMS_S64)
+# error SYMS: SYMS_S64 was not #define'd. It should be #define'd to the name of a signed 64-bit integer type (int64_t in the C Standard Library).
+#endif
+
+#if !defined(SYMS_U8)
+# error SYMS: SYMS_U8 was not #define'd. It should be #define'd to the name of an unsigned 8-bit integer type (uint8_t in the C Standard Library).
+#endif
+
+#if !defined(SYMS_U16)
+# error SYMS: SYMS_U16 was not #define'd. It should be #define'd to the name of an unsigned 16-bit integer type (uint16_t in the C Standard Library).
+#endif
+
+#if !defined(SYMS_U32)
+# error SYMS: SYMS_U32 was not #define'd. It should be #define'd to the name of an unsigned 32-bit integer type (uint32_t in the C Standard Library).
+#endif
+
+#if !defined(SYMS_U64)
+# error SYMS: SYMS_U64 was not #define'd. It should be #define'd to the name of an unsigned 64-bit integer type (uint64_t in the C Standard Library).
+#endif
+
+#if !defined(SYMS_PRIu64)
+# error SYMS: SYMS_PRIu64 was not #define'd. It should be #define'd to a string constant suitable for printing 64-bit unsigned integers in C format strings (PRIu64 in the C Standard Library).
+#endif
+
+#if !defined(SYMS_PRId64)
+# error SYMS: SYMS_PRId64 was not #define'd. It should be #define'd to a string constant suitable for printing 64-bit signed integers in C format strings (PRId64 in the C Standard Library).
+#endif
+
+#if !defined(SYMS_PRIx64)
+# error SYMS: SYMS_PRIx64 was not #define'd. It should be #define'd to a string constant suitable for printing 64-bit hex integers in C format strings (PRIx64 in the C Standard Library).
+#endif
+
+
+// NOTE(allen): large data operations
+
+#if !defined(syms_memmove)
+# error SYMS: syms_memmove was not #define'd. It should be #define'd to the name of a memmove implementation.
+#endif
+
+#if !defined(syms_memset)
+# error SYMS: syms_memset was not #define'd. It should be #define'd to the name of a memset implementation.
+#endif
+
+#if !defined(syms_memcmp)
+# error SYMS: syms_memcmp was not #define'd. It should be #define'd to the name of a memcmp implementation.
+#endif
+
+#if !defined(syms_strlen)
+# error SYMS: syms_strlen was not #define'd. It should be #define'd to the name of a strlen implementation.
+#endif
+
+
+// NOTE(allen): assert
+
+#if !defined(SYMS_ASSERT_BREAK)
+# error SYMS: SYMS_ASSERT_BREAK(message) was not #define'd. It should be #define'd to something that you'd like expanded when an assertion failure occurs.
+#endif
+
+
+// NOTE(allen): arena
+
+#if !defined(SYMS_Arena)
+# error SYMS: SYMS_Arena was not #define'd.
+#endif
+
+#if !defined(syms_arena_alloc__impl)
+# error SYMS: syms_arena_alloc__impl was not #define'd.
+#endif
+
+#if !defined(syms_arena_release__impl)
+# error SYMS: syms_arena_release__impl was not #define'd.
+#endif
+
+#if !defined(syms_arena_get_pos__impl)
+# error SYMS: syms_arena_get_pos__impl was not #define'd.
+#endif
+
+#if !defined(syms_arena_push__impl)
+# error SYMS: syms_arena_push__impl was not #define'd.
+#endif
+
+#if !defined(syms_arena_pop_to__impl)
+# error SYMS: syms_arena_pop_to__impl was not #define'd.
+#endif
+
+#if !defined(syms_arena_set_auto_align__impl)
+# error SYMS: syms_arena_set_auto_align__impl was not #define'd.
+#endif
+
+#if !defined(syms_arena_absorb__impl)
+# error SYMS: syms_arena_absorb__impl was not #define'd.
+#endif
+
+#if !defined(syms_arena_tidy__impl)
+# error SYMS: syms_arena_tidy__impl was not #define'd.
+#endif
+
+#if !defined(syms_get_implicit_thread_arena__impl)
+# error SYMS: syms_get_implicit_thread_arena__impl was not #define'd.
+#endif
+
+#if !defined(syms_scratch_pool_tidy__impl)
+# error SYMS: syms_scratch_pool_tidy__impl was not #define'd.
+#endif
+
+#endif //SYMS_BASE_OVERRIDES_CHECK_H
diff --git a/Engine/Source/Runtime/SymsLib/syms/core/elf/syms_elf.h b/Engine/Source/Runtime/SymsLib/syms/core/elf/syms_elf.h
new file mode 100644
--- /dev/null
+++ b/Engine/Source/Runtime/SymsLib/syms/core/elf/syms_elf.h
@@ -0,0 +1,372 @@
+// Copyright Epic Games, Inc. All Rights Reserved.
+/* date = May 11th 2021 4:54 pm */
+
+#ifndef SYMS_ELF_H
+#define SYMS_ELF_H
+
+//- rjf: p_type
+typedef enum SYMS_ElfPKind{
+  SYMS_ElfPKind_Null      = 0u,
+  SYMS_ElfPKind_Load      = 1u,
+  SYMS_ElfPKind_Dynamic   = 2u,
+  SYMS_ElfPKind_Interp    = 3u,
+  SYMS_ElfPKind_Note      = 4u,
+  SYMS_ElfPKind_SHLib     = 5u,
+  SYMS_ElfPKind_PHDR      = 6u,
+  SYMS_ElfPKind_LowProc   = 0x70000000u,
+  SYMS_ElfPKind_HighProc  = 0x7fffffffu,
+  
+  //- specific to Sun
+  SYMS_ElfPKind_LowSunW   = 0x6ffffffau,
+  SYMS_ElfPKind_SunWBSS   = 0x6ffffffbu,
+  SYMS_ElfPKind_SunWStack = 0x6ffffffau,
+  SYMS_ElfPKind_HighSunW  = 0x6fffffffu,
+}SYMS_ElfPKind;
+
+//- rjf: e_machine
+typedef enum SYMS_ElfMachineKind{
+  SYMS_ElfMachineKind_NONE           = 0u, 
+  SYMS_ElfMachineKind_M32            = 1u, 
+  SYMS_ElfMachineKind_SPARC          = 2u, 
+  SYMS_ElfMachineKind_386            = 3u, 
+  SYMS_ElfMachineKind_68K            = 4u, 
+  SYMS_ElfMachineKind_88K            = 5u, 
+  SYMS_ElfMachineKind_IAMCU          = 6u, 
+  SYMS_ElfMachineKind_860            = 7u, 
+  SYMS_ElfMachineKind_MIPS           = 8u, 
+  SYMS_ElfMachineKind_S370           = 9u, 
+  SYMS_ElfMachineKind_MIPS_RS3_LE    = 10u, 
+  SYMS_ElfMachineKind_PARISC         = 15u, 
+  SYMS_ElfMachineKind_PPC_OLD        = 17u, // nick: Old version of PowerPC. Deprecated
+  SYMS_ElfMachineKind_SPARC32PLUS    = 18u, // nick: Sun's "v8plus"
+  SYMS_ElfMachineKind_960            = 19u, 
+  SYMS_ElfMachineKind_PPC            = 20u, 
+  SYMS_ElfMachineKind_PPC64          = 21u, 
+  SYMS_ElfMachineKind_S390           = 22u, 
+  SYMS_ElfMachineKind_SPU            = 23u, 
+  SYMS_ElfMachineKind_V800           = 36u, 
+  SYMS_ElfMachineKind_FR20           = 37u, 
+  SYMS_ElfMachineKind_RH32           = 38u, 
+  SYMS_ElfMachineKind_MCORE          = 39u, 
+  SYMS_ElfMachineKind_ARM            = 40u, 
+  SYMS_ElfMachineKind_SH             = 42u, 
+  SYMS_ElfMachineKind_IA_64          = 50u, 
+  SYMS_ElfMachineKind_X86_64         = 62u, 
+  SYMS_ElfMachineKind_AARCH64        = 183u, 
+  SYMS_ElfMachineKind_RISCV          = 243u, 
+}SYMS_ElfMachineKind;
+
+typedef enum SYMS_ElfClass{
+  SYMS_ElfClass_None = 0,
+  SYMS_ElfClass_32   = 1,
+  SYMS_ElfClass_64   = 2,
+}SYMS_ElfClass;
+
+typedef enum SYMS_ElfIdentifier{
+  SYMS_ElfIdentifier_MAG0    = 0,
+  SYMS_ElfIdentifier_MAG1    = 1,
+  SYMS_ElfIdentifier_MAG2    = 2,
+  SYMS_ElfIdentifier_MAG3    = 3,
+  SYMS_ElfIdentifier_CLASS   = 4,
+  SYMS_ElfIdentifier_DATA    = 5,
+  SYMS_ElfIdentifier_VERSION = 6,
+  SYMS_ElfIdentifier_PAD     = 7,
+  SYMS_ElfIdentifier_NINDENT = 16,
+}SYMS_ElfIdentifier;
+
+typedef struct SYMS_ElfEhdr64 SYMS_ElfEhdr64;
+struct SYMS_ElfEhdr64
+{
+  SYMS_U8 e_ident[SYMS_ElfIdentifier_NINDENT];
+  SYMS_U16 e_type;
+  SYMS_U16 e_machine;
+  SYMS_U32 e_version;
+  SYMS_U64 e_entry;
+  SYMS_U64 e_phoff;
+  SYMS_U64 e_shoff;
+  SYMS_U32 e_flags;
+  SYMS_U16 e_ehsize;
+  SYMS_U16 e_phentsize;
+  SYMS_U16 e_phnum;
+  SYMS_U16 e_shentsize;
+  SYMS_U16 e_shnum;
+  SYMS_U16 e_shstrndx;
+};
+
+typedef struct SYMS_ElfEhdr32 SYMS_ElfEhdr32;
+struct SYMS_ElfEhdr32
+{
+  SYMS_U8 e_ident[SYMS_ElfIdentifier_NINDENT];
+  SYMS_U16 e_type;
+  SYMS_U16 e_machine;
+  SYMS_U32 e_version;
+  SYMS_U32 e_entry;
+  SYMS_U32 e_phoff;
+  SYMS_U32 e_shoff;
+  SYMS_U32 e_flags;
+  SYMS_U16 e_ehsize;
+  SYMS_U16 e_phentsize;
+  SYMS_U16 e_phnum;
+  SYMS_U16 e_shentsize;
+  SYMS_U16 e_shnum;
+  SYMS_U16 e_shstrndx;
+};
+
+typedef struct SYMS_ElfShdr64 SYMS_ElfShdr64;
+struct SYMS_ElfShdr64
+{
+  SYMS_U32 sh_name;
+  SYMS_U32 sh_type;
+  SYMS_U64 sh_flags;
+  SYMS_U64 sh_addr;
+  SYMS_U64 sh_offset;
+  SYMS_U64 sh_size;
+  SYMS_U32 sh_link;
+  SYMS_U32 sh_info;
+  SYMS_U64 sh_addralign;
+  SYMS_U64 sh_entsize;
+};
+
+typedef struct SYMS_ElfShdr32 SYMS_ElfShdr32;
+struct SYMS_ElfShdr32
+{
+  SYMS_U32 sh_name;
+  SYMS_U32 sh_type;
+  SYMS_U32 sh_flags;
+  SYMS_U32 sh_addr;
+  SYMS_U32 sh_offset;
+  SYMS_U32 sh_size;
+  SYMS_U32 sh_link;
+  SYMS_U32 sh_info;
+  SYMS_U32 sh_addralign;
+  SYMS_U32 sh_entsize;
+};
+
+typedef struct SYMS_ElfPhdr64 SYMS_ElfPhdr64;
+struct SYMS_ElfPhdr64
+{
+  SYMS_U32 p_type;
+  SYMS_U32 p_flags;
+  SYMS_U64 p_offset;
+  SYMS_U64 p_vaddr;
+  SYMS_U64 p_paddr;
+  SYMS_U64 p_filesz;
+  SYMS_U64 p_memsz;
+  SYMS_U64 p_align;
+};
+
+typedef struct SYMS_ElfPhdr32 SYMS_ElfPhdr32;
+struct SYMS_ElfPhdr32
+{
+  SYMS_U32 p_type;
+  SYMS_U32 p_offset;
+  SYMS_U32 p_vaddr;
+  SYMS_U32 p_paddr;
+  SYMS_U32 p_filesz;
+  SYMS_U32 p_memsz;
+  SYMS_U32 p_flags;
+  SYMS_U32 p_align;
+};
+
+typedef enum SYMS_ElfSectionCode{
+  SYMS_ElfSectionCode_NULL                   = 0,
+  SYMS_ElfSectionCode_PROGBITS               = 1,
+  SYMS_ElfSectionCode_SYMTAB                 = 2,
+  SYMS_ElfSectionCode_STRTAB                 = 3,
+  SYMS_ElfSectionCode_RELA                   = 4,
+  SYMS_ElfSectionCode_HASH                   = 5,
+  SYMS_ElfSectionCode_DYNAMIC                = 6,
+  SYMS_ElfSectionCode_NOTE                   = 7,
+  SYMS_ElfSectionCode_NOBITS                 = 8,
+  SYMS_ElfSectionCode_REL                    = 9,
+  SYMS_ElfSectionCode_SHLIB                  = 10,
+  SYMS_ElfSectionCode_DYNSYM                 = 11,
+  SYMS_ElfSectionCode_INIT_ARRAY	           = 14,                // nick: Array of ptrs to init functions
+  SYMS_ElfSectionCode_FINI_ARRAY	           = 15,                // nick: Array of ptrs to finish functions
+  SYMS_ElfSectionCode_PREINIT_ARRAY          = 16,                // nick: Array of ptrs to pre-init funcs
+  SYMS_ElfSectionCode_GROUP	                = 17,                // nick: Section contains a section group
+  SYMS_ElfSectionCode_SYMTAB_SHNDX           = 18,                // nick: Indices for SHN_XINDEX entries
+  SYMS_ElfSectionCode_GNU_INCREMENTAL_INPUTS = 0x6fff4700,        // nick: incremental build data
+  SYMS_ElfSectionCode_GNU_ATTRIBUTES         = 0x6ffffff5,        // nick: Object attributes
+  SYMS_ElfSectionCode_GNU_HASH               = 0x6ffffff6,        // nick: GNU style symbol hash table
+  SYMS_ElfSectionCode_GNU_LIBLIST	          = 0x6ffffff7,        // nick: List of prelink dependencies
+  
+  // NOTE(nick): The next three section types are defined by Solaris, and are
+  // named SHT_SUNW*.  We use them in GNU code, so we also define SHT_GNU*
+  // versions.
+  SYMS_ElfSectionCode_SUNW_verdef	   = 0x6ffffffd,	// nick: Versions defined by file
+  SYMS_ElfSectionCode_SUNW_verneed    = 0x6ffffffe,	// nick: Versions needed by file
+  SYMS_ElfSectionCode_SUNW_versym	   = 0x6fffffff,	// nick: Symbol versions
+  
+  SYMS_ElfSectionCode_GNU_verdef	    = SYMS_ElfSectionCode_SUNW_verdef,
+  SYMS_ElfSectionCode_GNU_verneed	   = SYMS_ElfSectionCode_SUNW_verneed,
+  SYMS_ElfSectionCode_GNU_versym	    = SYMS_ElfSectionCode_SUNW_versym,
+  
+  SYMS_ElfSectionCode_LOPROC          = 0x70000000,
+  SYMS_ElfSectionCode_HIPROC          = 0x7fffffff,
+  SYMS_ElfSectionCode_LOUSER          = 0x80000000,
+  SYMS_ElfSectionCode_HIUSER          = 0xffffffff,
+}SYMS_ElfSectionCode;
+
+enum{
+  SYMS_ElfSectionFlag_WRITE            = 0x1,
+  SYMS_ElfSectionFlag_ALLOC            = 0x2,
+  SYMS_ElfSectionFlag_EXECINSTR        = 0x4,
+  SYMS_ElfSectionFlag_MERGE            = 0x10,
+  SYMS_ElfSectionFlag_STRINGS          = 0x20,
+  SYMS_ElfSectionFlag_INFO_LINK        = 0x40,
+  SYMS_ElfSectionFlag_LINK_ORDER       = 0x80,
+  SYMS_ElfSectionFlag_OS_NONCONFORMING = 0x100,
+  SYMS_ElfSectionFlag_GROUP            = 0x200,
+  SYMS_ElfSectionFlag_TLS              = 0x400,
+  SYMS_ElfSectionFlag_MASKOS           = 0xff0000,
+  SYMS_ElfSectionFlag_AMD64_LARGE      = 0x10000000,
+  SYMS_ElfSectionFlag_ORDERED          = 0x40000000,
+  SYMS_ElfSectionFlagT_EXCLUDE         = 0x80000000,
+  SYMS_ElfSectionFlag_MASKPROC         = 0xf0000000
+};
+
+////////////////////////////////
+//~ NOTE(allen): Auxiliary Vectors
+
+// these appear in /proc/<pid>/auxv of a process, they are not in elf files
+
+typedef enum SYMS_ElfAuxType{
+  SYMS_ElfAuxType_NULL = 0,
+  SYMS_ElfAuxType_PHDR = 3,   // program headers
+  SYMS_ElfAuxType_PHENT = 4,  // size of a program header
+  SYMS_ElfAuxType_PHNUM = 5,  // number of program headers
+  SYMS_ElfAuxType_PAGESZ = 6, // system page size
+  SYMS_ElfAuxType_BASE = 7,   // interpreter base address
+  SYMS_ElfAuxType_FLAGS = 8,
+  SYMS_ElfAuxType_ENTRY = 9,  // program entry point
+  SYMS_ElfAuxType_UID = 11,
+  SYMS_ElfAuxType_EUID = 12,
+  SYMS_ElfAuxType_GID = 13,
+  SYMS_ElfAuxType_EGID = 14,
+  SYMS_ElfAuxType_PLATFORM = 15, // 'platform' as a string (TODO(allen): study)
+  SYMS_ElfAuxType_HWCAP = 16,
+  SYMS_ElfAuxType_CLKTCK = 17,
+  SYMS_ElfAuxType_DCACHEBSIZE = 19,
+  SYMS_ElfAuxType_ICACHEBSIZE = 20,
+  SYMS_ElfAuxType_UCACHEBSIZE = 21,
+  SYMS_ElfAuxType_IGNOREPPC = 22,
+  SYMS_ElfAuxType_SECURE = 23,
+  SYMS_ElfAuxType_BASE_PLATFORM = 24, // 'platform' as a string (different) (TODO(allen): study)
+  SYMS_ElfAuxType_RANDOM = 25, // addres to 16 random bytes
+  SYMS_ElfAuxType_HWCAP2 = 26,
+  SYMS_ElfAuxType_EXECFN = 31, // file name of executable
+  
+  SYMS_ElfAuxType_SYSINFO = 32,
+  SYMS_ElfAuxType_SYSINFO_EHDR = 33,
+  
+  // cool info about caches? (TODO(allen): study)
+  SYMS_ElfAuxType_L1I_CACHESIZE = 40,
+  SYMS_ElfAuxType_L1I_CACHEGEOMETRY = 41,
+  SYMS_ElfAuxType_L1D_CACHESIZE = 42,
+  SYMS_ElfAuxType_L1D_CACHEGEOMETRY = 43,
+  SYMS_ElfAuxType_L2_CACHESIZE = 44,
+  SYMS_ElfAuxType_L2_CACHEGEOMETRY = 45,
+  SYMS_ElfAuxType_L3_CACHESIZE = 46,
+  SYMS_ElfAuxType_L3_CACHEGEOMETRY = 47,
+} SYMS_ElfAuxType;
+
+typedef struct SYMS_ElfAuxv32{
+  SYMS_U32 a_type;
+  SYMS_U32 a_val;
+} SYMS_ElfAuxv32;
+
+typedef struct SYMS_ElfAuxv64{
+  SYMS_U64 a_type;
+  SYMS_U64 a_val;
+} SYMS_ElfAuxv64;
+
+////////////////////////////////
+//~ NOTE(allen): Dynamic Structures
+
+// these appear in the virtual address space of a process, they are not in elf files
+
+typedef enum SYMS_ElfDynTag{
+  SYMS_ElfDynTag_NULL = 0,
+  SYMS_ElfDynTag_NEEDED = 1,
+  SYMS_ElfDynTag_PLTRELSZ = 2,
+  SYMS_ElfDynTag_PLTGOT = 3,
+  SYMS_ElfDynTag_HASH = 4,
+  SYMS_ElfDynTag_STRTAB = 5,
+  SYMS_ElfDynTag_SYMTAB = 6,
+  SYMS_ElfDynTag_RELA = 7,
+  SYMS_ElfDynTag_RELASZ = 8,
+  SYMS_ElfDynTag_RELAENT = 9,
+  SYMS_ElfDynTag_STRSZ = 10,
+  SYMS_ElfDynTag_SYMENT = 11,
+  SYMS_ElfDynTag_INIT = 12,
+  SYMS_ElfDynTag_FINI = 13,
+  SYMS_ElfDynTag_SONAME = 14,
+  SYMS_ElfDynTag_RPATH = 15,
+  SYMS_ElfDynTag_SYMBOLIC = 16,
+  SYMS_ElfDynTag_REL = 17,
+  SYMS_ElfDynTag_RELSZ = 18,
+  SYMS_ElfDynTag_RELENT = 19,
+  SYMS_ElfDynTag_PLTREL = 20,
+  SYMS_ElfDynTag_DEBUG = 21,
+  SYMS_ElfDynTag_TEXTREL = 22,
+  SYMS_ElfDynTag_JMPREL = 23,
+  SYMS_ElfDynTag_BIND_NOW = 24,
+  SYMS_ElfDynTag_INIT_ARRAY = 25,
+  SYMS_ElfDynTag_FINI_ARRAY = 26,
+  SYMS_ElfDynTag_INIT_ARRAYSZ = 27,
+  SYMS_ElfDynTag_FINI_ARRAYSZ = 28,
+  SYMS_ElfDynTag_RUNPATH = 29,
+  SYMS_ElfDynTag_FLAGS = 30,
+  SYMS_ElfDynTag_PREINIT_ARRAY = 32,
+  SYMS_ElfDynTag_PREINIT_ARRAYSZ = 33,
+  SYMS_ElfDynTag_SYMTAB_SHNDX = 34,
+  SYMS_ElfDynTag_LOOS = 0x6000000D,
+  SYMS_ElfDynTag_HIOS = 0x6ffff000,
+  SYMS_ElfDynTag_LOPROC = 0x70000000,
+  SYMS_ElfDynTag_HIPROC = 0x7fffffff,
+} SYMS_ElfDynTag;
+
+typedef struct SYMS_ElfDyn32{
+  SYMS_U32 tag;
+  SYMS_U32 val;
+} SYMS_ElfDyn32;
+
+typedef struct SYMS_ElfDyn64{
+  SYMS_U64 tag;
+  SYMS_U64 val;
+} SYMS_ElfDyn64;
+
+typedef struct SYMS_ElfLinkMap32{
+  SYMS_U32 base;
+  SYMS_U32 name;
+  SYMS_U32 ld;
+  SYMS_U32 next;
+} SYMS_ElfLinkMap32;
+
+typedef struct SYMS_ElfLinkMap64{
+  SYMS_U64 base;
+  SYMS_U64 name;
+  SYMS_U64 ld;
+  SYMS_U64 next;
+} SYMS_ElfLinkMap64;
+
+////////////////////////////////
+//~ rjf: ELF Format Functions
+
+SYMS_C_LINKAGE_BEGIN
+
+// TODO(allen): avoid extra copies (use pointers where convenient)
+SYMS_API SYMS_ElfEhdr64 syms_elf_ehdr64_from_ehdr32(SYMS_ElfEhdr32 h32);
+SYMS_API SYMS_ElfShdr64 syms_elf_shdr64_from_shdr32(SYMS_ElfShdr32 h32);
+SYMS_API SYMS_ElfPhdr64 syms_elf_phdr64_from_phdr32(SYMS_ElfPhdr32 h32);
+// TODO(allen): auxv?
+// TODO(allen): dyn?
+// TODO(allen): linkmap?
+
+// TODO(allen): promote to base
+SYMS_API SYMS_U32 syms_elf_gnu_debuglink_crc32(SYMS_U32 crc, SYMS_String8 data);
+
+SYMS_C_LINKAGE_END
+
+#endif // SYMS_ELF_H
