Check the build-id of elf file when resolving symbols in Insights

diff --git a/Engine/Source/Developer/TraceServices/Private/Model/SymslibResolver.cpp b/Engine/Source/Developer/TraceServices/Private/Model/SymslibResolver.cpp
--- a/Engine/Source/Developer/TraceServices/Private/Model/SymslibResolver.cpp
+++ b/Engine/Source/Developer/TraceServices/Private/Model/SymslibResolver.cpp
@@ -167,6 +167,86 @@ namespace
 
 		return syms_str8(reinterpret_cast<SYMS_U8*>(File.GetData()), File.Num());
 	}
+
+	static bool MatchImageId(const TArray<uint8>& ImageId, SYMS_ParseBundle DataParsed)
+	{
+		if (DataParsed.dbg->format == SYMS_FileFormat_PDB)
+		{
+			// for Pdbs checksum is a 16 byte guid and 4 byte unsigned integer for age, but usually age is not used for matching debug file to exe
+			static_assert(sizeof(FGuid) == 16, "Expected 16 byte FGuid");
+			check(ImageId.Num() == 20);
+			FGuid* ModuleGuid = (FGuid*)ImageId.GetData();
+
+			SYMS_ExtMatchKey MatchKey = syms_ext_match_key_from_dbg(DataParsed.dbg_data, DataParsed.dbg);
+			FGuid* PdbGuid = (FGuid*)MatchKey.v;
+
+			if (*ModuleGuid != *PdbGuid)
+			{
+				// mismatch
+				return false;
+			}
+		}
+		else if (DataParsed.bin->format == SYMS_FileFormat_ELF)
+		{
+			// try different ways of geting build id from elf binary
+			SYMS_String8 FoundId = { 0, 0 };
+
+			SYMS_String8 Bin = DataParsed.bin_data;
+			SYMS_ElfSectionArray Sections = DataParsed.bin->elf_accel.sections;
+			for (SYMS_U64 SectionIndex = 0; SectionIndex < Sections.count; SectionIndex += 1)
+			{
+				SYMS_U64 SectionOffset = Sections.v[SectionIndex].file_range.min;
+				SYMS_U64 SectionSize = Sections.v[SectionIndex].file_range.max - SectionOffset;
+
+				if (syms_string_match(Sections.v[SectionIndex].name, syms_str8_lit(".note.gnu.build-id"), 0))
+				{
+					if (SectionSize > 12)
+					{
+						SYMS_U32 NameSize = *(SYMS_U32*)&Bin.str[SectionOffset + 0];
+						SYMS_U32 DescSize = *(SYMS_U32*)&Bin.str[SectionOffset + 4];
+						SYMS_U32 Type = *(SYMS_U32*)&Bin.str[SectionOffset + 8];
+
+						const SYMS_U32 NT_GNU_BUILD_ID = 3;
+						// name must be "GNU\0", contents must be at least 16 bytes, and type must be 3
+						if (NameSize == 4 && DescSize >= 16 && Type == NT_GNU_BUILD_ID)
+						{
+							SYMS_U64 NameOffset = sizeof(NameSize) + sizeof(DescSize) + sizeof(Type);
+							SYMS_String8 NameStr = syms_str8(&Bin.str[SectionOffset + NameOffset], 4);
+							if (NameSize <= SectionSize && syms_string_match(NameStr, syms_str8((SYMS_U8*)"GNU", 4), 0))
+							{
+								SYMS_U64 DescOffset = NameOffset + SYMS_AlignPow2(NameSize, 4);
+								if (DescOffset + 16 <= SectionSize)
+								{
+									FoundId = syms_str8(&Bin.str[SectionOffset + DescOffset], 16);
+								}
+							}
+						}
+					}
+					break;
+				}
+				else if (syms_hash_djb2(Sections.v[SectionIndex].name) == 0xaab84f54dfa67dee)
+				{
+					if (SectionSize >= 16)
+					{
+						FoundId = syms_str8(&Bin.str[SectionOffset], 16);
+					}
+					break;
+				}
+			}
+
+			if (FoundId.size == 16 && FoundId.size == ImageId.Num())
+			{
+				if (FMemory::Memcmp(ImageId.GetData(), FoundId.str, FoundId.size) != 0)
+				{
+					// mismatch
+					return false;
+				}
+			}
+		}
+
+		// either ID's are matching, or ID is not found in which case return "success" case
+		return true;
+	}
 }
 
 /////////////////////////////////////////////////////////////////////
@@ -448,22 +528,11 @@ FSymslibResolver::EModuleStatus FSymslibResolver::LoadModule(FModuleEntry* Modul
 	}
 
 	// check debug data mismatch to captured module
-	if (!Module->ImageId.IsEmpty() && Dbg->format == SYMS_FileFormat_PDB)
+	if (!Module->ImageId.IsEmpty() && !MatchImageId(Module->ImageId, FileInf.data_parsed))
 	{
-		// for Pdbs checksum is a 16 byte guid and 4 byte unsigned integer for age, but usually age is not used for matching debug file to exe
-		static_assert(sizeof(FGuid) == 16, "Expected 16 byte FGuid");
-		check(Module->ImageId.Num() == 20);
-		FGuid* ModuleGuid = (FGuid*)Module->ImageId.GetData();
-
-		SYMS_ExtMatchKey MatchKey = syms_ext_match_key_from_dbg(FileInf.data_parsed.dbg_data, Dbg);
-		FGuid* PdbGuid = (FGuid*)MatchKey.v;
-
-		if (*ModuleGuid != *PdbGuid)
-		{
-			syms_group_release(Group);
-			UE_LOG(LogSymslib, Warning, TEXT("Symbols for '%s' does not match traced binary."), Module->Name);
-			return EModuleStatus::VersionMismatch;
-		}
+		syms_group_release(Group);
+		UE_LOG(LogSymslib, Warning, TEXT("Symbols for '%s' does not match traced binary."), Module->Name);
+		return EModuleStatus::VersionMismatch;
 	}
 
 	FSymsInstance* Instance = &Module->Instance;
