[Insights] Introduced Tag type alias

Introduced a type alias for Tag type id, abstracting how we represent tags internally. This doesn't represent a a functional change, as in practise existing instrumentation only emitted postive values and previously uint32 was used for storage.

diff --git a/Engine/Source/Developer/TraceServices/Private/Analyzers/AllocationsAnalysis.cpp b/Engine/Source/Developer/TraceServices/Private/Analyzers/AllocationsAnalysis.cpp
--- a/Engine/Source/Developer/TraceServices/Private/Analyzers/AllocationsAnalysis.cpp
+++ b/Engine/Source/Developer/TraceServices/Private/Analyzers/AllocationsAnalysis.cpp
@@ -238,8 +238,8 @@ bool FAllocationsAnalyzer::OnEvent(uint16 RouteId, EStyle Style, const FOnEventC
 		}
 		case RouteId_TagSpec:
 		{
-			const int32 Tag = Context.EventData.GetValue<int32>("Tag");
-			const int32 Parent = Context.EventData.GetValue<int32>("Parent");
+			const TagIdType Tag = Context.EventData.GetValue<TagIdType>("Tag");
+			const TagIdType Parent = Context.EventData.GetValue<TagIdType>("Parent");
 
 			FString Display;
 			Context.EventData.GetString("Display", Display);
@@ -258,7 +258,7 @@ bool FAllocationsAnalyzer::OnEvent(uint16 RouteId, EStyle Style, const FOnEventC
 			{
 				if (FCStringAnsi::Strlen(Context.EventData.GetTypeInfo().GetName()) == 11) // "MemoryScope"
 				{
-					const int32 Tag = Context.EventData.GetValue<int32>("Tag");
+					const TagIdType Tag = Context.EventData.GetValue<TagIdType>("Tag");
 					FAllocationsProvider::FEditScopeLock _(AllocationsProvider);
 					AllocationsProvider.EditPushTag(ThreadId, Tracker, Tag);
 				}
diff --git a/Engine/Source/Developer/TraceServices/Private/Model/AllocationsProvider.cpp b/Engine/Source/Developer/TraceServices/Private/Model/AllocationsProvider.cpp
--- a/Engine/Source/Developer/TraceServices/Private/Model/AllocationsProvider.cpp
+++ b/Engine/Source/Developer/TraceServices/Private/Model/AllocationsProvider.cpp
@@ -113,7 +113,7 @@ void FAllocationsProviderLock::EndWrite()
 // FTagTracker
 ////////////////////////////////////////////////////////////////////////////////////////////////////
 
-void FTagTracker::AddTagSpec(uint32 InTag, uint32 InParentTag, const TCHAR* InDisplay)
+void FTagTracker::AddTagSpec(TagIdType InTag, TagIdType InParentTag, const TCHAR* InDisplay)
 {
 	if (ensure(!TagMap.Contains(InTag)))
 	{
@@ -131,7 +131,7 @@ void FTagTracker::AddTagSpec(uint32 InTag, uint32 InParentTag, const TCHAR* InDi
 
 ////////////////////////////////////////////////////////////////////////////////////////////////////
 
-void FTagTracker::PushTag(uint32 InThreadId, uint8 InTracker, uint32 InTag)
+void FTagTracker::PushTag(uint32 InThreadId, uint8 InTracker, TagIdType InTag)
 {
 	const uint32 TrackerThreadId = GetTrackerThreadId(InThreadId, InTracker);
 	ThreadState& State = TrackerThreadStates.FindOrAdd(TrackerThreadId);
@@ -161,7 +161,7 @@ void FTagTracker::PopTag(uint32 InThreadId, uint8 InTracker)
 
 ////////////////////////////////////////////////////////////////////////////////////////////////////
 
-uint32 FTagTracker::GetCurrentTag(uint32 InThreadId, uint8 InTracker) const
+TagIdType FTagTracker::GetCurrentTag(uint32 InThreadId, uint8 InTracker) const
 {
 	const uint32 TrackerThreadId = GetTrackerThreadId(InThreadId, InTracker);
 	const ThreadState* State = TrackerThreadStates.Find(TrackerThreadId);
@@ -169,24 +169,24 @@ uint32 FTagTracker::GetCurrentTag(uint32 InThreadId, uint8 InTracker) const
 	{
 		return 0; // Untagged
 	}
-	return State->TagStack.Top() & ~0x80000000;
+	return State->TagStack.Top() & ~PtrTagMask;
 }
 
 ////////////////////////////////////////////////////////////////////////////////////////////////////
 
-const TCHAR* FTagTracker::GetTagString(uint32 InTag) const
+const TCHAR* FTagTracker::GetTagString(TagIdType InTag) const
 {
 	const TagEntry* Entry = TagMap.Find(InTag);
-	return Entry ? Entry->Display : TEXT("");
+	return Entry ? Entry->Display : TEXT("Unknown");
 }
 
 ////////////////////////////////////////////////////////////////////////////////////////////////////
 
-void FTagTracker::PushTagFromPtr(uint32 InThreadId, uint8 InTracker, uint32 InTag)
+void FTagTracker::PushTagFromPtr(uint32 InThreadId, uint8 InTracker, TagIdType InTag)
 {
 	const uint32 TrackerThreadId = GetTrackerThreadId(InThreadId, InTracker);
 	ThreadState& State = TrackerThreadStates.FindOrAdd(TrackerThreadId);
-	State.TagStack.Push(InTag | 0x80000000);
+	State.TagStack.Push(InTag | PtrTagMask);
 }
 
 ////////////////////////////////////////////////////////////////////////////////////////////////////
@@ -197,7 +197,7 @@ void FTagTracker::PopTagFromPtr(uint32 InThreadId, uint8 InTracker)
 	ThreadState* State = TrackerThreadStates.Find(TrackerThreadId);
 	if (State && !State->TagStack.IsEmpty())
 	{
-		INSIGHTS_SLOW_CHECK((State->TagStack.Top() & 0x80000000) != 0);
+		INSIGHTS_SLOW_CHECK((State->TagStack.Top() & PtrTagMask) != 0);
 		State->TagStack.Pop();
 	}
 	else
@@ -216,7 +216,7 @@ bool FTagTracker::HasTagFromPtrScope(uint32 InThreadId, uint8 InTracker) const
 {
 	const uint32 TrackerThreadId = GetTrackerThreadId(InThreadId, InTracker);
 	const ThreadState* State = TrackerThreadStates.Find(TrackerThreadId);
-	return State && State->TagStack.Num() > 0 && ((State->TagStack.Top() & 0x80000000) != 0);
+	return State && State->TagStack.Num() > 0 && ((State->TagStack.Top() & PtrTagMask) != 0);
 }
 
 ////////////////////////////////////////////////////////////////////////////////////////////////////
@@ -286,7 +286,7 @@ const FCallstack* IAllocationsProvider::FAllocation::GetCallstack() const
 
 ////////////////////////////////////////////////////////////////////////////////////////////////////
 
-uint32 IAllocationsProvider::FAllocation::GetTag() const
+TagIdType IAllocationsProvider::FAllocation::GetTag() const
 {
 	const auto* Inner = (const FAllocationItem*)this;
 	return Inner->Tag;
@@ -1107,7 +1107,7 @@ void FAllocationsProvider::EditAlloc(double Time, uint64 Owner, uint64 Address,
 
 	AdvanceTimelines(Time);
 
-	const uint32 Tag = TagTracker.GetCurrentTag(ThreadId, Tracker);
+	const TagIdType Tag = TagTracker.GetCurrentTag(ThreadId, Tracker);
 
 #if INSIGHTS_DEBUG_WATCH
 	for (int32 AddrIndex = 0; AddrIndex < UE_ARRAY_COUNT(GWatchAddresses); ++AddrIndex)
@@ -1450,7 +1450,7 @@ HeapId FAllocationsProvider::FindRootHeap(HeapId Heap) const
 
 ////////////////////////////////////////////////////////////////////////////////////////////////////
 
-void FAllocationsProvider::EditPushTag(uint32 ThreadId, uint8 Tracker, uint32 Tag)
+void FAllocationsProvider::EditPushTag(uint32 ThreadId, uint8 Tracker, TagIdType Tag)
 {
 	EditAccessCheck();
 
@@ -1477,7 +1477,7 @@ void FAllocationsProvider::EditPushTagFromPtr(uint32 ThreadId, uint8 Tracker, ui
 	if (ensure(Allocs))
 	{
 		FAllocationItem* Alloc = Allocs->FindRef(Ptr);
-		const int32 Tag = Alloc ? Alloc->Tag : 0; // If ptr is not found use "Untagged"
+		const TagIdType Tag = Alloc ? Alloc->Tag : 0; // If ptr is not found use "Untagged"
 		TagTracker.PushTagFromPtr(ThreadId, Tracker, Tag);
 	}
 	else
diff --git a/Engine/Source/Developer/TraceServices/Private/Model/AllocationsProvider.h b/Engine/Source/Developer/TraceServices/Private/Model/AllocationsProvider.h
--- a/Engine/Source/Developer/TraceServices/Private/Model/AllocationsProvider.h
+++ b/Engine/Source/Developer/TraceServices/Private/Model/AllocationsProvider.h
@@ -41,26 +41,27 @@ class FTagTracker
 private:
 	static constexpr uint32 TrackerIdShift = 24;
 	static constexpr uint32 TrackerIdMask = 0xFF000000;
-
+	static constexpr uint32 PtrTagMask = 0x80000000;
+	
 	struct ThreadState
 	{
-		TArray<uint32> TagStack;
+		TArray<TagIdType> TagStack;
 	};
 
 	struct TagEntry
 	{
 		const TCHAR* Display;
-		uint32 ParentTag;
+		TagIdType ParentTag;
 	};
 
 public:
-	void AddTagSpec(uint32 Tag, uint32 ParentTag, const TCHAR* Display);
-	void PushTag(uint32 ThreadId, uint8 Tracker, uint32 Tag);
+	void AddTagSpec(TagIdType Tag, TagIdType ParentTag, const TCHAR* Display);
+	void PushTag(uint32 ThreadId, uint8 Tracker, TagIdType Tag);
 	void PopTag(uint32 ThreadId, uint8 Tracker);
-	uint32 GetCurrentTag(uint32 ThreadId, uint8 Tracker) const;
-	const TCHAR* GetTagString(uint32 Tag) const;
+	TagIdType GetCurrentTag(uint32 ThreadId, uint8 Tracker) const;
+	const TCHAR* GetTagString(TagIdType Tag) const;
 
-	void PushTagFromPtr(uint32 ThreadId, uint8 Tracker, uint32 Tag);
+	void PushTagFromPtr(uint32 ThreadId, uint8 Tracker, TagIdType Tag);
 	void PopTagFromPtr(uint32 ThreadId, uint8 Tracker);
 	bool HasTagFromPtrScope(uint32 ThreadId, uint8 Tracker) const;
 
@@ -71,7 +72,7 @@ private:
 	}
 
 	TMap<uint32, ThreadState> TrackerThreadStates;
-	TMap<uint32, TagEntry> TagMap;
+	TMap<TagIdType, TagEntry> TagMap;
 	uint32 NumErrors = 0;
 };
 
@@ -100,7 +101,7 @@ struct FAllocationItem
 	uint64 Address;
 	uint64 SizeAndAlignment; // (Alignment << AlignmentShift) | Size
 	mutable const FCallstack* Callstack;
-	uint32 Tag;
+	TagIdType Tag;
 	uint8 RootHeap;
 	EMemoryTraceHeapAllocationFlags Flags;
 };
@@ -311,7 +312,7 @@ public:
 	void EnumerateLiveAllocs(TFunctionRef<void(const FAllocationItem& Alloc)> Callback) const;
 	uint32 GetNumLiveAllocs() const;
 
-	virtual const TCHAR* GetTagName(int32 Tag) const { ReadAccessCheck(); return TagTracker.GetTagString(Tag); }
+	virtual const TCHAR* GetTagName(TagIdType Tag) const { ReadAccessCheck(); return TagTracker.GetTagString(Tag); }
 	bool HasTagFromPtrScope(uint32 ThreadId, uint8 Tracker) const { ReadAccessCheck(); return TagTracker.HasTagFromPtrScope(ThreadId, Tracker); }
 
 	void DebugPrint() const;
@@ -328,8 +329,8 @@ public:
 	void EditMarkAllocationAsHeap(double Time, uint64 Address, HeapId Heap, EMemoryTraceHeapAllocationFlags Flags);
 	void EditUnmarkAllocationAsHeap(double Time, uint64 Address, HeapId Heap);
 	
-	void EditAddTagSpec(int32 Tag, uint32 ParentTag, const TCHAR* Display) { EditAccessCheck(); TagTracker.AddTagSpec(Tag, ParentTag, Display); }
-	void EditPushTag(uint32 ThreadId, uint8 Tracker, uint32 Tag);
+	void EditAddTagSpec(TagIdType Tag, TagIdType ParentTag, const TCHAR* Display) { EditAccessCheck(); TagTracker.AddTagSpec(Tag, ParentTag, Display); }
+	void EditPushTag(uint32 ThreadId, uint8 Tracker, TagIdType Tag);
 	void EditPopTag(uint32 ThreadId, uint8 Tracker);
 	void EditPushTagFromPtr(uint32 ThreadId, uint8 Tracker, uint64 Ptr);
 	void EditPopTagFromPtr(uint32 ThreadId, uint8 Tracker);
diff --git a/Engine/Source/Developer/TraceServices/Public/TraceServices/Model/AllocationsProvider.h b/Engine/Source/Developer/TraceServices/Public/TraceServices/Model/AllocationsProvider.h
--- a/Engine/Source/Developer/TraceServices/Public/TraceServices/Model/AllocationsProvider.h
+++ b/Engine/Source/Developer/TraceServices/Public/TraceServices/Model/AllocationsProvider.h
@@ -14,6 +14,11 @@ namespace TraceServices
 
 ////////////////////////////////////////////////////////////////////////////////////////////////////
 
+// Id type for tags
+typedef uint32 TagIdType;
+	
+////////////////////////////////////////////////////////////////////////////////////////////////////
+
 class IAllocationsProvider : public IProvider
 {
 public:
@@ -94,7 +99,7 @@ public:
 		uint64 GetSize() const;
 		uint32 GetAlignment() const;
 		const struct FCallstack* GetCallstack() const;
-		uint32 GetTag() const;
+		TagIdType GetTag() const;
 		HeapId GetRootHeap() const;
 		bool IsHeap() const;
 	};
@@ -179,7 +184,7 @@ public:
 
 	// Returns the display name of the specified LLM tag.
 	// Lifetime of returned string matches the session lifetime.
-	virtual const TCHAR* GetTagName(int32 Tag) const = 0;
+	virtual const TCHAR* GetTagName(TagIdType Tag) const = 0;
 };
 
 TRACESERVICES_API FName GetAllocationsProviderName();
