[Insights] Fixed bad cpu timer names for timers with more than 128 chars traced as ANSICHAR. The string conversion allocates dynamically a buffer when the source string has more than 128 chars and that is freed when string conversion object goes out of scope.

diff --git a/Engine/Source/Developer/TraceServices/Private/Analyzers/CpuProfilerTraceAnalysis.cpp b/Engine/Source/Developer/TraceServices/Private/Analyzers/CpuProfilerTraceAnalysis.cpp
--- a/Engine/Source/Developer/TraceServices/Private/Analyzers/CpuProfilerTraceAnalysis.cpp
+++ b/Engine/Source/Developer/TraceServices/Private/Analyzers/CpuProfilerTraceAnalysis.cpp
@@ -57,7 +57,9 @@ bool FCpuProfilerAnalyzer::OnEvent(uint16 RouteId, EStyle Style, const FOnEventC
 			uint8 CharSize = EventData.GetValue<uint8>("CharSize");
 			if (CharSize == sizeof(ANSICHAR))
 			{
-				TimerName = StringCast<TCHAR>(reinterpret_cast<const ANSICHAR*>(EventData.GetAttachment())).Get();
+				const ANSICHAR* AnsiName = reinterpret_cast<const ANSICHAR*>(EventData.GetAttachment());
+				Name = StringCast<TCHAR>(AnsiName).Get();
+				TimerName = *Name;
 			}
 			else if (CharSize == 0 || CharSize == sizeof(TCHAR)) // 0 for backwards compatibility
 			{
