[Insights] Fixed Counters not being stored in monotonic order. For a counter with this bug, the graph track series shows points displayed at random locations (above or below the line of the graph series) that moves or disappears when panning or zooming.

diff --git a/Engine/Source/Developer/TraceServices/Private/Model/CountersPrivate.h b/Engine/Source/Developer/TraceServices/Private/Model/CountersPrivate.h
--- a/Engine/Source/Developer/TraceServices/Private/Model/CountersPrivate.h
+++ b/Engine/Source/Developer/TraceServices/Private/Model/CountersPrivate.h
@@ -147,12 +147,14 @@ public:
 			{
 				CurrentPage = TimestampIterator.PrevPage();
 			}
+			check(CurrentPage != nullptr);
 			uint64 PageInsertionIndex = Algo::LowerBound(*CurrentPage, Timestamp);
 #if USE_VARIABLE_PAGED_ARRAY
-			InsertionIndex = TimestampIterator.GetCurrentItemIndex() + PageInsertionIndex;
+			InsertionIndex = TimestampIterator.GetCurrentItemIndex() + PageInsertionIndex + 1 - CurrentPage->ItemCount;
 #else
 			InsertionIndex = TimestampIterator.GetCurrentPageIndex() * Timestamps.GetPageSize() + PageInsertionIndex;
 #endif
+			check(InsertionIndex <= Timestamps.Num());
 		}
 		Timestamps.Insert(InsertionIndex) = Timestamp;
 		OpTypes.Insert(InsertionIndex) = OpType;
