[Insights] Add symbol search paths and reload module triggers

 * Implements functionality to search for symbols in a stack of search directories. Search directories can be added from configuration, environment variables or explicitly specified by user.
 * Exposes module information and ability to enumerate detected modules.
 * Adds functionality to attempt loading symbols for a module in a specific directory.

diff --git a/Engine/Source/Developer/TraceInsights/Private/Insights/MemoryProfiler/Widgets/SMemAllocTableTreeView.cpp b/Engine/Source/Developer/TraceInsights/Private/Insights/MemoryProfiler/Widgets/SMemAllocTableTreeView.cpp
--- a/Engine/Source/Developer/TraceInsights/Private/Insights/MemoryProfiler/Widgets/SMemAllocTableTreeView.cpp
+++ b/Engine/Source/Developer/TraceInsights/Private/Insights/MemoryProfiler/Widgets/SMemAllocTableTreeView.cpp
@@ -1004,9 +1004,9 @@ FText SMemAllocTableTreeView::GetSymbolResolutionStatus() const
 	if (ModuleProvider)
 	{
 		ModuleProvider->GetStats(&Stats);
-		check(Stats.SymbolsDiscovered >= Stats.SymbolsResolved + Stats.SymbolsFailed);
-		const uint32 SymbolsPending = Stats.SymbolsDiscovered - Stats.SymbolsResolved - Stats.SymbolsFailed;
-		if (SymbolsPending != 0)
+		//check(Stats.SymbolsDiscovered >= Stats.SymbolsResolved + Stats.SymbolsFailed);
+		const int32 SymbolsPending = Stats.SymbolsDiscovered - Stats.SymbolsResolved - Stats.SymbolsFailed;
+		if (SymbolsPending > 0)
 		{
 			return FText::Format(LOCTEXT("SymbolsResolved1", "Resolving {0} / {1} symbols ({2} resolved, {3} failed)"), SymbolsPending, Stats.SymbolsDiscovered, Stats.SymbolsResolved, Stats.SymbolsFailed);
 		}
@@ -1026,11 +1026,11 @@ FText SMemAllocTableTreeView::GetSymbolResolutionStatus() const
 FText SMemAllocTableTreeView::GetSymbolResolutionTooltip() const
 {
 	auto ModuleProvider = Session->ReadProvider<IModuleProvider>(FName("ModuleProvider"));
-
-	static const TCHAR* SymbolPathEnvVar = TEXT("UE_INSIGHTS_SYMBOL_PATH");
-	FString SymbolPath = FPlatformMisc::GetEnvironmentVariable(SymbolPathEnvVar);
-	return FText::Format(LOCTEXT("SymbolResolutionEnvVarHelp", "Symbol paths can be provided with the environment variable {0}.\nUE_INSIGHTS_SYMBOL_PATH=\"{1}\""),
-		FText::FromString(SymbolPathEnvVar), FText::FromString(*SymbolPath));
+	TStringBuilder<2048> Sb;
+	uint32 Index = 0;
+	ModuleProvider->EnumerateSymbolSearchPaths([&Sb, &Index](FStringView Path) { Sb << ++Index << TEXT(": ") << Path << TEXT("\n");});
+	return FText::Format(LOCTEXT("SymbolResolutionEnvVarHelp", "Symbol paths:\n{0}"),
+		FText::FromStringView(Sb));
 }
 
 ////////////////////////////////////////////////////////////////////////////////////////////////////
diff --git a/Engine/Source/Developer/TraceServices/Private/Model/DbgHelpResolver.cpp b/Engine/Source/Developer/TraceServices/Private/Model/DbgHelpResolver.cpp
--- a/Engine/Source/Developer/TraceServices/Private/Model/DbgHelpResolver.cpp
+++ b/Engine/Source/Developer/TraceServices/Private/Model/DbgHelpResolver.cpp
@@ -4,6 +4,7 @@
 
 #include "DbgHelpResolver.h"
 #include "Algo/Sort.h"
+#include "Algo/ForEach.h"
 #include "Containers/Queue.h"
 #include "Containers/StringView.h"
 #include "HAL/PlatformProcess.h"
@@ -28,21 +29,37 @@ namespace TraceServices {
 
 static const TCHAR* GUnknownModuleTextDbgHelp = TEXT("Unknown");
 
+
 /////////////////////////////////////////////////////////////////////
 FDbgHelpResolver::FDbgHelpResolver(IAnalysisSession& InSession)
 	: bRunWorkerThread(false)
 	, bDrainThenStop(false)
-	, bInitialized(false)
 	, Session(InSession)
 {
-	bInitialized = SetupSyms();
-
-	if (bInitialized)
+	// Setup search paths. The SearchPaths array is a priority stack, which
+	// means paths are searched in reversed order.
+	// 1. Any new paths entered by the user this session
+	// 2. Path of the executable (if available)
+	// 3. Paths from UE_INSIGHTS_SYMBOLPATH
+	// 4. Paths from the user configuration file
+
+	// Paths from configuration
+	FString SettingsIni;
+	if (FConfigCacheIni::LoadGlobalIniFile(SettingsIni, TEXT("UnrealInsightsSettings")))
+	{
+		GConfig->GetArray(TEXT("Insights.MemoryProfiler"), TEXT("SymbolSearchPaths"), SymbolSearchPaths, SettingsIni);
+	}
+	
+	// Paths from environment
+	FString SymbolPathEnvVar =  FPlatformMisc::GetEnvironmentVariable(TEXT("UE_INSIGHTS_SYMBOL_PATH"));
+	UE_LOG(LogDbgHelp, Log, TEXT("UE_INSIGHTS_SYMBOL_PATH: '%s'"), *SymbolPathEnvVar);
+	FString SymbolPathPart;
+	while (SymbolPathEnvVar.Split(TEXT(";"), &SymbolPathPart, &SymbolPathEnvVar))
 	{
-		// Start the worker thread
-		bRunWorkerThread = true;
-		Thread = FRunnableThread::Create(this, TEXT("DbgHelpWorker"), 0, TPri_Normal);
+		SymbolSearchPaths.Emplace(SymbolPathPart);
 	}
+	
+	Start();
 }
 
 /////////////////////////////////////////////////////////////////////
@@ -56,17 +73,22 @@ FDbgHelpResolver::~FDbgHelpResolver()
 }
 
 /////////////////////////////////////////////////////////////////////
-void FDbgHelpResolver::QueueModuleLoad(const FStringView& ModulePath, uint64 Base, uint32 Size, const uint8* ImageId, uint32 ImageIdSize)
+void FDbgHelpResolver::QueueModuleLoad(const uint8* ImageId, uint32 ImageIdSize, FModule* Module)
 {
+	check(Module != nullptr);
+	
 	FScopeLock _(&ModulesCs);
-
-	const FStringView ModuleName = FPathViews::GetCleanFilename(ModulePath);
+	
+	const FStringView ModuleName = FPathViews::GetCleanFilename(Module->FullName);
 
 	// Add module and sort list according to base address
-	const int32 Index = LoadedModules.Add(FModuleEntry{ Base, Size, Session.StoreString(ModuleName), Session.StoreString(ModulePath), false });
+	const int32 Index = LoadedModules.Add(FModuleEntry{
+		Module->Base, Module->Size, Session.StoreString(ModuleName), Session.StoreString(Module->FullName),
+		Module, TArray(ImageId, ImageIdSize)
+	});
 
 	// Queue up module to have symbols loaded
-	LoadSymbolsQueue.Enqueue(FQueuedModule{ Base, Size, LoadedModules[Index].Path });
+	LoadSymbolsQueue.Enqueue(FQueuedModule{ Module, nullptr, LoadedModules[Index].ImageId});
 
 	// Sort list according to base address
 	Algo::Sort(LoadedModules, [](const FModuleEntry& Lhs, const FModuleEntry& Rhs) { return Lhs.Base < Rhs.Base; });
@@ -74,24 +96,61 @@ void FDbgHelpResolver::QueueModuleLoad(const FStringView& ModulePath, uint64 Bas
 	++ModulesDiscovered;
 }
 
+/////////////////////////////////////////////////////////////////////
+void FDbgHelpResolver::QueueModuleReload(const FModule* Module, const TCHAR* InPath, TFunction<void(SymbolArray&)> ResolveOnSuccess)
+{
+	FScopeLock _(&ModulesCs);
+	const uint64 ModuleBase = Module->Base;
+	const FModuleEntry* Entry = LoadedModules.FindByPredicate([ModuleBase](const FModuleEntry& Entry) { return Entry.Base == ModuleBase; });
+	if (Entry)
+	{
+		Entry->Module->Status.store(EModuleStatus::Pending);
+		LoadSymbolsQueue.Enqueue(FQueuedModule{Module, Session.StoreString(InPath), TArrayView<const uint8>(Entry->ImageId)});
+	}
+	
+	SymbolArray SymbolsToResolve;
+	ResolveOnSuccess(SymbolsToResolve);
+	for(TTuple<uint64, FResolvedSymbol*> Pair : SymbolsToResolve)
+	{
+		QueueSymbolResolve(Pair.Get<0>(), Pair.Get<1>());
+	}
+	
+	if (!bRunWorkerThread && Thread)
+	{
+		// Restart the worker thread if it has stopped.
+		Start();
+	}
+}
+	
 /////////////////////////////////////////////////////////////////////
 void FDbgHelpResolver::QueueSymbolResolve(uint64 Address, FResolvedSymbol* Symbol)
 {
-	++SymbolsDiscovered;
 	ResolveQueue.Enqueue(FQueuedAddress{Address, Symbol});
 }
 
 /////////////////////////////////////////////////////////////////////
 void FDbgHelpResolver::GetStats(IModuleProvider::FStats* OutStats) const
 {
+	FScopeLock _(&ModulesCs);
+	FMemory::Memzero(*OutStats);
+	for(const FModuleEntry& Entry : LoadedModules)
+	{
+		OutStats->SymbolsDiscovered += Entry.Module->Stats.Discovered.load();
+		OutStats->SymbolsResolved += Entry.Module->Stats.Resolved.load();
+		OutStats->SymbolsFailed += Entry.Module->Stats.Failed.load();
+	}
 	OutStats->ModulesDiscovered = ModulesDiscovered.load();
 	OutStats->ModulesFailed = ModulesFailed.load();
 	OutStats->ModulesLoaded = ModulesLoaded.load();
-	OutStats->SymbolsDiscovered = SymbolsDiscovered.load();
-	OutStats->SymbolsFailed = SymbolsFailed.load();
-	OutStats->SymbolsResolved = SymbolsResolved.load();
 }
 
+/////////////////////////////////////////////////////////////////////
+void FDbgHelpResolver::EnumerateSymbolSearchPaths(TFunctionRef<void(FStringView Path)> Callback) const
+{
+	FScopeLock _(&SymbolSearchPathsLock);
+	Algo::ForEach(SymbolSearchPaths, Callback);
+}
+	
 /////////////////////////////////////////////////////////////////////
 void FDbgHelpResolver::OnAnalysisComplete()
 {
@@ -112,12 +171,13 @@ bool FDbgHelpResolver::SetupSyms()
 	ULONG SymOpts = 0;
 	SymOpts |= SYMOPT_LOAD_LINES;
 	SymOpts |= SYMOPT_OMAP_FIND_NEAREST;
-	SymOpts |= SYMOPT_DEFERRED_LOADS;
+	//SymOpts |= SYMOPT_DEFERRED_LOADS;
 	SymOpts |= SYMOPT_EXACT_SYMBOLS;
 	SymOpts |= SYMOPT_IGNORE_NT_SYMPATH;
 	SymOpts |= SYMOPT_UNDNAME;
 
 	SymSetOptions(SymOpts);
+	
 	return SymInitialize((HANDLE)Handle, NULL, FALSE);
 }
 
@@ -133,7 +193,9 @@ void FDbgHelpResolver::FreeSyms() const
 /////////////////////////////////////////////////////////////////////
 uint32 FDbgHelpResolver::Run()
 {
-	while (bRunWorkerThread)
+	const bool bInitialized = SetupSyms();
+	
+	while (bInitialized && bRunWorkerThread)
 	{
 		// Prioritize queued module loads
 		while (!LoadSymbolsQueue.IsEmpty() && bRunWorkerThread)
@@ -141,7 +203,7 @@ uint32 FDbgHelpResolver::Run()
 			FQueuedModule Item;
 			if (LoadSymbolsQueue.Dequeue(Item))
 			{
-				LoadModuleSymbols(Item.Base, Item.Size, Item.ImagePath);
+				LoadModuleSymbols(Item.Module, Item.Path, Item.ImageId);
 			}
 		}
 
@@ -170,6 +232,14 @@ uint32 FDbgHelpResolver::Run()
 	return 0;
 }
 
+/////////////////////////////////////////////////////////////////////
+void FDbgHelpResolver::Start()
+{
+	// Start the worker thread
+	bRunWorkerThread = true;
+	Thread = FRunnableThread::Create(this, TEXT("DbgHelpWorker"), 0, TPri_Normal);
+}
+	
 /////////////////////////////////////////////////////////////////////
 void FDbgHelpResolver::Stop()
 {
@@ -191,15 +261,34 @@ void FDbgHelpResolver::ResolveSymbol(uint64 Address, FResolvedSymbol* Target)
 {
 	check(Target);
 
-	const FModuleEntry* Module = GetModuleForAddress(Address);
-	if (!Module)
+	if (Target->Result.load() == ESymbolQueryResult::OK)
+	{
+		return;
+	}
+	
+	const FModuleEntry* Entry = GetModuleForAddress(Address);
+	if (!Entry)
 	{
-		++SymbolsFailed;
-		UpdateResolvedSymbol(Target, ESymbolQueryResult::NotLoaded, GUnknownModuleTextDbgHelp,
+		++Entry->Module->Stats.Failed;
+		UpdateResolvedSymbol(Target, ESymbolQueryResult::NotFound, GUnknownModuleTextDbgHelp,
 		                     GUnknownModuleTextDbgHelp, GUnknownModuleTextDbgHelp, 0);
 		return;
 	}
 	
+	const EModuleStatus ModuleStatus = Entry->Module->Status.load();
+	if (ModuleStatus != EModuleStatus::Loaded)
+	{
+		const ESymbolQueryResult Result = ModuleStatus == EModuleStatus::VersionMismatch
+			                                  ? ESymbolQueryResult::Mismatch
+			                                  : ESymbolQueryResult::NotLoaded;
+		++Entry->Module->Stats.Failed;
+		UpdateResolvedSymbol(Target, Result, GUnknownModuleTextDbgHelp,
+		                     GUnknownModuleTextDbgHelp, GUnknownModuleTextDbgHelp, 0);
+		return;
+	}
+
+	++Entry->Module->Stats.Discovered;
+	
 	uint8 InfoBuffer[sizeof(SYMBOL_INFO) + (MaxNameLen * sizeof(char) + 1)];
 	SYMBOL_INFO* Info = (SYMBOL_INFO*)InfoBuffer;
 	Info->SizeOfStruct = sizeof(SYMBOL_INFO);
@@ -208,8 +297,8 @@ void FDbgHelpResolver::ResolveSymbol(uint64 Address, FResolvedSymbol* Target)
 	// Find and build the symbol name
 	if (!SymFromAddr((HANDLE)Handle, Address, NULL, Info))
 	{
-		++SymbolsFailed;
-		UpdateResolvedSymbol(Target, ESymbolQueryResult::NotFound, Module->Name, GUnknownModuleTextDbgHelp,
+		++Entry->Module->Stats.Failed;
+		UpdateResolvedSymbol(Target, ESymbolQueryResult::NotFound, Entry->Name, GUnknownModuleTextDbgHelp,
 		                     GUnknownModuleTextDbgHelp, 0);
 		return;
 	}
@@ -223,42 +312,104 @@ void FDbgHelpResolver::ResolveSymbol(uint64 Address, FResolvedSymbol* Target)
 
 	if (!SymGetLineFromAddr((HANDLE)Handle, Address, &dwDisplacement, &Line))
 	{
-		++SymbolsFailed;
-		UpdateResolvedSymbol(Target, ESymbolQueryResult::OK, Module->Name, SymbolNameStr, GUnknownModuleTextDbgHelp, 0);
+		++Entry->Module->Stats.Failed;
+		UpdateResolvedSymbol(Target, ESymbolQueryResult::OK, Entry->Name, SymbolNameStr, GUnknownModuleTextDbgHelp, 0);
 		return;
 	}
 	
 	const TCHAR* SymbolFileStr = Session.StoreString(ANSI_TO_TCHAR(Line.FileName));
 	
-	++SymbolsResolved;
-	UpdateResolvedSymbol(Target, ESymbolQueryResult::OK, Module->Name, SymbolNameStr, SymbolFileStr, Line.LineNumber);
+	++Entry->Module->Stats.Resolved;
+	UpdateResolvedSymbol(Target, ESymbolQueryResult::OK, Entry->Name, SymbolNameStr, SymbolFileStr, Line.LineNumber);
 }
-
+	
 /////////////////////////////////////////////////////////////////////
-bool FDbgHelpResolver::LoadModuleSymbols(uint64 Base, uint64 Size, const TCHAR* Path)
+void FDbgHelpResolver::LoadModuleSymbols(const FModule* Module, const TCHAR* Path, const TArrayView<const uint8> ImageId)
 {
+	check(Module);
+
+	const uint64 Base = Module->Base;
+	const uint32 Size = Module->Size;
+
+	// Setup symbol search path
+	{
+		TAnsiStringBuilder<1024> UserSearchPath;
+		UserSearchPath << Path << ";";
+		Algo::ForEach(SymbolSearchPaths, [&UserSearchPath] (const FString& Path){ UserSearchPath.Appendf("%s;", TCHAR_TO_ANSI(*Path));});
+
+		if (!SymSetSearchPath((HANDLE)Handle, UserSearchPath.ToString()))
+		{
+			UE_LOG(LogDbgHelp, Warning, TEXT("Unable to set symbol search path to '%s'."), UserSearchPath.ToString());
+		}
+		TCHAR OutPath[1024];
+		SymGetSearchPathW((HANDLE) Handle, OutPath, 1024);
+		UE_LOG(LogDbgHelp, Display, TEXT("Search path: %s"), OutPath);
+	}
+	
 	// Attempt to load symbols
-	const DWORD64 LoadedBaseAddress = SymLoadModuleEx((HANDLE)Handle, NULL, TCHAR_TO_ANSI(Path), NULL, Base, Size, NULL, 0);
-	const bool bSymbolsLoaded = Base == LoadedBaseAddress;
+	const DWORD64 LoadedBaseAddress = SymLoadModuleEx((HANDLE)Handle, NULL, TCHAR_TO_ANSI(Module->Name), NULL, Base, Size, NULL, 0);
+	const bool bModuleLoaded = Base == LoadedBaseAddress;
+	bool bPdbLoaded = true;
+	bool bPdbMatchesImage = true;
+	IMAGEHLP_MODULE ModuleInfo;
 
-	if (!bSymbolsLoaded)
+	if (bModuleLoaded)
 	{
-		UE_LOG(LogDbgHelp, Warning, TEXT("Unable to load symbols for %s at 0x%016llx"), Path, Base);
+		ModuleInfo.SizeOfStruct = sizeof(IMAGEHLP_MODULE);
+		SymGetModuleInfo((HANDLE)Handle, Base, &ModuleInfo);
+
+		if (ModuleInfo.SymType != SymPdb)
+		{
+			bPdbLoaded = false;
+		}
+		// Check image checksum if it exists
+		else if (!ImageId.IsEmpty())
+		{
+			// for Pdbs checksum is a 16 byte guid and 4 byte unsigned integer for age, but usually age is not used for matching debug file to exe
+			static_assert(sizeof(FGuid) == 16, "Expected 16 byte FGuid");
+			check(ImageId.Num() == 20);
+			const FGuid* ModuleGuid = (FGuid*) ImageId.GetData();
+			const FGuid* PdbGuid = (FGuid*) &ModuleInfo.PdbSig70;
+			bPdbMatchesImage = *ModuleGuid == *PdbGuid;
+		}
+	}
+
+	TStringBuilder<256> StatusMessage;
+	EModuleStatus Status;
+	if (!bModuleLoaded || !bPdbLoaded)
+	{
+		// Unload the module, otherwise any subsequent attempts to load module with another
+		// path will fail.
+		SymUnloadModule((HANDLE)Handle, Base);
+		StatusMessage.Appendf(TEXT("Unable to load symbols for %s"), Path);
+		Status = EModuleStatus::Failed;
+		++ModulesFailed;
+	}
+	else if (!bPdbMatchesImage)
+	{
+		// Unload the module, otherwise any subsequent attempts to load module with another
+		// path will fail.
+		SymUnloadModule((HANDLE)Handle, Base);
+		StatusMessage.Appendf(TEXT("Unable to load symbols for %s, pdb signature does not match."), Path);
+		Status = EModuleStatus::VersionMismatch;
 		++ModulesFailed;
 	}
 	else
 	{
-		UE_LOG(LogDbgHelp, Display, TEXT("Loaded symbols for %s at 0x%016llx."), Path, Base);
+		StatusMessage.Appendf(TEXT("Loaded symbols for %s from %s."), Path, ANSI_TO_TCHAR(ModuleInfo.LoadedImageName));
+		Status = EModuleStatus::Loaded;
 		++ModulesLoaded;
 	}
 
 	// Update the module entry with the result
 	FScopeLock _(&ModulesCs);
 	const int32 EntryIdx = Algo::BinarySearchBy(LoadedModules, Base, [](const FModuleEntry& Entry) { return Entry.Base; });
-	check(EntryIdx >= 0 && EntryIdx < LoadedModules.Num());
-	LoadedModules[EntryIdx].bSymbolsLoaded = bSymbolsLoaded;
-	
-	return bSymbolsLoaded;
+	check(EntryIdx != INDEX_NONE);
+	const FModuleEntry& Entry = LoadedModules[EntryIdx];
+
+	// Make the status visible to the world
+	Entry.Module->StatusMessage = Session.StoreString(StatusMessage.ToView());
+	Entry.Module->Status.store(Status);
 }
 
 
diff --git a/Engine/Source/Developer/TraceServices/Private/Model/DbgHelpResolver.h b/Engine/Source/Developer/TraceServices/Private/Model/DbgHelpResolver.h
--- a/Engine/Source/Developer/TraceServices/Private/Model/DbgHelpResolver.h
+++ b/Engine/Source/Developer/TraceServices/Private/Model/DbgHelpResolver.h
@@ -19,11 +19,16 @@ namespace TraceServices {
 class FDbgHelpResolver : public FRunnable
 {
 public:
+	typedef TArray<TTuple<uint64, FResolvedSymbol*>> SymbolArray;
+	
+	void Start();
 	FDbgHelpResolver(IAnalysisSession& InSession);
 	~FDbgHelpResolver();
-	void QueueModuleLoad(const FStringView& ModulePath, uint64 Base, uint32 Size, const uint8* ImageId, uint32 ImageIdSize);
+	void QueueModuleLoad(const uint8* ImageId, uint32 ImageIdSize, FModule* Module);
+	void QueueModuleReload(const FModule* Module, const TCHAR* InPath, TFunction<void(SymbolArray&)> ResolveOnSuccess);
 	void QueueSymbolResolve(uint64 Address, FResolvedSymbol* Symbol);
 	void GetStats(IModuleProvider::FStats* OutStats) const;
+	void EnumerateSymbolSearchPaths(TFunctionRef<void(FStringView Path)> Callback) const;
 	void OnAnalysisComplete();
 
 private:	
@@ -33,7 +38,8 @@ private:
 		uint32 Size;
 		const TCHAR* Name;
 		const TCHAR* Path;
-		bool bSymbolsLoaded;
+		FModule* Module;
+		TArray<uint8> ImageId;
 	};
 
 	struct FQueuedAddress
@@ -44,9 +50,9 @@ private:
 
 	struct FQueuedModule
 	{
-		uint64 Base;
-		uint64 Size;
-		const TCHAR* ImagePath;
+		const FModule* Module;
+		const TCHAR* Path;
+		TArrayView<const uint8> ImageId;
 	};
 
 	enum : uint32 {
@@ -61,10 +67,10 @@ private:
 	static void UpdateResolvedSymbol(FResolvedSymbol* Symbol, ESymbolQueryResult Result, const TCHAR* Module, const TCHAR* Name, const TCHAR* File, uint16 Line);
 
 	void ResolveSymbol(uint64 Address, FResolvedSymbol* Target);
-	bool LoadModuleSymbols(uint64 Base, uint64 Size, const TCHAR* Path);
+	void LoadModuleSymbols(const FModule* Module, const TCHAR* Path, const TArrayView<const uint8> ImageId);
 	const FModuleEntry* GetModuleForAddress(uint64 Address) const;
 
-	FCriticalSection ModulesCs;
+	mutable FCriticalSection ModulesCs;
 	TArray<FModuleEntry> LoadedModules;
 	TQueue<FQueuedModule, EQueueMode::Mpsc> LoadSymbolsQueue;
 	TQueue<FQueuedAddress, EQueueMode::Mpsc> ResolveQueue;
@@ -72,13 +78,11 @@ private:
 	std::atomic<uint32> ModulesDiscovered;
 	std::atomic<uint32> ModulesFailed;
 	std::atomic<uint32> ModulesLoaded;
-	std::atomic<uint32> SymbolsDiscovered;
-	std::atomic<uint32> SymbolsFailed;
-	std::atomic<uint32> SymbolsResolved;
 
+	mutable FCriticalSection SymbolSearchPathsLock;
+	TArray<FString> SymbolSearchPaths;
 	bool bRunWorkerThread;
 	bool bDrainThenStop;
-	bool bInitialized;
 	UPTRINT Handle;
 	IAnalysisSession& Session;
 	FRunnableThread* Thread;
diff --git a/Engine/Source/Developer/TraceServices/Private/Model/ModuleProvider.cpp b/Engine/Source/Developer/TraceServices/Private/Model/ModuleProvider.cpp
--- a/Engine/Source/Developer/TraceServices/Private/Model/ModuleProvider.cpp
+++ b/Engine/Source/Developer/TraceServices/Private/Model/ModuleProvider.cpp
@@ -2,11 +2,14 @@
 
 #include "ModuleProvider.h"
 #include <atomic>
+
+#include "Algo/Find.h"
 #include "Algo/Transform.h"
 #include "Common/CachedPagedArray.h"
 #include "Common/CachedStringStore.h"
 #include "Common/Utils.h"
 #include "Containers/Map.h"
+#include "HAL/PlatformFileManager.h"
 #include "Misc/PathViews.h"
 #include "TraceServices/Model/AnalysisCache.h"
 #include "TraceServices/Model/AnalysisSession.h"
@@ -26,28 +29,6 @@
 
 namespace TraceServices {
 
-/////////////////////////////////////////////////////////////////////
-struct FModuleEntry
-{
-	const TCHAR*		Name;
-	const TCHAR*		Path;
-	uint64				Base;
-	uint32				Size;
-	bool				bLoaded;
-	std::atomic<bool>	bReady;
-
-	FModuleEntry(const TCHAR* InName, const TCHAR* InPath, uint64 InBase, uint32 InSize, bool InLoaded)
-		: Name(InName)
-		, Path(InPath)
-		, Base(InBase)
-		, Size(InSize)
-		, bLoaded(InLoaded)
-		, bReady(false)
-	{
-
-	}
-};
-
 /////////////////////////////////////////////////////////////////////
 template<typename SymbolResolverType>
 class TModuleProvider : public IModuleAnalysisProvider
@@ -58,6 +39,10 @@ public:
 
 	//Query interface
 	const FResolvedSymbol*		GetSymbol(uint64 Address) override;
+	uint32						GetNumModules() const override;
+	void						EnumerateModules(uint32 Start, TFunctionRef<void(const FModule& Module)> Callback) const override;
+	FGraphEventRef				LoadSymbolsForModuleUsingPath(uint64 Base, const TCHAR* Path) override;
+	void						EnumerateSymbolSearchPaths(TFunctionRef<void(FStringView Path)> Callback) const override;
 	void						GetStats(FStats* OutStats) const override;
 
 	//Analysis interface
@@ -68,6 +53,7 @@ public:
 	void						LoadSymbolsFromCache(IAnalysisCache& Cache);
 
 private:
+	uint32						GetNumCachedSymbolsFromModule(uint64 Base, uint32 Size);
 
 	struct FSavedSymbol
 	{
@@ -78,10 +64,10 @@ private:
 		uint32 Line;
 	};
 	
-	FRWLock						ModulesLock;
-	TPagedArray<FModuleEntry>	Modules;
+	mutable FRWLock				ModulesLock;
+	TPagedArray<FModule>		Modules;
 
-	FRWLock						SymbolsLock;
+	FRWLock				SymbolsLock;
 	// Persistently stored symbol strings
 	FCachedStringStore Strings;
 	// Efficient representation of symbols
@@ -90,10 +76,14 @@ private:
 	TMap<uint64, const FResolvedSymbol*> SymbolCacheLookup;
 	// Number of cached symbols that was loaded (used for stats)
 	uint32 NumCachedSymbols;
+	// Number of discovered symbols
+	std::atomic<uint32> SymbolsDiscovered;
 	
 	IAnalysisSession&			Session;
 	FString						Platform;
 	TUniquePtr<SymbolResolverType>	Resolver;
+	FGraphEventRef				LoadSymbolsTask;
+	bool						LoadSymbolsAbort = false;
 };
 
 /////////////////////////////////////////////////////////////////////
@@ -113,6 +103,11 @@ TModuleProvider<SymbolResolverType>::TModuleProvider(IAnalysisSession& Session)
 template <typename SymbolResolverType>
 TModuleProvider<SymbolResolverType>::~TModuleProvider()
 {
+	if (LoadSymbolsTask)
+	{
+		LoadSymbolsAbort = true;
+		LoadSymbolsTask->Wait();
+	}
 	// Delete the resolver in order to flush all pending resolves
 	Resolver.Reset();
 	SaveSymbolsToCache(Session.GetCache());
@@ -141,6 +136,7 @@ const FResolvedSymbol* TModuleProvider<SymbolResolverType>::GetSymbol(uint64 Add
 		}
 		ResolvedSymbol = &SymbolCache.EmplaceBack(ESymbolQueryResult::Pending, nullptr, nullptr, nullptr, 0);
 		SymbolCacheLookup.Add(Address, ResolvedSymbol);
+		++SymbolsDiscovered;
 	}
 
 	// If not in cache yet, queue it up in the resolver
@@ -150,13 +146,119 @@ const FResolvedSymbol* TModuleProvider<SymbolResolverType>::GetSymbol(uint64 Add
 	return ResolvedSymbol;
 }
 
+/////////////////////////////////////////////////////////////////////
+template <typename SymbolResolverType>
+uint32 TModuleProvider<SymbolResolverType>::GetNumModules() const
+{
+	FReadScopeLock _(ModulesLock);
+	return Modules.Num();
+}
+
+/////////////////////////////////////////////////////////////////////
+template <typename SymbolResolverType>
+void TModuleProvider<SymbolResolverType>::EnumerateModules(uint32 Start, TFunctionRef<void(const FModule& Module)> Callback) const 
+{
+	FReadScopeLock _(ModulesLock);
+	for (uint32 i = Start; i < Modules.Num(); ++i)
+	{
+		Callback(Modules[i]);
+	}
+}
+
+/////////////////////////////////////////////////////////////////////
+template <typename SymbolResolverType>
+FGraphEventRef TModuleProvider<SymbolResolverType>::LoadSymbolsForModuleUsingPath(uint64 Base, const TCHAR* Path)
+{
+	FReadScopeLock _(ModulesLock);
+	const FModule* Module = Algo::FindBy(Modules, Base, &FModule::Base);
+	if (Module)
+	{
+		const FString FullPath = FPaths::ConvertRelativePathToFull(Path);
+		if (Resolver && !FullPath.IsEmpty() && Module->Status.load() != EModuleStatus::Loaded)
+		{
+			// Setup a task to queue and watch the queued module. If it succeeds in resolving with the new path
+			// re-resolve any cached symbols 
+			LoadSymbolsTask = FFunctionGraphTask::CreateAndDispatchWhenReady([this, Module, FullPath]()
+			{
+				auto ReloadModuleFn = [this] (const FModule* InModule, const TCHAR* InPath)
+				{
+					const uint32 DiscoveredSymbols = InModule->Stats.Discovered.load();
+					const uint64 ModuleBegin = InModule->Base;
+					const uint64 ModuleEnd = InModule->Base + InModule->Size;
+
+					auto ReresolveOnSuccess = [this, DiscoveredSymbols, ModuleBegin, ModuleEnd] (TArray<TTuple<uint64,FResolvedSymbol*>>& OutSymbols)
+					{
+						OutSymbols.Reserve(DiscoveredSymbols);
+					
+						FReadScopeLock _(SymbolsLock);
+						for (auto Pair : SymbolCacheLookup)
+						{
+							const uint64 Address = Pair.Get<0>();
+							FResolvedSymbol* Symbol = const_cast<FResolvedSymbol*>(Pair.Get<1>());
+							if (FMath::IsWithin(Pair.Get<0>(), ModuleBegin, ModuleEnd))
+							{
+								OutSymbols.Add(TTuple<uint64,FResolvedSymbol*>(Address, Symbol));
+							}
+						}
+					};
+
+					Resolver->QueueModuleReload(InModule, InPath, ReresolveOnSuccess);
+					// Wait for the resolver to do it's work
+					while (InModule->Status.load() == EModuleStatus::Pending)
+					{
+						FPlatformProcess::Sleep(.1);
+					}
+
+					return InModule->Status.load();
+				};
+			
+				UE_LOG(LogTraceServices, Display, TEXT("Queing symbol loading using path %s."), *FullPath);
+
+				// Load the requested module
+				const EModuleStatus Result = ReloadModuleFn(Module, *FullPath);
+
+				if (Result == EModuleStatus::Loaded)
+				{
+					// Queue up any other failed module using the directory.
+					IPlatformFile* PlatformFile = &FPlatformFileManager::Get().GetPlatformFile();
+					const FString Directory = PlatformFile->DirectoryExists(*FullPath) ? FullPath : FPaths::GetPath(FullPath);
+					for (auto& OtherModule : Modules)
+					{
+						if (LoadSymbolsAbort)
+						{
+							return;
+						}
+						const EModuleStatus ModuleStatus = OtherModule.Status.load();
+						if (&OtherModule != Module && ModuleStatus != EModuleStatus::Loaded && ModuleStatus != EModuleStatus::Pending)
+						{
+							ReloadModuleFn(&OtherModule, *Directory);
+						}
+					}
+				}
+
+				UE_LOG(LogTraceServices, Display, TEXT("Loading symbols for path %s complete."), *FullPath);
+			});
+
+			return LoadSymbolsTask;
+		}
+	}
+	return FGraphEventRef();
+}
+
+/////////////////////////////////////////////////////////////////////
+template <typename SymbolResolverType>
+void TModuleProvider<SymbolResolverType>::EnumerateSymbolSearchPaths(TFunctionRef<void(FStringView Path)> Callback) const
+{
+	Resolver->EnumerateSymbolSearchPaths(Callback);
+}
+
 /////////////////////////////////////////////////////////////////////
 template<typename SymbolResolverType>
 void TModuleProvider<SymbolResolverType>::GetStats(FStats* OutStats) const
 {
 	Resolver->GetStats(OutStats);
+	OutStats->SymbolsDiscovered = SymbolsDiscovered.load();
 	// Add the cached symbols (the resolver doesn't know about them)
-	OutStats->SymbolsDiscovered += NumCachedSymbols;
 	OutStats->SymbolsResolved += NumCachedSymbols;
 }
 
@@ -169,15 +271,18 @@ void TModuleProvider<SymbolResolverType>::OnModuleLoad(const FStringView& Module
 		return;
 	}
 
-	const TCHAR* Path = Session.StoreString(Module);
-	const TCHAR* Name = Session.StoreString(Module);
+	const FStringView NameTemp = FPathViews::GetCleanFilename(Module);
+	const TCHAR* Name = Session.StoreString(NameTemp);
+	const TCHAR* FullName = Session.StoreString(Module);
 
-	{
-		FWriteScopeLock _(ModulesLock);
-		Modules.EmplaceBack(Name, Path, Base, Size, false);
-	}
+	FWriteScopeLock _(ModulesLock);
+	FModule& NewModule = Modules.EmplaceBack(Name, FullName, Base, Size, EModuleStatus::Pending);
 
-	Resolver->QueueModuleLoad(Module, Base, Size, ImageId, ImageIdSize);
+	// The number of cached symbols for this module is equal to the number
+	// of matching addresses in the cache.
+	NewModule.Stats.Cached = GetNumCachedSymbolsFromModule(Base, Size);
+	
+	Resolver->QueueModuleLoad(ImageId, ImageIdSize, &NewModule);
 }
 
 /////////////////////////////////////////////////////////////////////
@@ -226,6 +331,7 @@ void TModuleProvider<SymbolResolverType>::SaveSymbolsToCache(IAnalysisCache& Cac
 	UE_LOG(LogTraceServices, Display, TEXT("Added %d symbols to the %d previously saved symbols."), NumSavedSymbols, NumPreviouslySavedSymbols);
 }
 
+/////////////////////////////////////////////////////////////////////
 template <typename SymbolResolverType>
 void TModuleProvider<SymbolResolverType>::LoadSymbolsFromCache(IAnalysisCache& Cache)
 {
@@ -249,6 +355,25 @@ void TModuleProvider<SymbolResolverType>::LoadSymbolsFromCache(IAnalysisCache& C
 	UE_LOG(LogTraceServices, Display, TEXT("Loaded %d symbols from cache."), SymbolCacheLookup.Num());
 }
 
+/////////////////////////////////////////////////////////////////////
+template <typename SymbolResolverType>
+uint32 TModuleProvider<SymbolResolverType>::GetNumCachedSymbolsFromModule(uint64 Base, uint32 Size)
+{
+	uint32 Count(0);
+	const uint64 Start = Base;
+	const uint64 End = Base + Size;
+	FWriteScopeLock _(SymbolsLock);
+	for (auto& AddressSymbolPair : SymbolCacheLookup)
+	{
+		const uint64 Address = AddressSymbolPair.Get<0>();
+		if (FMath::IsWithin(Address, Start, End))
+		{
+			++Count;
+		}
+	}
+	return Count;
+}
+
 /////////////////////////////////////////////////////////////////////
 IModuleAnalysisProvider* CreateModuleProvider(IAnalysisSession& InSession, const FAnsiStringView& InSymbolFormat)
 {
diff --git a/Engine/Source/Developer/TraceServices/Private/Model/SymslibResolver.cpp b/Engine/Source/Developer/TraceServices/Private/Model/SymslibResolver.cpp
--- a/Engine/Source/Developer/TraceServices/Private/Model/SymslibResolver.cpp
+++ b/Engine/Source/Developer/TraceServices/Private/Model/SymslibResolver.cpp
@@ -5,6 +5,7 @@
 #include "SymslibResolver.h"
 #include "Algo/ForEach.h"
 #include "Algo/Sort.h"
+#include "Algo/Transform.h"
 #include "Async/MappedFileHandle.h"
 #include "Async/ParallelFor.h"
 #include "Containers/StringView.h"
@@ -19,6 +20,7 @@
 #include "Misc/ScopeRWLock.h"
 #include "Misc/StringBuilder.h"
 #include "TraceServices/Model/AnalysisSession.h"
+#include "TraceServices/Model/Diagnostics.h"
 #include <atomic>
 
 /////////////////////////////////////////////////////////////////////
@@ -58,37 +60,129 @@ namespace
 		}
 	};
 
-	static FString FindSymbolFile(const FString& Path)
+	static FString FindBinaryFileInPath(const FString& File, const FString& SearchPath)
 	{
 		IPlatformFile* PlatformFile = &FPlatformFileManager::Get().GetPlatformFile();
 
+		// If the search path is an absolute path to a file use this. Filenames does
+		// not need to to match (e.g. eboot.bin <-> gamename.self)
+		if (PlatformFile->FileExists(*SearchPath))
+		{
+			return SearchPath;
+		}
+
+		
+		// Extract only filename part in case Path is absolute path
+		FString FileName = FPaths::GetCleanFilename(File);
+		
+		// Look in the search path 
+		FString Result = FPaths::Combine(SearchPath, FileName);
+		if (PlatformFile->FileExists(*Result))
+		{
+			return Result;
+		}
+		
+		// In case File is relative
+		if (FPaths::IsRelative(SearchPath))
+		{
+			Result = FPaths::Combine(SearchPath, File);
+			if (PlatformFile->FileExists(*Result))
+			{
+				return Result;
+			}
+		}
+
+		// Finally try the file path itself
+		if (PlatformFile->FileExists(*File))
+		{
+			return File;
+		}
+
+
+		// File not found
+		return FString();
+	}
+
+	static FString FindSymbolFileInPath(const FString& File, const FString& SearchPath)
+	{
+		IPlatformFile* PlatformFile = &FPlatformFileManager::Get().GetPlatformFile();
+		FString SearchPathBase;
+
+		// If search path is an absolute path to a file
+		if (PlatformFile->FileExists(*SearchPath))
+		{
+			// For symbol files, the filename do need to match.
+			if (FPaths::GetCleanFilename(SearchPath) == FPaths::GetCleanFilename(File))
+			{
+				return SearchPath;
+			}
+
+			// Strip the filename from the search path
+			SearchPathBase = FPaths::GetPath(SearchPath);
+		}
+		else
+		{
+			SearchPathBase = SearchPath;
+		}
+		
 		// Extract only filename part in case Path is absolute path
-		FString FileName = FPaths::GetCleanFilename(Path);
+		FString FileName = FPaths::GetCleanFilename(File);
+		
+		// Look in the search path 
+		FString Result = FPaths::Combine(SearchPathBase, FileName);
+		if (PlatformFile->FileExists(*Result))
+		{
+			return Result;
+		}
+
+		// In case File is relative
+		Result = FPaths::Combine(SearchPathBase, File);
+		if (PlatformFile->FileExists(*Result))
+		{
+			return Result;
+		}
 
-		FString SymbolPath = FPlatformMisc::GetEnvironmentVariable(TEXT("UE_INSIGHTS_SYMBOL_PATH"));
+		// File not found
+		return FString();
+	}
 
-		FString SymbolPathPart;
-		while (SymbolPath.Split(TEXT(";"), &SymbolPathPart, &SymbolPath))
+	// Find a file assuming that search path is a root engine folder.
+	static FString FindFileInEngineFolder(const FString& FilePath, const FString& SearchPath, const FString& Platform, const FString& Project)
+	{
+		IPlatformFile* PlatformFile = &FPlatformFileManager::Get().GetPlatformFile();
+		const TCHAR* EngineFolders[] = {TEXT("Engine"), TEXT("Plugins")};
+		FString Result;
+
+		// If this file name has been renamed try mapping it to known
+		if (FilePath.Compare(TEXT("eboot.bin")) == 0)
 		{
-			FString Result = FPaths::Combine(SymbolPathPart, FileName);
+			Result = FPaths::Combine(SearchPath, Project, TEXT("Binaries"), Platform, Project, TEXT(".self"));
 			if (PlatformFile->FileExists(*Result))
 			{
-				// File found in one of symbol paths
 				return Result;
 			}
 		}
-
-		if (!SymbolPath.IsEmpty())
+		
+		const FString File = FPaths::GetCleanFilename(FilePath);
+		
+		// Look in some well known directories
+		for (const TCHAR* Folder : EngineFolders)
 		{
-			FString Result = FPaths::Combine(SymbolPath, FileName);
+			Result = FPaths::Combine(SearchPath, Folder, TEXT("Binaries"), Platform, File);
 			if (PlatformFile->FileExists(*Result))
 			{
-				// File found in symbol path
 				return Result;
 			}
+			
+		}
+
+		// Look in project directory
+		Result = FPaths::Combine(SearchPath, Project, TEXT("Binaries"), Platform, File);
+		if (PlatformFile->FileExists(*Result))
+		{
+			return Result;
 		}
 
-		// File not found
 		return FString();
 	}
 
@@ -96,73 +190,51 @@ namespace
 	// for more information see: https://docs.microsoft.com/en-us/windows-hardware/drivers/debugger/advanced-symsrv-use
 	// to explicitly download symbol from MS Symbol Server for specific dll file, run:
 	// "C:\Program Files (x86)\Windows Kits\10\Debuggers\x64\symchk.exe" /r C:\Windows\system32\kernel32.dll /s srv*C:\Symbols*https://msdl.microsoft.com/download/symbols
-	static FString FindWindowsSymbolFile(const FString& GuidPath)
+	static FString FindWindowsSymbolFile(const FString& GuidPath, const FString& SearchPath)
 	{
-		FString SymbolPath = FPlatformMisc::GetEnvironmentVariable(TEXT("_NT_SYMBOL_PATH"));
-		if (SymbolPath.IsEmpty())
-		{
-			return FString();
-		}
-
 		IPlatformFile* PlatformFile = &FPlatformFileManager::Get().GetPlatformFile();
 
-		TArray<FString> Parts;
-		SymbolPath.ParseIntoArray(Parts, TEXT(";"));
-		for (const FString& Part : Parts)
+		if (SearchPath.StartsWith(TEXT("srv*")) || SearchPath.StartsWith(TEXT("cache*")))
 		{
-			if (Part.StartsWith(TEXT("srv*")))
+			TArray<FString> Srv;
+			SearchPath.RightChop(4).ParseIntoArray(Srv, TEXT("*"));
+			for (const FString& Path : Srv)
 			{
-				TArray<FString> Srv;
-				Part.RightChop(4).ParseIntoArray(Srv, TEXT("*"));
-				for (const FString& Path : Srv)
-				{
-					FString FilePath = FPaths::Combine(Path, GuidPath);
-					if (PlatformFile->FileExists(*FilePath))
-					{
-						return FilePath;
-					}
-				}
-			}
-			else if (!Part.Contains(TEXT("*")))
-			{
-				FString FilePath = FPaths::Combine(Part, GuidPath);
+				FString FilePath = FPaths::Combine(Path, GuidPath);
 				if (PlatformFile->FileExists(*FilePath))
 				{
 					return FilePath;
 				}
 			}
-		;}
+		}
+		else if (SearchPath.Contains(TEXT("*")))
+		{
+			FString FilePath = FPaths::Combine(SearchPath, GuidPath);
+			if (PlatformFile->FileExists(*FilePath))
+			{
+				return FilePath;
+			}
+		}
 
 		return FString();
 	}
 
-	static bool LoadBinary(const TCHAR* Path, SYMS_Arena* Arena, SYMS_ParseBundle& Bundle, TArray<FAutoMappedFile>& Files)
+	static bool LoadBinary(const TCHAR* Path, SYMS_Arena* Arena, SYMS_ParseBundle& Bundle, TArray<FAutoMappedFile>& Files, const FString& SearchPath, const FString& Platform, const FString& Project)
 	{
-		FString FilePath = Path;
-
+		const FString FilePath = Path;
 		bool bFileFound = false;
-		IPlatformFile* PlatformFile = &FPlatformFileManager::Get().GetPlatformFile();
 
 		// First lookup file in symbol path
-		FString BinaryPath = FindSymbolFile(FilePath);
-		if (!BinaryPath.IsEmpty())
+		FString BinaryPath = FindBinaryFileInPath(FilePath, SearchPath);
+		if (BinaryPath.IsEmpty() && !Platform.IsEmpty())
 		{
-			bFileFound = true;
+			BinaryPath = FindFileInEngineFolder(FilePath, SearchPath, Platform, Project);
 		}
-
-		if (!bFileFound)
-		{
-			// If file is absolute path, try that
-			if (!FPathViews::IsRelativePath(*FilePath) && PlatformFile->FileExists(*FilePath))
-			{
-				BinaryPath = FilePath;
-				bFileFound = true;
-			}
-		}
-
+		bFileFound = !BinaryPath.IsEmpty();
+		
 		if (!bFileFound)
 		{
-			UE_LOG(LogSymslib, Warning, TEXT("Binary file '%s' not found"), *FilePath);
+			UE_LOG(LogSymslib, Verbose, TEXT("Binary file '%s' not found"), *FilePath);
 			return false;
 		}
 
@@ -189,8 +261,8 @@ namespace
 
 		return true;
 	}
-
-	static bool LoadDebug(SYMS_Arena* Arena, SYMS_ParseBundle& Bundle, TArray<FAutoMappedFile>& Files)
+	
+	static bool LoadDebug(SYMS_Arena* Arena, SYMS_ParseBundle& Bundle, TArray<FAutoMappedFile>& Files, const FString& SearchPath)
 	{
 		if (syms_bin_is_dbg(Bundle.bin))
 		{
@@ -213,10 +285,9 @@ namespace
 		FString FilePath = ANSI_TO_TCHAR(reinterpret_cast<char*>(ExtFile.file_name.str));
 
 		bool bFileFound = false;
-		IPlatformFile* PlatformFile = &FPlatformFileManager::Get().GetPlatformFile();
 
 		// First try file in symbol path
-		FString DebugPath = FindSymbolFile(FilePath);
+		FString DebugPath = FindSymbolFileInPath(FilePath, SearchPath);
 		if (!DebugPath.IsEmpty())
 		{
 			bFileFound = true;
@@ -225,7 +296,7 @@ namespace
 		if (!bFileFound)
 		{
 			// if executable is PE format, try Windows symbol path format with guid in path
-			if (Bundle.bin->format == SYMS_FileFormat_PE)
+			if (Bundle.bin->format == SYMS_FileFormat_PE && SearchPath.Contains(TEXT("*")))
 			{
 				SYMS_PeBinAccel* PeAccel = reinterpret_cast<SYMS_PeBinAccel*>(Bundle.bin);
 
@@ -240,7 +311,7 @@ namespace
 					Age,
 					*FileName);
 
-				DebugPath = FindWindowsSymbolFile(GuidPath);
+				DebugPath = FindWindowsSymbolFile(GuidPath, SearchPath);
 				if (!DebugPath.IsEmpty())
 				{
 					bFileFound = true;
@@ -250,31 +321,10 @@ namespace
 
 		if (!bFileFound)
 		{
-			// Try direct path if it is absolute
-			if (!FPathViews::IsRelativePath(*FilePath) && PlatformFile->FileExists(*FilePath))
-			{
-				DebugPath = FilePath;
-				bFileFound = true;
-			}
-		}
-
-		if (!bFileFound)
-		{
-			// Try relative path to module binary file that was found in LoadBinary function
-			FString ModuleBasePath = FPaths::GetPath(Files[0].FilePath);
-			DebugPath = FPaths::Combine(ModuleBasePath, FPaths::GetCleanFilename(FilePath));
-			if (PlatformFile->FileExists(*DebugPath))
-			{
-				bFileFound = true;
-			}
-		}
-
-		if (!bFileFound)
-		{
-			UE_LOG(LogSymslib, Warning, TEXT("Debug file '%s' not found"), *FilePath);
+			UE_LOG(LogSymslib, Verbose, TEXT("Debug file '%s' not found"), *FilePath);
 			return false;
 		}
-
+		
 		FAutoMappedFile& File = Files.AddDefaulted_GetRef();
 		if (!File.Load(*DebugPath))
 		{
@@ -318,7 +368,7 @@ namespace
 		}
 		else if (DataParsed.bin->format == SYMS_FileFormat_ELF)
 		{
-			// try different ways of geting build id from elf binary
+			// try different ways of getting build id from elf binary
 			SYMS_String8 FoundId = { 0, 0 };
 
 			SYMS_String8 Bin = DataParsed.bin_data;
@@ -428,28 +478,77 @@ FSymslibResolver::FSymslibResolver(IAnalysisSession& InSession)
 	, CancelTasks(false)
 	, Session(InSession)
 {
-	UE_LOG(LogSymslib, Log, TEXT("UE_INSIGHTS_SYMBOL_PATH: '%s'"), *FPlatformMisc::GetEnvironmentVariable(TEXT("UE_INSIGHTS_SYMBOL_PATH")));
+	// Setup search paths. The SearchPaths array is a priority stack, which
+	// means paths are searched in reversed order.
+	// 1. Any new paths entered by the user this session
+	// 2. Path of the executable (if available)
+	// 3. Paths from UE_INSIGHTS_SYMBOLPATH
+	// 4. Paths from _NT_SYMBOL_PATH
+	// 5. Paths from the user configuration file
+
+	// Paths from configuration
+	FString SettingsIni;
+	if (FConfigCacheIni::LoadGlobalIniFile(SettingsIni, TEXT("UnrealInsightsSettings")))
+	{
+		GConfig->GetArray(TEXT("Insights.MemoryProfiler"), TEXT("SymbolSearchPaths"), SymbolSearchPaths, SettingsIni);
+	}
+
+	// Paths from environment
+	FString InsightsSymbolPath =  FPlatformMisc::GetEnvironmentVariable(TEXT("UE_INSIGHTS_SYMBOL_PATH"));
+	UE_LOG(LogSymslib, Log, TEXT("UE_INSIGHTS_SYMBOL_PATH: '%s'"), *InsightsSymbolPath);
+	FString SymbolPathPart;
+	while (InsightsSymbolPath.Split(TEXT(";"), &SymbolPathPart, &InsightsSymbolPath))
+	{
+		SymbolSearchPaths.Emplace(SymbolPathPart);
+	}
+
+#if PLATFORM_WINDOWS
+	FString NTSymbolPath = FPlatformMisc::GetEnvironmentVariable(TEXT("_NT_SYMBOL_PATH"));
+	UE_LOG(LogSymslib, Log, TEXT("_NT_SYMBOL_PATH: '%s'"), *NTSymbolPath);
+	while (NTSymbolPath.Split(TEXT(";"), &SymbolPathPart, &NTSymbolPath))
+	{
+		 SymbolSearchPaths.Emplace(SymbolPathPart);
+	}
+#endif
+
+	// Try to get session information
+	// todo: Sadly this will not be available in the constructor or when the first modules are traced, since the
+	//		 session info is traced later in the process startup sequence.
+#if 0
+	FAnalysisSessionReadScope _(Session);
+	const IDiagnosticsProvider* DiagnosticsProvider = ReadDiagnosticsProvider(Session);
+	if (DiagnosticsProvider && DiagnosticsProvider->IsSessionInfoAvailable())
+	{
+		const FSessionInfo Info = DiagnosticsProvider->GetSessionInfo();
+		Platform = Info.Platform;
+		Project = Info.AppName;
+	}
+#endif
 }
 
 FSymslibResolver::~FSymslibResolver()
 {
 	CancelTasks = true;
-	// Wait for cleanup task to finish
-	CleanupTask->Wait();
+	// Wait for any existing cleanup tasks to finish
+	if (CleanupTask)
+	{
+		CleanupTask->Wait();
+	}
+	// Wait for module reload task to finish
+	if (ModuleReloadTask)
+	{
+		ModuleReloadTask->Wait();
+	}
 }
 
-void FSymslibResolver::QueueModuleLoad(const FStringView& ModulePath, uint64 Base, uint32 Size, const uint8* ImageId, uint32 ImageIdSize)
+void FSymslibResolver::QueueModuleLoad(const uint8* ImageId, uint32 ImageIdSize, FModule* Module)
 {
 	FWriteScopeLock _(ModulesLock);
-	const FStringView ModuleName = FPathViews::GetCleanFilename(ModulePath);
-
+	check(!CleanupTask); // No queued addresses should be queued after cleanup task is started
+	
 	// Add module
 	FModuleEntry* Entry = &Modules.PushBack();
-	Entry->Base = Base;
-	Entry->Size = Size;
-	Entry->Name = Session.StoreString(ModuleName); //Allows safe sharing to UI
-	Entry->Path = Session.StoreString(ModulePath);
-	Entry->Status = EModuleStatus::Pending;
+	Entry->Module = Module;
 	Entry->ImageId = TArrayView<const uint8>(ImageId, ImageIdSize);
 
 	++ModulesDiscovered;
@@ -460,43 +559,126 @@ void FSymslibResolver::QueueModuleLoad(const FStringView& ModulePath, uint64 Bas
 	FFunctionGraphTask::CreateAndDispatchWhenReady(
 		[this, Entry]
 		{
-			LoadModuleTracked(Entry);
+			LoadModuleTracked(Entry, FStringView());
 			--TasksInFlight;
 		}, TStatId{}, nullptr, ENamedThreads::AnyBackgroundThreadNormalTask);
-
+	
 	// Sort list according to base address
 	SortedModules.Add(Entry);
-	Algo::Sort(SortedModules, [](const FModuleEntry* Lhs, const FModuleEntry* Rhs) { return Lhs->Base < Rhs->Base; });
+	Algo::Sort(SortedModules, [](const FModuleEntry* Lhs, const FModuleEntry* Rhs) { return Lhs->Module->Base < Rhs->Module->Base; });
+}
+
+void FSymslibResolver::QueueModuleReload(const FModule* Module, const TCHAR* InPath, TFunction<void(SymbolArray&)> ResolveOnSuccess)
+{
+	FWriteScopeLock _(ModulesLock);
+
+	// Find the entry
+	FModuleEntry* Entry = nullptr;
+	for (uint32 Idx = 0; Idx < Modules.Num(); ++Idx)
+	{
+		if (Modules[Idx].Module == Module)
+		{
+			Entry = &Modules[Idx];
+			break;
+		}
+	}
+	
+	if (Entry)
+	{
+		// No use in trying reload already loaded modules
+		EModuleStatus PreviousStatus = Entry->Module->Status.load();
+		if (!Entry->Module->Status.compare_exchange_strong(PreviousStatus, EModuleStatus::Pending))
+		{
+			return;
+		}
+		
+		FString PathStr(InPath);
+		FPaths::NormalizeDirectoryName(PathStr);
+		const TCHAR* Path = Session.StoreString(PathStr);
+
+		// Can only launch one task at a time
+		if (ModuleReloadTask && !ModuleReloadTask->IsComplete())
+		{
+			ModuleReloadTask->Wait();
+		}
+		
+		FFunctionGraphTask::CreateAndDispatchWhenReady([this, Entry, Path, PreviousStatus, ResolveOnSuccess]
+		{
+			FScopeLock CleanupExclusive(&CleanupLock);
+			
+			LoadModuleTracked(Entry, Path);
+
+			if (PreviousStatus >= EModuleStatus::FailedStatusStart)
+			{
+				--ModulesFailed;
+			}
+			
+			if (Entry->Module->Status.load() == EModuleStatus::Loaded)
+			{
+				// If an additional search path was specified and successful, add to the permanent
+				// list of search paths.
+				{
+					FWriteScopeLock _(SymbolSearchPathsLock);
+					//Get the base path if path points to a file 
+					IPlatformFile* PlatformFile = &FPlatformFileManager::Get().GetPlatformFile();
+					const FString Directory = PlatformFile->DirectoryExists(Path) ? Path : FPaths::GetPath(Path);
+					if (!Directory.IsEmpty())
+					{
+						SymbolSearchPaths.AddUnique(Directory);
+					}
+				}
+				
+				// Ask the caller for a list of symbols that should be resolved, now that module is properly loaded and
+				// resolve them immediately.
+				{
+					SymbolArray SymbolsToResolve;
+					FSymbolStringAllocator StringAllocator(Session.GetLinearAllocator(), (2 << 12) / sizeof(TCHAR) );
+					ResolveOnSuccess(SymbolsToResolve);
+					for(TTuple<uint64, FResolvedSymbol*> Pair : SymbolsToResolve)
+					{
+						ResolveSymbolTracked(Pair.Get<0>(), Pair.Get<1>(), StringAllocator);
+					}
+				}
+
+				// Finally if the cleanup task has already been dispatched, no other resolve request will come in and we
+				// can safely release our resources. CleanupTask and a reload task cannot run at the same time due to CleanupLock.
+				if (CleanupTask)
+				{
+					// Need to wait for tasks, since there could be queued symbol resolves for this module waiting
+					// in the queue.
+					WaitForTasks();
+					
+					Algo::ForEach(Entry->Instance.Arenas, syms_arena_release);
+					Entry->Instance.Arenas.Empty();
+				}
+			}
+		}, TStatId{}, nullptr, ENamedThreads::AnyBackgroundThreadNormalTask);
+	}
 }
 
 void FSymslibResolver::QueueSymbolResolve(uint64 Address, FResolvedSymbol* Symbol)
 {
 	FScopeLock _(&SymbolsQueueLock);
+	check(!CleanupTask); // No queued addresses should be queued after cleanup task is started
 	MaybeDispatchQueuedAddresses();
-	++SymbolsDiscovered;
 	ResolveQueue.Add(FQueuedAddress{Address, Symbol});
 }
 
 void FSymslibResolver::OnAnalysisComplete()
 {
+	// Dispatch any remaining requests.
+	{
+		 FScopeLock _(&SymbolsQueueLock);
+		 DispatchQueuedAddresses();
+		 ResolveQueue.Reset();
+	}
+	
 	CleanupTask = FFunctionGraphTask::CreateAndDispatchWhenReady([this]
 	{
-		// Dispatch any remaining requests.
-		{
-			FScopeLock _(&SymbolsQueueLock);
-			DispatchQueuedAddresses();
-			ResolveQueue.Reset();
-		}
+		FScopeLock CleanupExclusive(&CleanupLock);
 
 		// Wait for outstanding batches to complete.
-		uint32 OutstandingTasks = TasksInFlight.load(std::memory_order_acquire);
-		UE_LOG(LogSymslib, Display, TEXT("Waiting for %d outstanding tasks..."), OutstandingTasks);
-		do
-		{
-			OutstandingTasks = TasksInFlight.load(std::memory_order_acquire);
-			FPlatformProcess::Sleep(0.0);
-		}
-		while (OutstandingTasks > 0);
+		WaitForTasks();
 
 		// Release memory used by syms library
 		{
@@ -517,12 +699,23 @@ void FSymslibResolver::OnAnalysisComplete()
 
 void FSymslibResolver::GetStats(IModuleProvider::FStats* OutStats) const
 {
+	FReadScopeLock _(ModulesLock);
+	FMemory::Memzero(*OutStats);
+	for(const FModuleEntry& Entry : Modules)
+	{
+		OutStats->SymbolsDiscovered += Entry.Module->Stats.Discovered.load();
+		OutStats->SymbolsResolved += Entry.Module->Stats.Resolved.load();
+		OutStats->SymbolsFailed += Entry.Module->Stats.Failed.load();
+	}
 	OutStats->ModulesDiscovered = ModulesDiscovered.load();
 	OutStats->ModulesFailed = ModulesFailed.load();
 	OutStats->ModulesLoaded = ModulesLoaded.load();
-	OutStats->SymbolsDiscovered = SymbolsDiscovered.load();
-	OutStats->SymbolsFailed = SymbolsFailed.load();
-	OutStats->SymbolsResolved = SymbolsResolved.load();
+}
+
+void FSymslibResolver::EnumerateSymbolSearchPaths(TFunctionRef<void(FStringView Path)> Callback) const
+{
+	FReadScopeLock _(SymbolSearchPathsLock);
+	Algo::ForEach(SymbolSearchPaths, Callback);
 }
 
 /**
@@ -591,7 +784,7 @@ void FSymslibResolver::ResolveSymbols(TArrayView<FQueuedAddress>& QueuedWork)
 FSymslibResolver::FModuleEntry* FSymslibResolver::GetModuleForAddress(uint64 Address)
 {
 	FReadScopeLock _(ModulesLock);
-	const int32 EntryIdx = Algo::LowerBoundBy(SortedModules, Address, [](const FModuleEntry* Entry) { return Entry->Base; }) - 1;
+	const int32 EntryIdx = Algo::LowerBoundBy(SortedModules, Address, [](const FModuleEntry* Entry) { return Entry->Module->Base; }) - 1;
 	if (EntryIdx < 0 || EntryIdx >= SortedModules.Num())
 	{
 		return nullptr;
@@ -608,26 +801,64 @@ void FSymslibResolver::UpdateResolvedSymbol(FResolvedSymbol* Symbol, ESymbolQuer
 	Symbol->Result.store(Result, std::memory_order_release);
 }
 
-void FSymslibResolver::LoadModuleTracked(FModuleEntry* Module)
+void FSymslibResolver::LoadModuleTracked(FModuleEntry* Entry, FStringView OverrideSearchPath)
 {
-	if (CancelTasks.load(std::memory_order_relaxed))
+	// Build search paths
+	FString ModuleNamePath;
+	TArray<FStringView> SearchPaths;
+	if (OverrideSearchPath.IsEmpty())
 	{
-		return;
+		FReadScopeLock _(SymbolSearchPathsLock);
+		for (const FString& SearchPath : SymbolSearchPaths)
+		{
+			SearchPaths.AddUnique(SearchPath);
+		}
+		
+		ModuleNamePath = FPaths::GetPath(Entry->Module->FullName);
+		FPaths::NormalizeDirectoryName(ModuleNamePath);
+		if (!ModuleNamePath.IsEmpty())
+		{
+			SearchPaths.AddUnique(ModuleNamePath);
+		}
+	}
+	else 
+	{
+		SearchPaths.Add(OverrideSearchPath);
+	}
+
+	// Exhaust all the search paths
+	EModuleStatus Status = EModuleStatus::Pending;
+	TStringBuilder<128> StatusMessage;
+	while(Status != EModuleStatus::Loaded && !CancelTasks.load() && SearchPaths.Num() > 0)
+	{
+		StatusMessage.Reset();
+		Status = LoadModule(Entry,SearchPaths.Top(), StatusMessage);
+		UE_LOG(LogSymslib, Display, TEXT("%s"), StatusMessage.ToString());
+		SearchPaths.Pop();
+	}
+
+	// If we are still pending after the loop we have failed
+	if (Status == EModuleStatus::Pending)
+	{
+		Status = EModuleStatus::Failed;
 	}
 	
-	const EModuleStatus Status = LoadModule(Module);
 	if (Status == EModuleStatus::Loaded)
 	{
+		// Reset this modules stats
+		Entry->Module->Stats.Discovered.store(0u);
+		Entry->Module->Stats.Resolved.store(0u);
+		Entry->Module->Stats.Failed.store(0u);
 		++ModulesLoaded;
 	}
 	else
 	{
 		++ModulesFailed;
-		Module->Status.store(Status);
 	}
-
-	//  Make the status visible
-	Module->Status.store(Status);
+	
+	// Make the final status visible to the world
+	Entry->Module->StatusMessage = Session.StoreString(StatusMessage.ToView());
+	Entry->Module->Status.store(Status);
 
 	// If this is the last module in flight dispatch pending queries
 	{
@@ -636,8 +867,9 @@ void FSymslibResolver::LoadModuleTracked(FModuleEntry* Module)
 	}
 }
 
-FSymslibResolver::EModuleStatus FSymslibResolver::LoadModule(FModuleEntry* Module)
+EModuleStatus FSymslibResolver::LoadModule(FModuleEntry* Entry, FStringView SearchPathView, FStringBuilderBase& OutStatusMessage) const
 {
+	const FString SearchPath(SearchPathView);
 	// temporary memory used for loading
 	SYMS_Group* Group = syms_group_alloc();
 
@@ -647,27 +879,29 @@ FSymslibResolver::EModuleStatus FSymslibResolver::LoadModule(FModuleEntry* Modul
 	// contents of binary & debug file
 	SYMS_ParseBundle Bundle;
 
-	if (!LoadBinary(Module->Path, Group->arena, Bundle, Files))
+	if (!LoadBinary(Entry->Module->FullName, Group->arena, Bundle, Files, SearchPath, Platform, Project))
 	{
 		syms_group_release(Group);
+		OutStatusMessage.Appendf(TEXT("Failed to load image for '%s' in '%s'."), Entry->Module->Name, *SearchPath);
 		return EModuleStatus::Failed;
 	}
 
-	if (!LoadDebug(Group->arena, Bundle, Files))
+	if (!LoadDebug(Group->arena, Bundle, Files, SearchPath))
 	{
 		syms_group_release(Group);
+		OutStatusMessage.Appendf(TEXT("Failed to load symbols for '%s' in '%s'."), Entry->Module->Name, *SearchPath);
 		return EModuleStatus::Failed;
 	}
 
 	// check debug data mismatch to captured module
-	if (!Module->ImageId.IsEmpty() && !MatchImageId(Module->ImageId, Bundle))
+	if (!Entry->ImageId.IsEmpty() && !MatchImageId(Entry->ImageId, Bundle))
 	{
 		syms_group_release(Group);
-		UE_LOG(LogSymslib, Warning, TEXT("Symbols for '%s' does not match traced binary."), Module->Name);
+		OutStatusMessage << TEXT("Symbols for '") << Entry->Module->Name << TEXT("' found in '") << SearchPath << TEXT("' does not match trace binary.");
 		return EModuleStatus::VersionMismatch;
 	}
 
-	FSymsInstance* Instance = &Module->Instance;
+	FSymsInstance* Instance = &Entry->Instance;
 
 	// initialize group
 	syms_set_lane(0);
@@ -791,47 +1025,56 @@ FSymslibResolver::EModuleStatus FSymslibResolver::LoadModule(FModuleEntry* Modul
 	syms_group_release(Group);
 	Instance->Arenas.RemoveAll([](SYMS_Arena* arena) { return arena == nullptr; });
 
-	UE_LOG(LogSymslib, Display, TEXT("Loaded symbols for '%s' at base 0x%016llx, %u symbols."), Module->Name, Module->Base, SymbolCount.load());
-
-	return EModuleStatus::Loaded;
+	// In order to avoid listing all files, use only the last. On dwarf/elf this will be the main binary, on Windows
+	// this will be the pdb (the first is the exe/dll).
+	OutStatusMessage.Appendf(TEXT("%s loaded with %u symbols."), *Files.Last().FilePath, SymbolCount.load());
+	return EModuleStatus::Loaded; 
 }
 
 void FSymslibResolver::ResolveSymbolTracked(uint64 Address, FResolvedSymbol* Target, FSymbolStringAllocator& StringAllocator)
 {
-	if (!ResolveSymbol(Address, Target, StringAllocator))
+	const ESymbolQueryResult PreviousResult = Target->Result.load();
+	if (PreviousResult == ESymbolQueryResult::OK)
 	{
-		++SymbolsFailed;
+		return;
 	}
-	else
+	
+	FModuleEntry* Entry = GetModuleForAddress(Address);
+	if (!Entry)
 	{
-		++SymbolsResolved;
+		UE_LOG(LogSymslib, Warning, TEXT("No module mapped to address 0x%016llx."), Address);
+		UpdateResolvedSymbol(Target, ESymbolQueryResult::NotLoaded, GUnknownModuleTextSymsLib, GUnknownModuleTextSymsLib, GUnknownModuleTextSymsLib, 0);
+		return;
 	}
-}
 
-bool FSymslibResolver::ResolveSymbol(uint64 Address, FResolvedSymbol* Target, FSymbolStringAllocator& StringAllocator)
-{
-	FModuleEntry* Module = GetModuleForAddress(Address);
-	if (!Module)
+	FModule::SymbolStats& ModuleStats = Entry->Module->Stats;
+	++ModuleStats.Discovered;
+	if (!ResolveSymbol(Address, Target, StringAllocator, Entry))
 	{
-		UE_LOG(LogSymslib, Warning, TEXT("No module mapped to address 0x%016llx."), Address);
-		UpdateResolvedSymbol(Target, ESymbolQueryResult::NotLoaded, GUnknownModuleTextSymsLib, GUnknownModuleTextSymsLib, GUnknownModuleTextSymsLib, 0);
-		return false;
+		++ModuleStats.Failed;
 	}
+	else
+	{
+		++ModuleStats.Resolved;
+	}
+}
 
-	EModuleStatus Status = Module->Status.load();
+bool FSymslibResolver::ResolveSymbol(uint64 Address, FResolvedSymbol* Target, FSymbolStringAllocator& StringAllocator, FModuleEntry* Entry) const
+{
+	EModuleStatus Status = Entry->Module->Status.load();
 	while (Status == EModuleStatus::Pending)
 	{
 		//todo: yield
-		Status = Module->Status.load();
+		Status = Entry->Module->Status.load();
 	}
 
 	switch (Status)
 	{
 	case EModuleStatus::Failed:
-		UpdateResolvedSymbol(Target, ESymbolQueryResult::NotLoaded, Module->Name, GUnknownModuleTextSymsLib, GUnknownModuleTextSymsLib, 0);
+		UpdateResolvedSymbol(Target, ESymbolQueryResult::NotLoaded, Entry->Module->Name, GUnknownModuleTextSymsLib, GUnknownModuleTextSymsLib, 0);
 		return false;
 	case EModuleStatus::VersionMismatch:
-		UpdateResolvedSymbol(Target, ESymbolQueryResult::Mismatch, Module->Name, GUnknownModuleTextSymsLib, GUnknownModuleTextSymsLib, 0);
+		UpdateResolvedSymbol(Target, ESymbolQueryResult::Mismatch, Entry->Module->Name, GUnknownModuleTextSymsLib, GUnknownModuleTextSymsLib, 0);
 		return false;
 	default:
 		break;
@@ -839,8 +1082,8 @@ bool FSymslibResolver::ResolveSymbol(uint64 Address, FResolvedSymbol* Target, FS
 
 	// Find procedure and source file for address
 
-	FSymsInstance* Instance = &Module->Instance;
-	SYMS_U64 VirtualOffset = Address + Instance->DefaultBase - Module->Base;
+	FSymsInstance* Instance = &Entry->Instance;
+	SYMS_U64 VirtualOffset = Address + Instance->DefaultBase - Entry->Module->Base;
 
 	FSymsSymbol* SymsSymbol = nullptr;
 
@@ -886,14 +1129,14 @@ bool FSymslibResolver::ResolveSymbol(uint64 Address, FResolvedSymbol* Target, FS
 			SymsSymbol = &Instance->StrippedSymbols[Value - 1];
 
 			// use module name as filename
-			SourceFilePersistent = StringAllocator.Store(Module->Name);
+			SourceFilePersistent = StringAllocator.Store(Entry->Module->Name);
 		}
 	}
 
 	// this includes skipping symbols without name (empty string)
 	if (!SymsSymbol || !SourceFilePersistent || SymsSymbol && SymsSymbol->Name[0] == 0)
 	{
-		UpdateResolvedSymbol(Target, ESymbolQueryResult::NotFound, Module->Name, GUnknownModuleTextSymsLib, GUnknownModuleTextSymsLib, 0);
+		UpdateResolvedSymbol(Target, ESymbolQueryResult::NotFound, Entry->Module->Name, GUnknownModuleTextSymsLib, GUnknownModuleTextSymsLib, 0);
 		return false;
 	}
 
@@ -906,14 +1149,25 @@ bool FSymslibResolver::ResolveSymbol(uint64 Address, FResolvedSymbol* Target, FS
 	UpdateResolvedSymbol(
 		Target,
 		ESymbolQueryResult::OK,
-		Module->Name,
+		Entry->Module->Name,
 		SymbolNamePersistent,
 		SourceFilePersistent,
 		SourceFileLine
 	);
-
+		
 	return true;
 }
+	
+void FSymslibResolver::WaitForTasks()
+{
+	uint32 OutstandingTasks = TasksInFlight.load(std::memory_order_acquire);
+	do
+	{
+		 OutstandingTasks = TasksInFlight.load(std::memory_order_acquire);
+		 FPlatformProcess::Sleep(0.0);
+	}
+	while (OutstandingTasks > 0);
+}
 
 /////////////////////////////////////////////////////////////////////
 
diff --git a/Engine/Source/Developer/TraceServices/Private/Model/SymslibResolver.h b/Engine/Source/Developer/TraceServices/Private/Model/SymslibResolver.h
--- a/Engine/Source/Developer/TraceServices/Private/Model/SymslibResolver.h
+++ b/Engine/Source/Developer/TraceServices/Private/Model/SymslibResolver.h
@@ -20,22 +20,18 @@ namespace TraceServices {
 class FSymslibResolver
 {
 public:
+	typedef TArray<TTuple<uint64, FResolvedSymbol*>> SymbolArray;
+	
 	FSymslibResolver(IAnalysisSession& InSession);
 	~FSymslibResolver();
-	void QueueModuleLoad(const FStringView& ModulePath, uint64 Base, uint32 Size, const uint8* ImageId, uint32 ImageIdSize);
+	void QueueModuleLoad(const uint8* ImageId, uint32 ImageIdSize, FModule* Module);
+	void QueueModuleReload(const FModule* Module, const TCHAR* Path, TFunction<void(SymbolArray&)> ResolveOnSuccess);
 	void QueueSymbolResolve(uint64 Address, FResolvedSymbol* Symbol);
 	void OnAnalysisComplete();
 	void GetStats(IModuleProvider::FStats* OutStats) const;
+	void EnumerateSymbolSearchPaths(TFunctionRef<void(FStringView Path)> Callback) const;
 
 private:
-	enum class EModuleStatus : uint8
-	{
-		Pending,
-		Loaded,
-		VersionMismatch,
-		Failed
-	};
-
 	struct FSymsSymbol
 	{
 		const char* Name;
@@ -61,13 +57,9 @@ private:
 
 	struct FModuleEntry
 	{
-		uint64 Base;
-		uint32 Size;
-		const TCHAR* Name;
-		const TCHAR* Path;
 		FSymsInstance Instance;
-		std::atomic<EModuleStatus> Status;
 		TArray<uint8> ImageId;
+		FModule* Module;
 	};
 
 	struct FQueuedAddress
@@ -101,32 +93,37 @@ private:
 	void ResolveSymbols(TArrayView<FQueuedAddress>& QueuedWork);
 	FModuleEntry* GetModuleForAddress(uint64 Address);
 	static void UpdateResolvedSymbol(FResolvedSymbol* Symbol, ESymbolQueryResult Result, const TCHAR* Module, const TCHAR* Name, const TCHAR* File, uint16 Line);
-	void LoadModuleTracked(FModuleEntry* Module);
-	EModuleStatus LoadModule(FModuleEntry* Module);
+	void LoadModuleTracked(FModuleEntry* Entry, FStringView OverrideSearchPath);
+	EModuleStatus LoadModule(FModuleEntry* Entry, FStringView SearchPathView, FStringBuilderBase& OutStatusMessage) const;
 	void ResolveSymbolTracked(uint64 Address, FResolvedSymbol* Target, class FSymbolStringAllocator& StringAllocator);
-	bool ResolveSymbol(uint64 Address, FResolvedSymbol* Target, class FSymbolStringAllocator& StringAllocator);
+	bool ResolveSymbol(uint64 Address, FResolvedSymbol* Target, FSymbolStringAllocator& StringAllocator, FModuleEntry* Entry) const;
+	void WaitForTasks();
 
-	FRWLock ModulesLock;
+	mutable FRWLock ModulesLock;
 	TPagedArray<FModuleEntry> Modules;
 	TArray<FModuleEntry*> SortedModules;
 	FCriticalSection SymbolsQueueLock;
 	TArray<FQueuedAddress, TInlineAllocator<QueuedAddressLength>> ResolveQueue;
 	std::atomic<uint32> TasksInFlight;
+	FCriticalSection CleanupLock;
 	FGraphEventRef CleanupTask;
+	FGraphEventRef ModuleReloadTask;
 	std::atomic<bool> CancelTasks;
+	mutable FRWLock SymbolSearchPathsLock;
+	TArray<FString> SymbolSearchPaths;
 
 	std::atomic<uint32> ModulesDiscovered;
 	std::atomic<uint32> ModulesFailed;
 	std::atomic<uint32> ModulesLoaded;
-	std::atomic<uint32> SymbolsDiscovered;
-	std::atomic<uint32> SymbolsFailed;
-	std::atomic<uint32> SymbolsResolved;
 	std::atomic<uint64> SymbolBytesAllocated;
 	std::atomic<uint64> SymbolBytesWasted;
 
 	IAnalysisSession& Session;
+	FString Platform;
+	FString Project;
 };
 
+
 /////////////////////////////////////////////////////////////////////
 
 } // namespace TraceServices
diff --git a/Engine/Source/Developer/TraceServices/Public/TraceServices/Model/Modules.h b/Engine/Source/Developer/TraceServices/Public/TraceServices/Model/Modules.h
--- a/Engine/Source/Developer/TraceServices/Public/TraceServices/Model/Modules.h
+++ b/Engine/Source/Developer/TraceServices/Public/TraceServices/Model/Modules.h
@@ -67,6 +67,70 @@ struct FResolvedSymbol
 	{}
 };
 
+////////////////////////////////////////////////////////////////////////////////
+enum class EModuleStatus
+{
+	Pending,			// Module is pending load
+	Loaded,				// Module has been successfully loaded
+	VersionMismatch,	// Debug data was found, but did not match traced version
+	NotFound,			// Debug data for module was not found
+	Failed,				// Unable to parse debug information
+	StatusNum,
+	FailedStatusStart = VersionMismatch
+};
+	
+////////////////////////////////////////////////////////////////////////////////
+inline const TCHAR* ModuleStatusToString(EModuleStatus Result)
+{
+	static const TCHAR* DisplayStrings[] = {
+		TEXT("Pending..."),
+		TEXT("Loaded"),
+		TEXT("Version mismatch"),
+		TEXT("Not found"),
+		TEXT("Failed")
+	};
+	static_assert(UE_ARRAY_COUNT(DisplayStrings) == (uint8) EModuleStatus::StatusNum, "Missing QueryResult");
+	return DisplayStrings[(uint8)Result];
+}
+	
+////////////////////////////////////////////////////////////////////////////////
+/**
+ * Represents information about a module (engine/game/system dll or monolithic binary)
+ * and how debug information has been loaded.
+ */
+struct FModule
+{
+	// Name of the module.
+	const TCHAR*				Name;
+	// Full name as reported by the event
+	const TCHAR*				FullName;
+	// Base address
+	uint64						Base;
+	// Size in memory
+	uint32						Size;
+	// Status of loading debug information
+	std::atomic<EModuleStatus>	Status;
+	// If status is loaded, contains the path to the debug info. If failure state contains an error message.
+	const TCHAR*				StatusMessage;
+	// Statistics about the module
+	struct SymbolStats
+	{
+		std::atomic<uint32>		Discovered;
+		std::atomic<uint32>		Cached;
+		std::atomic<uint32>		Resolved;
+		std::atomic<uint32>		Failed;
+	} Stats;
+	
+	FModule(const TCHAR* InName, const TCHAR* InFullName, uint64 InBase, uint32 InSize, EModuleStatus InStatus)
+		: Name(InName)
+		, FullName(InFullName)
+		, Base(InBase)
+		, Size(InSize)
+		, Status(InStatus)
+		, StatusMessage(nullptr)
+	{}
+};
+	
 ////////////////////////////////////////////////////////////////////////////////
 class IModuleProvider : public IProvider
 {
@@ -83,12 +147,43 @@ public:
 
 	virtual ~IModuleProvider() = default;
 
-	/** Queries the name of the symbol at address. This function returns immedately, 
+	/** Queries the name of the symbol at address. This function returns immediately, 
 	 * but the lookup is async. See \ref FResolvedSymbol for details. It assumed that 
 	 * all calls to this function happens before analysis has ended.
 	 */
 	virtual const FResolvedSymbol* GetSymbol(uint64 Address) = 0;
 
+
+	/**
+	 * Gets the number of discovered modules.
+	 * @return Number of modules that were discovered so far.
+	 */
+	virtual uint32 GetNumModules() const = 0;
+
+	/**
+	 * Enumerates all detected modules and their state. Modules are listed in the order they
+	 * were found and is guaranteed to retain it's index in the list and address in memory.
+	 * @param Start Starting index
+	 * @param Callback Issued for each module
+	 */
+	virtual void EnumerateModules(uint32 Start, TFunctionRef<void(const FModule& Module)> Callback) const = 0;
+
+	/**
+	 * Trigger a manual attempt to load symbols for a module given a search path. This operation
+	 * is asynchronous and if successful resolve symbols that previously failed.
+	 * @param Base Base address of the module
+	 * @param Path Path where debug data can be located
+	 * @return Waitable graph event for the operation
+	 */
+	virtual FGraphEventRef LoadSymbolsForModuleUsingPath(uint64 Base, const TCHAR* Path) = 0;
+
+	/**
+	 * Gets the search paths used to find debug symbols. Paths are traversed in the reverse order
+	 * order of importance.
+	 * @param Callback Called for each search path.
+	 */
+	virtual void EnumerateSymbolSearchPaths(TFunctionRef<void(FStringView Path)> Callback) const = 0;
+
 	/** Gets statistics from provider */
 	virtual void GetStats(FStats* OutStats) const = 0;
 };
