Fix static code analysis error

diff --git a/Engine/Source/Developer/TraceServices/Private/Model/DbgHelpResolver.cpp b/Engine/Source/Developer/TraceServices/Private/Model/DbgHelpResolver.cpp
--- a/Engine/Source/Developer/TraceServices/Private/Model/DbgHelpResolver.cpp
+++ b/Engine/Source/Developer/TraceServices/Private/Model/DbgHelpResolver.cpp
@@ -269,7 +269,6 @@ void FDbgHelpResolver::ResolveSymbol(uint64 Address, FResolvedSymbol* Target)
 	const FModuleEntry* Entry = GetModuleForAddress(Address);
 	if (!Entry)
 	{
-		++Entry->Module->Stats.Failed;
 		UpdateResolvedSymbol(Target, ESymbolQueryResult::NotFound, GUnknownModuleTextDbgHelp,
 		                     GUnknownModuleTextDbgHelp, GUnknownModuleTextDbgHelp, 0);
 		return;
