Fix Mac compile issues

diff --git a/Engine/Source/Developer/TraceServices/Private/Model/ModuleProvider.cpp b/Engine/Source/Developer/TraceServices/Private/Model/ModuleProvider.cpp
--- a/Engine/Source/Developer/TraceServices/Private/Model/ModuleProvider.cpp
+++ b/Engine/Source/Developer/TraceServices/Private/Model/ModuleProvider.cpp
@@ -193,9 +193,9 @@ FGraphEventRef TModuleProvider<SymbolResolverType>::LoadSymbolsForModuleUsingPat
 						FReadScopeLock _(SymbolsLock);
 						for (auto Pair : SymbolCacheLookup)
 						{
-							const uint64 Address = Pair.Get<0>();
-							FResolvedSymbol* Symbol = const_cast<FResolvedSymbol*>(Pair.Get<1>());
-							if (FMath::IsWithin(Pair.Get<0>(), ModuleBegin, ModuleEnd))
+							const uint64 Address = Pair.template Get<0>();
+							FResolvedSymbol* Symbol = const_cast<FResolvedSymbol*>(Pair.template Get<1>());
+							if (FMath::IsWithin(Address, ModuleBegin, ModuleEnd))
 							{
 								OutSymbols.Add(TTuple<uint64,FResolvedSymbol*>(Address, Symbol));
 							}
@@ -365,7 +365,7 @@ uint32 TModuleProvider<SymbolResolverType>::GetNumCachedSymbolsFromModule(uint64
 	FWriteScopeLock _(SymbolsLock);
 	for (auto& AddressSymbolPair : SymbolCacheLookup)
 	{
-		const uint64 Address = AddressSymbolPair.Get<0>();
+		const uint64 Address = AddressSymbolPair.template Get<0>();
 		if (FMath::IsWithin(Address, Start, End))
 		{
 			++Count;
