Fix non unity compile errors

diff --git a/Engine/Source/Developer/TraceInsights/Private/Insights/Widgets/SModulesView.cpp b/Engine/Source/Developer/TraceInsights/Private/Insights/Widgets/SModulesView.cpp
--- a/Engine/Source/Developer/TraceInsights/Private/Insights/Widgets/SModulesView.cpp
+++ b/Engine/Source/Developer/TraceInsights/Private/Insights/Widgets/SModulesView.cpp
@@ -2,7 +2,7 @@
 
 #include "SModulesView.h"
 
-//#include "Framework/Application/SlateApplication.h"
+#include "Framework/Application/SlateApplication.h"
 #include "Framework/MultiBox/MultiBoxBuilder.h"
 #include "SlateOptMacros.h"
 #include "Textures/SlateIcon.h"
diff --git a/Engine/Source/Developer/TraceServices/Private/Model/DbgHelpResolver.cpp b/Engine/Source/Developer/TraceServices/Private/Model/DbgHelpResolver.cpp
--- a/Engine/Source/Developer/TraceServices/Private/Model/DbgHelpResolver.cpp
+++ b/Engine/Source/Developer/TraceServices/Private/Model/DbgHelpResolver.cpp
@@ -10,6 +10,8 @@
 #include "HAL/PlatformProcess.h"
 #include "HAL/RunnableThread.h"
 #include "Logging/LogMacros.h"
+#include "Misc/ConfigCacheIni.h"
+#include "Misc/Guid.h"
 #include "Misc/PathViews.h"
 #include "Misc/ScopeLock.h"
 #include "Misc/StringBuilder.h"
diff --git a/Engine/Source/Developer/TraceServices/Private/Model/ModuleProvider.cpp b/Engine/Source/Developer/TraceServices/Private/Model/ModuleProvider.cpp
--- a/Engine/Source/Developer/TraceServices/Private/Model/ModuleProvider.cpp
+++ b/Engine/Source/Developer/TraceServices/Private/Model/ModuleProvider.cpp
@@ -9,7 +9,9 @@
 #include "Common/CachedStringStore.h"
 #include "Common/Utils.h"
 #include "Containers/Map.h"
+#include "GenericPlatform/GenericPlatformFile.h"
 #include "HAL/PlatformFileManager.h"
+#include "Misc/Paths.h"
 #include "Misc/PathViews.h"
 #include "TraceServices/Model/AnalysisCache.h"
 #include "TraceServices/Model/AnalysisSession.h"
diff --git a/Engine/Source/Developer/TraceServices/Private/Model/SymslibResolver.cpp b/Engine/Source/Developer/TraceServices/Private/Model/SymslibResolver.cpp
--- a/Engine/Source/Developer/TraceServices/Private/Model/SymslibResolver.cpp
+++ b/Engine/Source/Developer/TraceServices/Private/Model/SymslibResolver.cpp
@@ -13,6 +13,7 @@
 #include "HAL/PlatformFileManager.h"
 #include "HAL/PlatformTLS.h"
 #include "Logging/LogMacros.h"
+#include "Misc/ConfigCacheIni.h"
 #include "Misc/CString.h"
 #include "Misc/Paths.h"
 #include "Misc/PathViews.h"
diff --git a/Engine/Source/Developer/TraceServices/Public/TraceServices/Model/Modules.h b/Engine/Source/Developer/TraceServices/Public/TraceServices/Model/Modules.h
--- a/Engine/Source/Developer/TraceServices/Public/TraceServices/Model/Modules.h
+++ b/Engine/Source/Developer/TraceServices/Public/TraceServices/Model/Modules.h
@@ -1,6 +1,8 @@
 // Copyright Epic Games, Inc. All Rights Reserved.
 
 #pragma once
+
+#include "Async/TaskGraphInterfaces.h"
 #include "TraceServices/Model/AnalysisSession.h"
 #include "Templates/UnrealTemplate.h"
 #include <atomic>
