[Insights] Fix crash in TaskProfiler when enabling the Task channel late.

diff --git a/Engine/Source/Developer/TraceServices/Private/Model/TasksProfiler.cpp b/Engine/Source/Developer/TraceServices/Private/Model/TasksProfiler.cpp
--- a/Engine/Source/Developer/TraceServices/Private/Model/TasksProfiler.cpp
+++ b/Engine/Source/Developer/TraceServices/Private/Model/TasksProfiler.cpp
@@ -288,7 +288,11 @@ namespace TraceServices
 	void FTasksProvider::WaitingFinished(double Timestamp, uint32 ThreadId)
 	{
 		TArray64<FWaitingForTasks>* Thread = WaitingThreads.Find(ThreadId);
-		checkf(Thread != nullptr, TEXT("%d"), ThreadId);
+		if (Thread == nullptr)
+		{
+			UE_LOG(LogTraceServices, Log, TEXT("WaitingFinished task event (Thread %d, Timestamp %.6f) skipped."), ThreadId, Timestamp);
+			return;
+		}
 
 		Thread->Last().FinishedTimestamp = Timestamp;
 	}
