[Insights] Fix end iterator

diff --git a/Engine/Source/Developer/TraceServices/Public/Common/PagedArray.h b/Engine/Source/Developer/TraceServices/Public/Common/PagedArray.h
--- a/Engine/Source/Developer/TraceServices/Public/Common/PagedArray.h
+++ b/Engine/Source/Developer/TraceServices/Public/Common/PagedArray.h
@@ -407,8 +407,8 @@ public:
 	
 	FORCEINLINE TIterator begin() { return TIterator(*this, 0, 0); }
 	FORCEINLINE TIterator begin() const { return TIterator(*this, 0, 0); }
-	FORCEINLINE TIterator end() { return TIterator(*this, NumPages()-1, uint64(TotalItemCount % PageSize)-1); }
-	FORCEINLINE TIterator end() const { return TIterator(*this, NumPages()-1, uint64(TotalItemCount % PageSize)-1); }
+	FORCEINLINE TIterator end() { return TIterator(*this, NumPages()-1, LastPage->Count - 1); }
+	FORCEINLINE TIterator end() const { return TIterator(*this, NumPages()-1, LastPage->Count - 1); }
 
 private:
 	template<typename ItemType, typename PageType>
