[Insights] Fixed crash occuring in Unreal Insights when trace contains cpu timing events ("CpuProfiler.EventBatch") emitted after "CpuProfiler.EndThread" event.

diff --git a/Engine/Source/Developer/TraceServices/Private/Analyzers/CpuProfilerTraceAnalysis.cpp b/Engine/Source/Developer/TraceServices/Private/Analyzers/CpuProfilerTraceAnalysis.cpp
--- a/Engine/Source/Developer/TraceServices/Private/Analyzers/CpuProfilerTraceAnalysis.cpp
+++ b/Engine/Source/Developer/TraceServices/Private/Analyzers/CpuProfilerTraceAnalysis.cpp
@@ -71,6 +71,7 @@ bool FCpuProfilerAnalyzer::OnEvent(uint16 RouteId, EStyle Style, const FOnEventC
 				TimerName = *Name;
 			}
 		}
+
 		if (TimerName[0] == 0)
 		{
 			Name = FString::Printf(TEXT("<noname %u>"), SpecId);
@@ -90,6 +91,7 @@ bool FCpuProfilerAnalyzer::OnEvent(uint16 RouteId, EStyle Style, const FOnEventC
 		DefineTimer(SpecId, Session.StoreString(TimerName), FileName, Line, bMergeByName);
 		break;
 	}
+
 	case RouteId_EndThread:
 	{
 		uint32 ThreadId = FTraceAnalyzerUtils::GetThreadIdField(Context);
@@ -105,23 +107,34 @@ bool FCpuProfilerAnalyzer::OnEvent(uint16 RouteId, EStyle Style, const FOnEventC
 				ThreadState.Timeline->AppendEndEvent(Timestamp);
 			}
 		}
-		ThreadStatesMap.Remove(ThreadId);
+		ThreadState.LastCycle = ~0ull;
+		break;
 	}
+
 	case RouteId_EventBatch:
 	case RouteId_EndCapture:
 	{
+		const uint32 ThreadId = FTraceAnalyzerUtils::GetThreadIdField(Context);
+		FThreadState& ThreadState = GetThreadState(ThreadId);
+
+		if (ThreadState.LastCycle == ~0ull)
+		{
+			// Ignore timing events received after EndThread.
+			break;
+		}
+
 		TArrayView<const uint8> DataView = FTraceAnalyzerUtils::LegacyAttachmentArray("Data", Context);
 		TotalEventSize += DataView.Num();
-		uint32 BufferSize = DataView.Num();
+		const uint32 BufferSize = DataView.Num();
 		const uint8* BufferPtr = DataView.GetData();
-		uint32 ThreadId = FTraceAnalyzerUtils::GetThreadIdField(Context);
-		if (uint64 LastCycle = ProcessBuffer(Context.EventTime, ThreadId, BufferPtr, BufferSize))
+
+		uint64 LastCycle = ProcessBuffer(Context.EventTime, ThreadState, BufferPtr, BufferSize);
+		if (LastCycle != 0)
 		{
 			double LastTimestamp = Context.EventTime.AsSeconds(LastCycle);
 			Session.UpdateDurationSeconds(LastTimestamp);
 			if (RouteId == RouteId_EndCapture)
 			{
-				FThreadState& ThreadState = GetThreadState(ThreadId);
 				for (int32 i = ThreadState.ScopeStack.Num(); i--;)
 				{
 					ThreadState.ScopeStack.Pop();
@@ -132,17 +145,24 @@ bool FCpuProfilerAnalyzer::OnEvent(uint16 RouteId, EStyle Style, const FOnEventC
 		BytesPerScope = double(TotalEventSize) / double(TotalScopeCount);
 		break;
 	}
+
 	case RouteId_CpuScope:
-		(Style == EStyle::EnterScope) ? OnCpuScopeEnter(Context) : OnCpuScopeLeave(Context);
+		if (Style == EStyle::EnterScope)
+		{
+			OnCpuScopeEnter(Context);
+		}
+		else
+		{
+			OnCpuScopeLeave(Context);
+		}
 		break;
 	}
 
 	return true;
 }
 
-uint64 FCpuProfilerAnalyzer::ProcessBuffer(const FEventTime& EventTime, uint32 ThreadId, const uint8* BufferPtr, uint32 BufferSize)
+uint64 FCpuProfilerAnalyzer::ProcessBuffer(const FEventTime& EventTime, FThreadState& ThreadState, const uint8* BufferPtr, uint32 BufferSize)
 {
-	FThreadState& ThreadState = GetThreadState(ThreadId);
 	uint64 LastCycle = ThreadState.LastCycle;
 
 	int32 RemainingPending = ThreadState.PendingEvents.Num();
@@ -191,16 +211,16 @@ uint64 FCpuProfilerAnalyzer::ProcessBuffer(const FEventTime& EventTime, uint32 T
 				PendingCycle = LastCycle;
 			}
 
-			double Time = EventTime.AsSeconds(PendingCycle);
+			double PendingTime = EventTime.AsSeconds(PendingCycle);
 			if (bEnter)
 			{
 				FTimingProfilerEvent Event;
 				Event.TimerIndex = PendingCursor->TimerId;
-				ThreadState.Timeline->AppendBeginEvent(Time, Event);
+				ThreadState.Timeline->AppendBeginEvent(PendingTime, Event);
 			}
 			else
 			{
-				ThreadState.Timeline->AppendEndEvent(Time);
+				ThreadState.Timeline->AppendEndEvent(PendingTime);
 			}
 		}
 
@@ -256,15 +276,15 @@ uint64 FCpuProfilerAnalyzer::ProcessBuffer(const FEventTime& EventTime, uint32 T
 		if (int64(PendingCycle) < 0)
 		{
 			PendingCycle = ~PendingCycle;
-			double Time = EventTime.AsSeconds(PendingCycle);
-			ThreadState.Timeline->AppendEndEvent(Time);
+			double PendingTime = EventTime.AsSeconds(PendingCycle);
+			ThreadState.Timeline->AppendEndEvent(PendingTime);
 		}
 		else
 		{
 			FTimingProfilerEvent Event;
 			Event.TimerIndex = PendingCursor->TimerId;
-			double Time = EventTime.AsSeconds(PendingCycle);
-			ThreadState.Timeline->AppendBeginEvent(Time, Event);
+			double PendingTime = EventTime.AsSeconds(PendingCycle);
+			ThreadState.Timeline->AppendBeginEvent(PendingTime, Event);
 		}
 	}
 
diff --git a/Engine/Source/Developer/TraceServices/Private/Analyzers/CpuProfilerTraceAnalysis.h b/Engine/Source/Developer/TraceServices/Private/Analyzers/CpuProfilerTraceAnalysis.h
--- a/Engine/Source/Developer/TraceServices/Private/Analyzers/CpuProfilerTraceAnalysis.h
+++ b/Engine/Source/Developer/TraceServices/Private/Analyzers/CpuProfilerTraceAnalysis.h
@@ -46,7 +46,7 @@ private:
 	void OnCpuScopeLeave(const FOnEventContext& Context);
 	uint32 DefineTimer(uint32 SpecId, const TCHAR* ScopeName, const TCHAR* File, uint32 Line, bool bMergeByName); // returns the TimerId
 	FThreadState& GetThreadState(uint32 ThreadId);
-	uint64 ProcessBuffer(const FEventTime& EventTime, uint32 ThreadId, const uint8* BufferPtr, uint32 BufferSize);
+	uint64 ProcessBuffer(const FEventTime& EventTime, FThreadState& ThreadState, const uint8* BufferPtr, uint32 BufferSize);
 
 	enum : uint16
 	{
