[Insights] Make specific no callstack entry

Allows us to display allocations that don't have a callstack (e.g. from TraceMalloc) compared to callstack that have a callstack id, but we couldn't find the id.

diff --git a/Engine/Source/Developer/TraceServices/Private/Model/CallstacksProvider.cpp b/Engine/Source/Developer/TraceServices/Private/Model/CallstacksProvider.cpp
--- a/Engine/Source/Developer/TraceServices/Private/Model/CallstacksProvider.cpp
+++ b/Engine/Source/Developer/TraceServices/Private/Model/CallstacksProvider.cpp
@@ -14,8 +14,11 @@ namespace TraceServices
 /////////////////////////////////////////////////////////////////////
 static const FResolvedSymbol GNeverResolveSymbol(ESymbolQueryResult::NotLoaded, nullptr, nullptr, nullptr, 0);
 static const FResolvedSymbol GNotFoundSymbol(ESymbolQueryResult::NotFound, TEXT("Unknown"), nullptr, nullptr, 0);
-static const FStackFrame GNotFoundStackFrame = { 0, &GNotFoundSymbol};
+static const FResolvedSymbol GNoSymbol(ESymbolQueryResult::NotFound, TEXT("No callstack recorded"), nullptr, nullptr, 0);
+static constexpr FStackFrame GNotFoundStackFrame = { 0, &GNotFoundSymbol};
+static constexpr FStackFrame GNoStackFrame = {0, &GNoSymbol};	
 static const FCallstack GNotFoundCallstack(&GNotFoundStackFrame, 1);
+static const FCallstack GNoCallstack(&GNoStackFrame, 1);
 
 /////////////////////////////////////////////////////////////////////
 #ifdef TRACE_CALLSTACK_STATS
@@ -105,6 +108,10 @@ void FCallstacksProvider::AddCallstack(uint32 InCallstackId, const uint64* InFra
 const FCallstack* FCallstacksProvider::GetCallstack(uint64 CallstackId) const
 {
 	FRWScopeLock ReadLock(EntriesLock, SLT_ReadOnly);
+	if (CallstackId == 0)
+	{
+		return &GNoCallstack;
+	}
 	const FCallstack* const* FindResult = CallstackEntries.Find(CallstackId);
 	return FindResult ? *FindResult : &GNotFoundCallstack;
 }
