[Insights] Fix parsing of environment variables for SymsLibResolver.

Previously only enviroment variables paths that contained ';' very correctly parsed. And even in that case the last paths was discarded.

diff --git a/Engine/Source/Developer/TraceServices/Private/Model/SymslibResolver.cpp b/Engine/Source/Developer/TraceServices/Private/Model/SymslibResolver.cpp
--- a/Engine/Source/Developer/TraceServices/Private/Model/SymslibResolver.cpp
+++ b/Engine/Source/Developer/TraceServices/Private/Model/SymslibResolver.cpp
@@ -5,7 +5,6 @@
 #include "SymslibResolver.h"
 #include "Algo/ForEach.h"
 #include "Algo/Sort.h"
-#include "Algo/Transform.h"
 #include "Async/MappedFileHandle.h"
 #include "Async/ParallelFor.h"
 #include "Containers/StringView.h"
@@ -494,22 +493,28 @@ FSymslibResolver::FSymslibResolver(IAnalysisSession& InSession)
 		GConfig->GetArray(TEXT("Insights.MemoryProfiler"), TEXT("SymbolSearchPaths"), SymbolSearchPaths, SettingsIni);
 	}
 
-	// Paths from environment
-	FString InsightsSymbolPath =  FPlatformMisc::GetEnvironmentVariable(TEXT("UE_INSIGHTS_SYMBOL_PATH"));
-	UE_LOG(LogSymslib, Log, TEXT("UE_INSIGHTS_SYMBOL_PATH: '%s'"), *InsightsSymbolPath);
-	FString SymbolPathPart;
-	while (InsightsSymbolPath.Split(TEXT(";"), &SymbolPathPart, &InsightsSymbolPath))
+	auto SplitEnvPaths = [](FStringView EnvVariable, TArray<FString>& OutList)
 	{
-		SymbolSearchPaths.Emplace(SymbolPathPart);
-	}
+		FString SymbolPathPart, SymbolPathRemainder(EnvVariable);
+		while (SymbolPathRemainder.Split(TEXT(";"), &SymbolPathPart, &SymbolPathRemainder))
+		{
+			OutList.Emplace(SymbolPathPart);
+		}
+		if (!SymbolPathRemainder.IsEmpty())
+		{
+			OutList.Emplace(SymbolPathRemainder);
+		}
+	};
+
+	// Paths from environment
+	const FString InsightsSymbolPath =  FPlatformMisc::GetEnvironmentVariable(TEXT("UE_INSIGHTS_SYMBOL_PATH"));
+	UE_LOG(LogSymslib, Log, TEXT("UE_INSIGHTS_SYMBOL_PATH: %s"), InsightsSymbolPath.IsEmpty() ? TEXT("Not set") : *InsightsSymbolPath);
+	SplitEnvPaths(InsightsSymbolPath, SymbolSearchPaths);
 
 #if PLATFORM_WINDOWS
-	FString NTSymbolPath = FPlatformMisc::GetEnvironmentVariable(TEXT("_NT_SYMBOL_PATH"));
-	UE_LOG(LogSymslib, Log, TEXT("_NT_SYMBOL_PATH: '%s'"), *NTSymbolPath);
-	while (NTSymbolPath.Split(TEXT(";"), &SymbolPathPart, &NTSymbolPath))
-	{
-		 SymbolSearchPaths.Emplace(SymbolPathPart);
-	}
+	const FString NTSymbolPath = FPlatformMisc::GetEnvironmentVariable(TEXT("_NT_SYMBOL_PATH"));
+	UE_LOG(LogSymslib, Log, TEXT("_NT_SYMBOL_PATH: %s"), NTSymbolPath.IsEmpty() ? TEXT("Not set") : *NTSymbolPath);
+	SplitEnvPaths(NTSymbolPath, SymbolSearchPaths);
 #endif
 
 	// Try to get session information
