[Insights] Fix Task Graph Insights crash when late connecting

diff --git a/Engine/Source/Developer/TraceServices/Private/Model/TasksProfiler.cpp b/Engine/Source/Developer/TraceServices/Private/Model/TasksProfiler.cpp
--- a/Engine/Source/Developer/TraceServices/Private/Model/TasksProfiler.cpp
+++ b/Engine/Source/Developer/TraceServices/Private/Model/TasksProfiler.cpp
@@ -160,6 +160,8 @@ namespace TraceServices
 
 	void FTasksProvider::TaskScheduled(TaskTrace::FId TaskId, double Timestamp, uint32 ThreadId)
 	{
+		InitTaskIdToIndexConversion(TaskId);
+
 		if (!TryRegisterEvent(TEXT("TaskScheduled"), TaskId, &FTaskInfo::ScheduledTimestamp, Timestamp, &FTaskInfo::ScheduledThreadId, ThreadId))
 		{
 			return;
@@ -183,6 +185,8 @@ namespace TraceServices
 
 	void FTasksProvider::SubsequentAdded(TaskTrace::FId TaskId, TaskTrace::FId SubsequentId, double Timestamp, uint32 ThreadId)
 	{
+		InitTaskIdToIndexConversion(TaskId);
+
 		// when FGraphEvent is used to wait for a notification, it doesn't have an associated task and so is not created or launched. 
 		// In this case we need to create it and initialise it before registering the completion event
 		FTaskInfo* Task = TryGetOrCreateTask(TaskId);
@@ -202,6 +206,8 @@ namespace TraceServices
 
 	void FTasksProvider::TaskStarted(TaskTrace::FId TaskId, double Timestamp, uint32 ThreadId)
 	{
+		InitTaskIdToIndexConversion(TaskId);
+
 		if (!TryRegisterEvent(TEXT("TaskStarted"), TaskId, &FTaskInfo::StartedTimestamp, Timestamp, &FTaskInfo::StartedThreadId, ThreadId))
 		{
 			return;
@@ -232,6 +238,8 @@ namespace TraceServices
 
 	void FTasksProvider::NestedAdded(TaskTrace::FId TaskId, TaskTrace::FId NestedId, double Timestamp, uint32 ThreadId)
 	{
+		InitTaskIdToIndexConversion(TaskId);
+
 		AddRelative(TEXT("Nested"), TaskId, &FTaskInfo::NestedTasks, NestedId, Timestamp, ThreadId);
 
 		FTaskInfo* Task = TryGetTask(NestedId);
@@ -243,6 +251,8 @@ namespace TraceServices
 
 	void FTasksProvider::TaskFinished(TaskTrace::FId TaskId, double Timestamp)
 	{
+		InitTaskIdToIndexConversion(TaskId);
+
 		if (!TryRegisterEvent(TEXT("TaskFinished"), TaskId, &FTaskInfo::FinishedTimestamp, Timestamp))
 		{
 			return;
@@ -262,6 +272,8 @@ namespace TraceServices
 
 	void FTasksProvider::TaskCompleted(TaskTrace::FId TaskId, double Timestamp, uint32 ThreadId)
 	{
+		InitTaskIdToIndexConversion(TaskId);
+
 		// when FGraphEvent is used to wait for a notification, it doesn't have an associated task and so is not created or launched. 
 		// In this case we need to create it and initialise it before registering the completion event
 		FTaskInfo* Task = TryGetOrCreateTask(TaskId);
@@ -372,6 +384,8 @@ namespace TraceServices
 	{
 		UE_LOG(LogTraceServices, Verbose, TEXT("%s (%d) added to TaskId: %d, Timestamp %.6f)"), RelationType, RelativeId, TaskId, Timestamp);
 
+		InitTaskIdToIndexConversion(TaskId);
+
 		FTaskInfo* Task = TryGetTask(TaskId);
 		if (Task == nullptr)
 		{
