[Insights]
  - Memory Insights: Fixed analysis for Memory.Free and Memory.ReallocFree trace events to correctly read the RootHeap field.
  - Memory Insights: Added INSIGHTS_MEM_TRACE_LEGACY_FORMAT toggle macro and comments to isolate code for temporary backward compatibility with legacy format of memory trace (5.0-EA) which can be deprecated and removed at a later date.

diff --git a/Engine/Source/Developer/TraceServices/Private/Analyzers/AllocationsAnalysis.cpp b/Engine/Source/Developer/TraceServices/Private/Analyzers/AllocationsAnalysis.cpp
--- a/Engine/Source/Developer/TraceServices/Private/Analyzers/AllocationsAnalysis.cpp
+++ b/Engine/Source/Developer/TraceServices/Private/Analyzers/AllocationsAnalysis.cpp
@@ -8,6 +8,8 @@
 namespace TraceServices
 {
 
+#define INSIGHTS_MEM_TRACE_LEGACY_FORMAT 1 // backward compatibility with legacy memory trace format (5.0-EA)
+
 ////////////////////////////////////////////////////////////////////////////////////////////////////
 
 FAllocationsAnalyzer::FAllocationsAnalyzer(IAnalysisSession& InSession, FAllocationsProvider& InAllocationsProvider)
@@ -45,8 +47,8 @@ void FAllocationsAnalyzer::OnAnalysisBegin(const FOnAnalysisContext& Context)
 	Builder.RouteEvent(RouteId_ReallocFreeSystem,   "Memory", "ReallocFreeSystem");
 	Builder.RouteEvent(RouteId_Marker,              "Memory", "Marker");
 	Builder.RouteEvent(RouteId_TagSpec,             "Memory", "TagSpec");
-	Builder.RouteEvent(RouteId_HeapSpec,	        "Memory", "HeapSpec");
-	Builder.RouteEvent(RouteId_HeapMarkAlloc,	    "Memory", "HeapMarkAlloc");
+	Builder.RouteEvent(RouteId_HeapSpec,            "Memory", "HeapSpec");
+	Builder.RouteEvent(RouteId_HeapMarkAlloc,       "Memory", "HeapMarkAlloc");
 	Builder.RouteEvent(RouteId_HeapUnmarkAlloc,     "Memory", "HeapUnmarkAlloc");
 
 	Builder.RouteLoggerEvents(RouteId_MemScope, "Memory", true);
@@ -82,17 +84,31 @@ bool FAllocationsAnalyzer::OnEvent(uint16 RouteId, EStyle Style, const FOnEventC
 	{
 		case RouteId_Init:
 		{
-			const double Time = GetCurrentTime();
-			const uint8 MinAlignment = EventData.GetValue<uint8>("MinAlignment");
-			SizeShift = EventData.GetValue<uint8>("SizeShift");
+#if !INSIGHTS_MEM_TRACE_LEGACY_FORMAT
+			const uint8 Version = EventData.GetValue<uint8>("Version");
 
+			constexpr uint8 MinSupportedVersion = 1; // UE 5.0
+			constexpr uint8 MaxSupportedVersion = 1; // UE 5.0
+			if (Version < MinSupportedVersion || Version > MaxSupportedVersion)
+			{
+				break;
+			}
+#endif // !INSIGHTS_MEM_TRACE_LEGACY_FORMAT
+
+			const double Time = GetCurrentTime();
+#if INSIGHTS_MEM_TRACE_LEGACY_FORMAT
 			BaseCycle = EventData.GetValue<uint64>("BaseCycle", 0);
+#endif // INSIGHTS_MEM_TRACE_LEGACY_FORMAT
 			MarkerPeriod = EventData.GetValue<uint32>("MarkerPeriod");
 
+			const uint8 MinAlignment = EventData.GetValue<uint8>("MinAlignment");
+			SizeShift = EventData.GetValue<uint8>("SizeShift");
+
 			FAllocationsProvider::FEditScopeLock _(AllocationsProvider);
 			AllocationsProvider.EditInit(Time, MinAlignment);
 			break;
 		}
+
 		case RouteId_HeapSpec:
 		{
 			const HeapId Id = EventData.GetValue<uint16>("Id");
@@ -105,6 +121,7 @@ bool FAllocationsAnalyzer::OnEvent(uint16 RouteId, EStyle Style, const FOnEventC
 			AllocationsProvider.EditHeapSpec(Id, ParentId, Name, Flags);
 			break;
 		}
+
 		case RouteId_AllocSystem:
 		case RouteId_AllocVideo:
 		case RouteId_ReallocAllocSystem:
@@ -122,12 +139,14 @@ bool FAllocationsAnalyzer::OnEvent(uint16 RouteId, EStyle Style, const FOnEventC
 			//     FAllocEvent Event; // aligned
 			//     EventData.CopyData(&Event, sizeof(Event));
 			const double Time = GetCurrentTime();
-			uint64 Owner = EventData.GetValue<uint32>("CallstackId");
-			if (!Owner)
+			uint64 CallstackId = EventData.GetValue<uint32>("CallstackId");
+#if INSIGHTS_MEM_TRACE_LEGACY_FORMAT
+			if (!CallstackId)
 			{
-				// Legeacy format of sending the hash value
-				Owner = EventData.GetValue<uint64>("Owner");
+				// Legacy format of sending the hash value
+				CallstackId = EventData.GetValue<uint64>("Owner");
 			}
+#endif // INSIGHTS_MEM_TRACE_LEGACY_FORMAT
 			uint64 Address = EventData.GetValue<uint64>("Address");
 			RootHeap = EventData.GetValue<uint8>("RootHeap", RootHeap);
 			uint64 SizeUpper = EventData.GetValue<uint32>("Size");
@@ -135,6 +154,7 @@ bool FAllocationsAnalyzer::OnEvent(uint16 RouteId, EStyle Style, const FOnEventC
 			const uint8 AlignmentMask = ~SizeLowerMask;
 			uint64 Size = 0;
 			uint32 Alignment = 0;
+#if INSIGHTS_MEM_TRACE_LEGACY_FORMAT
 			const uint8 Alignment_SizeLower = EventData.GetValue<uint8>("Alignment_SizeLower");
 			if (Alignment_SizeLower)
 			{
@@ -142,16 +162,18 @@ bool FAllocationsAnalyzer::OnEvent(uint16 RouteId, EStyle Style, const FOnEventC
 				Alignment = Alignment_SizeLower & AlignmentMask;
 			}
 			else
+#endif // INSIGHTS_MEM_TRACE_LEGACY_FORMAT
 			{
 				const uint8 AlignmentPow2_SizeLower = EventData.GetValue<uint8>("AlignmentPow2_SizeLower");
 				Size = SizeUpper << SizeShift | static_cast<uint64>(AlignmentPow2_SizeLower & SizeLowerMask);
 				Alignment = 1 << (AlignmentPow2_SizeLower >> SizeShift);
 			}
+
 			const uint32 ThreadId = Context.ThreadInfo.GetSystemId();
 			const uint8 Tracker = 0; // We only care about the default tracker for now.
 
 			FAllocationsProvider::FEditScopeLock _(AllocationsProvider);
-			AllocationsProvider.EditAlloc(Time, Owner, Address, Size, Alignment, ThreadId, Tracker, RootHeap);
+			AllocationsProvider.EditAlloc(Time, CallstackId, Address, Size, Alignment, ThreadId, Tracker, RootHeap);
 			if (RouteId == RouteId_ReallocAlloc || RouteId == RouteId_ReallocAllocSystem)
 			{
 				AllocationsProvider.EditPopTagFromPtr(ThreadId, Tracker);
@@ -171,14 +193,18 @@ bool FAllocationsAnalyzer::OnEvent(uint16 RouteId, EStyle Style, const FOnEventC
 		{
 			const double Time = GetCurrentTime();
 			uint64 Address = EventData.GetValue<uint64>("Address");
-			if (Address == 0) //@todo: Check that this isn't intentional for reallocs?
+#if INSIGHTS_MEM_TRACE_LEGACY_FORMAT
+			if (Address == 0)
 			{
+				// Legacy format of sending Address and RootHeap packed togather (after 5.0-EA).
 				constexpr uint32 HeapShift = 60;
 				constexpr uint64 RootHeapMask = uint64(0xF) << HeapShift;
 				const uint64 Address_RootHeap = EventData.GetValue<uint64>("Address_RootHeap");
 				Address = Address_RootHeap & ~RootHeapMask;
 				RootHeap = (Address_RootHeap & RootHeapMask) >> HeapShift;
 			}
+#endif // INSIGHTS_MEM_TRACE_LEGACY_FORMAT
+			RootHeap = EventData.GetValue<uint8>("RootHeap", RootHeap);
 
 			FAllocationsProvider::FEditScopeLock _(AllocationsProvider);
 			if (RouteId == RouteId_ReallocFree || RouteId == RouteId_ReallocFreeSystem)
@@ -190,6 +216,7 @@ bool FAllocationsAnalyzer::OnEvent(uint16 RouteId, EStyle Style, const FOnEventC
 			AllocationsProvider.EditFree(Time, Address, RootHeap);
 			break;
 		}
+
 		case RouteId_HeapMarkAlloc:
 		{
 			const double Time = GetCurrentTime();
@@ -201,6 +228,7 @@ bool FAllocationsAnalyzer::OnEvent(uint16 RouteId, EStyle Style, const FOnEventC
 			AllocationsProvider.EditMarkAllocationAsHeap(Time, Address, Heap, Flags);
 			break;
 		}
+
 		case RouteId_HeapUnmarkAlloc:
 		{
 			const double Time = GetCurrentTime();
@@ -211,6 +239,7 @@ bool FAllocationsAnalyzer::OnEvent(uint16 RouteId, EStyle Style, const FOnEventC
 			AllocationsProvider.EditUnmarkAllocationAsHeap(Time, Address, Heap);
 			break;
 		}
+
 		case RouteId_Marker:
 		{
 			// If BaseCycle is 0, then Cycle is a 64-bit absolute value, otherwise Cycle is a 32-bit value (relative to BaseCycle).
@@ -236,6 +265,7 @@ bool FAllocationsAnalyzer::OnEvent(uint16 RouteId, EStyle Style, const FOnEventC
 			}
 			break;
 		}
+
 		case RouteId_TagSpec:
 		{
 			const TagIdType Tag = Context.EventData.GetValue<TagIdType>("Tag");
@@ -249,6 +279,7 @@ bool FAllocationsAnalyzer::OnEvent(uint16 RouteId, EStyle Style, const FOnEventC
 			AllocationsProvider.EditAddTagSpec(Tag, Parent, DisplayString);
 			break;
 		}
+
 		case RouteId_MemScope:
 		{
 			const uint32 ThreadId = Context.ThreadInfo.GetSystemId();
@@ -297,4 +328,6 @@ double FAllocationsAnalyzer::GetCurrentTime() const
 
 ////////////////////////////////////////////////////////////////////////////////////////////////////
 
+#undef INSIGHTS_MEM_TRACE_LEGACY_FORMAT
+
 } // namespace TraceServices
diff --git a/Engine/Source/Developer/TraceServices/Private/Analyzers/CallstacksAnalysis.cpp b/Engine/Source/Developer/TraceServices/Private/Analyzers/CallstacksAnalysis.cpp
--- a/Engine/Source/Developer/TraceServices/Private/Analyzers/CallstacksAnalysis.cpp
+++ b/Engine/Source/Developer/TraceServices/Private/Analyzers/CallstacksAnalysis.cpp
@@ -25,17 +25,18 @@ void FCallstacksAnalyzer::OnAnalysisBegin(const FOnAnalysisContext& Context)
 ////////////////////////////////////////////////////////////////////////////////
 bool FCallstacksAnalyzer::OnEvent(uint16 RouteId, EStyle Style, const FOnEventContext& Context)
 {
-	switch(RouteId)
+	switch (RouteId)
 	{
 		case RouteId_Callstack:
 			const TArrayReader<uint64>& Frames = Context.EventData.GetArray<uint64>("Frames");
-			if (const uint64 Hash = Context.EventData.GetValue<uint64>("Id"))
+			if (const uint32 Id = Context.EventData.GetValue<uint32>("CallstackId"))
 			{
-				Provider->AddCallstack(Hash, Frames.GetData(), uint8(Frames.Num()));
+				Provider->AddCallstack(Id, Frames.GetData(), uint8(Frames.Num()));
 			}
-			else if (const uint32 Id = Context.EventData.GetValue<uint32>("CallstackId"))
+			// Backward compatibility with legacy memory trace format (5.0-EA).
+			else if (const uint64 Hash = Context.EventData.GetValue<uint64>("Id"))
 			{
-				Provider->AddCallstack(Id, Frames.GetData(), uint8(Frames.Num()));
+				Provider->AddCallstack(Hash, Frames.GetData(), uint8(Frames.Num()));
 			}
 			break;
 	}
