[Insights]
  - Memory Insights: Added analysis support for "free callstack" for an allocation (i.e. callstack when an allocation if freed).
  - Memory Insights: Switched CallstackId to 32 bit and moved the resolving of "id to pointer" outside of mem queries. This allows us to add the second callstack id and also ids for future metadata without increasing the size of FAllocationItem struct (from AllocationProvider).
  - Memory Insights: Refactored code re ThreadId and Tracker passed from AllocationAnalysis to AllocationProvider.
  - Memory Insights: Refactored code re CallstackProvider to map a CallstackId directly as an index into array of callstacks.

diff --git a/Engine/Source/Developer/TraceInsights/Private/Insights/MemoryProfiler/ViewModels/MemoryAlloc.h b/Engine/Source/Developer/TraceInsights/Private/Insights/MemoryProfiler/ViewModels/MemoryAlloc.h
--- a/Engine/Source/Developer/TraceInsights/Private/Insights/MemoryProfiler/ViewModels/MemoryAlloc.h
+++ b/Engine/Source/Developer/TraceInsights/Private/Insights/MemoryProfiler/ViewModels/MemoryAlloc.h
@@ -43,6 +43,7 @@ private:
 	int64 Size;
 	const TCHAR* Tag;
 	const TraceServices::FCallstack* Callstack;
+	const TraceServices::FCallstack* FreeCallstack;
 	HeapId RootHeap;
 	bool bIsBlock;
 	bool bIsDecline;
diff --git a/Engine/Source/Developer/TraceInsights/Private/Insights/MemoryProfiler/Widgets/SMemAllocTableTreeView.cpp b/Engine/Source/Developer/TraceInsights/Private/Insights/MemoryProfiler/Widgets/SMemAllocTableTreeView.cpp
--- a/Engine/Source/Developer/TraceInsights/Private/Insights/MemoryProfiler/Widgets/SMemAllocTableTreeView.cpp
+++ b/Engine/Source/Developer/TraceInsights/Private/Insights/MemoryProfiler/Widgets/SMemAllocTableTreeView.cpp
@@ -298,6 +298,8 @@ void SMemAllocTableTreeView::UpdateQuery(TraceServices::IAllocationsProvider::EQ
 	}
 	const TraceServices::IAllocationsProvider& Provider = *AllocationsProvider;
 
+	const TraceServices::ICallstacksProvider* CallstacksProvider = TraceServices::ReadCallstacksProvider(*Session.Get());
+
 	constexpr double MaxPollTime = 0.03; // Stop getting results after 30 ms so we don't tank the frame rate too much.
 	FStopwatch TotalStopwatch;
 	TotalStopwatch.Start();
@@ -361,7 +363,11 @@ void SMemAllocTableTreeView::UpdateQuery(TraceServices::IAllocationsProvider::EQ
 					Alloc.Address = Allocation->GetAddress();
 					Alloc.Size = int64(Allocation->GetSize());
 					Alloc.Tag = Provider.GetTagName(Allocation->GetTag());
-					Alloc.Callstack = Allocation->GetCallstack();
+					if (CallstacksProvider)
+					{
+						Alloc.Callstack = CallstacksProvider->GetCallstack(Allocation->GetCallstackId());
+						Alloc.FreeCallstack = CallstacksProvider->GetCallstack(Allocation->GetFreeCallstackId());
+					}
 					Alloc.RootHeap = Allocation->GetRootHeap();
 					Alloc.bIsBlock = Allocation->IsHeap();
 					check(Alloc.Callstack != nullptr);
diff --git a/Engine/Source/Developer/TraceServices/Private/Analyzers/AllocationsAnalysis.cpp b/Engine/Source/Developer/TraceServices/Private/Analyzers/AllocationsAnalysis.cpp
--- a/Engine/Source/Developer/TraceServices/Private/Analyzers/AllocationsAnalysis.cpp
+++ b/Engine/Source/Developer/TraceServices/Private/Analyzers/AllocationsAnalysis.cpp
@@ -138,17 +138,28 @@ bool FAllocationsAnalyzer::OnEvent(uint16 RouteId, EStyle Style, const FOnEventC
 			// Or something like:
 			//     FAllocEvent Event; // aligned
 			//     EventData.CopyData(&Event, sizeof(Event));
+
 			const double Time = GetCurrentTime();
-			uint64 CallstackId = EventData.GetValue<uint32>("CallstackId");
+
+			// CallstackId is optional. If the field is not present CallstackId will be 0 (i.e. "no callstack").
+			uint32 CallstackId = EventData.GetValue<uint32>("CallstackId", 0);
 #if INSIGHTS_MEM_TRACE_LEGACY_FORMAT
 			if (!CallstackId)
 			{
-				// Legacy format of sending the hash value
-				CallstackId = EventData.GetValue<uint64>("Owner");
+				const ICallstacksProvider* CallstacksProvider = ReadCallstacksProvider(Session);
+				if (CallstacksProvider)
+				{
+					// Legacy format of sending the hash value
+					const uint64 CallstackHash = EventData.GetValue<uint64>("Owner");
+					CallstackId = CallstacksProvider->GetCallstackIdForHash(CallstackHash);
+				}
 			}
 #endif // INSIGHTS_MEM_TRACE_LEGACY_FORMAT
+
 			uint64 Address = EventData.GetValue<uint64>("Address");
+
 			RootHeap = EventData.GetValue<uint8>("RootHeap", RootHeap);
+
 			uint64 SizeUpper = EventData.GetValue<uint32>("Size");
 			const uint8 SizeLowerMask = ((1 << SizeShift) - 1);
 			const uint8 AlignmentMask = ~SizeLowerMask;
@@ -169,14 +180,16 @@ bool FAllocationsAnalyzer::OnEvent(uint16 RouteId, EStyle Style, const FOnEventC
 				Alignment = 1 << (AlignmentPow2_SizeLower >> SizeShift);
 			}
 
-			const uint32 ThreadId = Context.ThreadInfo.GetSystemId();
-			const uint8 Tracker = 0; // We only care about the default tracker for now.
+			const uint32 TraceThreadId = Context.ThreadInfo.GetId();
+			const uint32 SystemThreadId = Context.ThreadInfo.GetSystemId();
 
 			FAllocationsProvider::FEditScopeLock _(AllocationsProvider);
-			AllocationsProvider.EditAlloc(Time, CallstackId, Address, Size, Alignment, ThreadId, Tracker, RootHeap);
+			AllocationsProvider.SetCurrentThreadId(TraceThreadId, SystemThreadId);
+			AllocationsProvider.EditAlloc(Time, CallstackId, Address, Size, Alignment, RootHeap);
 			if (RouteId == RouteId_ReallocAlloc || RouteId == RouteId_ReallocAllocSystem)
 			{
-				AllocationsProvider.EditPopTagFromPtr(ThreadId, Tracker);
+				const uint8 Tracker = 0; // We only care about the default tracker for now.
+				AllocationsProvider.EditPopTagFromPtr(SystemThreadId, Tracker);
 			}
 			break;
 		}
@@ -192,6 +205,10 @@ bool FAllocationsAnalyzer::OnEvent(uint16 RouteId, EStyle Style, const FOnEventC
 		case RouteId_ReallocFree:
 		{
 			const double Time = GetCurrentTime();
+
+			// CallstackId is optional. If the field is not present CallstackId will be 0 (i.e. "no callstack").
+			uint32 CallstackId = EventData.GetValue<uint32>("CallstackId", 0);
+
 			uint64 Address = EventData.GetValue<uint64>("Address");
 #if INSIGHTS_MEM_TRACE_LEGACY_FORMAT
 			if (Address == 0)
@@ -204,16 +221,20 @@ bool FAllocationsAnalyzer::OnEvent(uint16 RouteId, EStyle Style, const FOnEventC
 				RootHeap = (Address_RootHeap & RootHeapMask) >> HeapShift;
 			}
 #endif // INSIGHTS_MEM_TRACE_LEGACY_FORMAT
+
 			RootHeap = EventData.GetValue<uint8>("RootHeap", RootHeap);
 
+			const uint32 TraceThreadId = Context.ThreadInfo.GetId();
+			const uint32 SystemThreadId = Context.ThreadInfo.GetSystemId();
+
 			FAllocationsProvider::FEditScopeLock _(AllocationsProvider);
+			AllocationsProvider.SetCurrentThreadId(TraceThreadId, SystemThreadId);
 			if (RouteId == RouteId_ReallocFree || RouteId == RouteId_ReallocFreeSystem)
 			{
-				const uint32 ThreadId = Context.ThreadInfo.GetSystemId();
 				const uint8 Tracker = 0; // We only care about the default tracker for now.
-				AllocationsProvider.EditPushTagFromPtr(ThreadId, Tracker, Address);
+				AllocationsProvider.EditPushTagFromPtr(SystemThreadId, Tracker, Address);
 			}
-			AllocationsProvider.EditFree(Time, Address, RootHeap);
+			AllocationsProvider.EditFree(Time, CallstackId, Address, RootHeap);
 			break;
 		}
 
@@ -224,7 +245,11 @@ bool FAllocationsAnalyzer::OnEvent(uint16 RouteId, EStyle Style, const FOnEventC
 			const HeapId Heap = EventData.GetValue<uint16>("Heap", 0);
 			const EMemoryTraceHeapAllocationFlags Flags = EventData.GetValue<EMemoryTraceHeapAllocationFlags>("Flags");
 
+			const uint32 TraceThreadId = Context.ThreadInfo.GetId();
+			const uint32 SystemThreadId = Context.ThreadInfo.GetSystemId();
+
 			FAllocationsProvider::FEditScopeLock _(AllocationsProvider);
+			AllocationsProvider.SetCurrentThreadId(TraceThreadId, SystemThreadId);
 			AllocationsProvider.EditMarkAllocationAsHeap(Time, Address, Heap, Flags);
 			break;
 		}
@@ -235,7 +260,11 @@ bool FAllocationsAnalyzer::OnEvent(uint16 RouteId, EStyle Style, const FOnEventC
 			const uint64 Address = EventData.GetValue<uint64>("Address");
 			const HeapId Heap = EventData.GetValue<uint16>("Heap", 0);
 
+			const uint32 TraceThreadId = Context.ThreadInfo.GetId();
+			const uint32 SystemThreadId = Context.ThreadInfo.GetSystemId();
+
 			FAllocationsProvider::FEditScopeLock _(AllocationsProvider);
+			AllocationsProvider.SetCurrentThreadId(TraceThreadId, SystemThreadId);
 			AllocationsProvider.EditUnmarkAllocationAsHeap(Time, Address, Heap);
 			break;
 		}
diff --git a/Engine/Source/Developer/TraceServices/Private/Analyzers/CallstacksAnalysis.cpp b/Engine/Source/Developer/TraceServices/Private/Analyzers/CallstacksAnalysis.cpp
--- a/Engine/Source/Developer/TraceServices/Private/Analyzers/CallstacksAnalysis.cpp
+++ b/Engine/Source/Developer/TraceServices/Private/Analyzers/CallstacksAnalysis.cpp
@@ -36,7 +36,7 @@ bool FCallstacksAnalyzer::OnEvent(uint16 RouteId, EStyle Style, const FOnEventCo
 			// Backward compatibility with legacy memory trace format (5.0-EA).
 			else if (const uint64 Hash = Context.EventData.GetValue<uint64>("Id"))
 			{
-				Provider->AddCallstack(Hash, Frames.GetData(), uint8(Frames.Num()));
+				Provider->AddCallstackWithHash(Hash, Frames.GetData(), uint8(Frames.Num()));
 			}
 			break;
 	}
diff --git a/Engine/Source/Developer/TraceServices/Private/Model/AllocationsProvider.cpp b/Engine/Source/Developer/TraceServices/Private/Model/AllocationsProvider.cpp
--- a/Engine/Source/Developer/TraceServices/Private/Model/AllocationsProvider.cpp
+++ b/Engine/Source/Developer/TraceServices/Private/Model/AllocationsProvider.cpp
@@ -18,6 +18,10 @@
 #define INSIGHTS_SLOW_CHECK(expr) //check(expr)
 
 #define INSIGHTS_DEBUG_WATCH 0
+// Action to be executed when a match is found.
+#define INSIGHTS_DEBUG_WATCH_FOUND break // just log
+//#define INSIGHTS_DEBUG_WATCH_FOUND return // log and ignore events
+//#define INSIGHTS_DEBUG_WATCH_FOUND UE_DEBUG_BREAK(); break
 
 // Use optimized path for the short living allocations.
 // If enabled, caches the short living allocations, as there is a very high chance to be freed in the next few "free" events.
@@ -43,7 +47,7 @@ constexpr uint32 MaxLogMessagesPerErrorType = 100;
 static uint64 GWatchAddresses[] =
 {
 	// add here addresses to watch
-	0x0,
+	0x0ull,
 };
 #endif // INSIGHTS_DEBUG_WATCH
 
@@ -118,6 +122,10 @@ void FTagTracker::AddTagSpec(TagIdType InTag, TagIdType InParentTag, const TCHAR
 	if (ensure(!TagMap.Contains(InTag)))
 	{
 		TagMap.Emplace(InTag, TagEntry{ InDisplay, InParentTag });
+		if (!InDisplay || *InDisplay == TEXT('\0'))
+		{
+			UE_LOG(LogTraceServices, Warning, TEXT("[MemAlloc] Tag with id %u has invalid display name (ParentTag=%u)!"), InTag, InParentTag);
+		}
 	}
 	else
 	{
@@ -216,7 +224,7 @@ bool FTagTracker::HasTagFromPtrScope(uint32 InThreadId, uint8 InTracker) const
 {
 	const uint32 TrackerThreadId = GetTrackerThreadId(InThreadId, InTracker);
 	const ThreadState* State = TrackerThreadStates.Find(TrackerThreadId);
-	return State && State->TagStack.Num() > 0 && ((State->TagStack.Top() & PtrTagMask) != 0);
+	return State && (State->TagStack.Num() > 0) && ((State->TagStack.Top() & PtrTagMask) != 0);
 }
 
 ////////////////////////////////////////////////////////////////////////////////////////////////////
@@ -278,10 +286,24 @@ uint32 IAllocationsProvider::FAllocation::GetAlignment() const
 }
 
 ////////////////////////////////////////////////////////////////////////////////////////////////////
-const FCallstack* IAllocationsProvider::FAllocation::GetCallstack() const
+uint32 IAllocationsProvider::FAllocation::GetCallstackId() const
 {
 	const auto* Inner = (const FAllocationItem*)this;
-	return Inner->Callstack;
+	return Inner->CallstackId;
+}
+
+////////////////////////////////////////////////////////////////////////////////////////////////////
+uint32 IAllocationsProvider::FAllocation::GetFreeCallstackId() const
+{
+	const auto* Inner = (const FAllocationItem*)this;
+	return Inner->FreeCallstackId;
+}
+
+////////////////////////////////////////////////////////////////////////////////////////////////////
+uint32 IAllocationsProvider::FAllocation::GetMetadataId() const
+{
+	const auto* Inner = (const FAllocationItem*)this;
+	return Inner->MetadataId;
 }
 
 ////////////////////////////////////////////////////////////////////////////////////////////////////
@@ -1081,7 +1103,7 @@ void FAllocationsProvider::EditHeapSpec(HeapId Id, HeapId ParentId, const FStrin
 
 ////////////////////////////////////////////////////////////////////////////////////////////////////
 
-void FAllocationsProvider::EditAlloc(double Time, uint64 Owner, uint64 Address, uint64 InSize, uint32 InAlignment, uint32 ThreadId, uint8 Tracker, HeapId RootHeap)
+void FAllocationsProvider::EditAlloc(double Time, uint32 CallstackId, uint64 Address, uint64 InSize, uint32 InAlignment, HeapId RootHeap)
 {
 	Lock.WriteAccessCheck();
 
@@ -1092,7 +1114,7 @@ void FAllocationsProvider::EditAlloc(double Time, uint64 Owner, uint64 Address,
 
 	if (Address == 0)
 	{
-		UE_LOG(LogTraceServices, Warning, TEXT("[MemAlloc] Alloc at address 0 : Size=%llu, RootHeap=%u, Time=%f"), InSize, RootHeap, Time);
+		UE_LOG(LogTraceServices, Warning, TEXT("[MemAlloc] Alloc at address 0 : Size=%llu, RootHeap=%u, Time=%f, CallstackId=%u"), InSize, RootHeap, Time, CallstackId);
 		return;
 	}
 
@@ -1107,18 +1129,18 @@ void FAllocationsProvider::EditAlloc(double Time, uint64 Owner, uint64 Address,
 
 	AdvanceTimelines(Time);
 
-	const TagIdType Tag = TagTracker.GetCurrentTag(ThreadId, Tracker);
+	const TagIdType Tag = TagTracker.GetCurrentTag(CurrentSystemThreadId, CurrentTracker);
 
 #if INSIGHTS_DEBUG_WATCH
 	for (int32 AddrIndex = 0; AddrIndex < UE_ARRAY_COUNT(GWatchAddresses); ++AddrIndex)
 	{
 		if (GWatchAddresses[AddrIndex] == Address)
 		{
-			UE_LOG(LogTraceServices, Warning, TEXT("[MemAlloc] Alloc 0x%llX : Size=%llu, Tag=%u, RootHeap=%u, Time=%f"), Address, InSize, Tag, RootHeap, Time);
-			break;
+			UE_LOG(LogTraceServices, Warning, TEXT("[MemAlloc][%u] Alloc 0x%llX : Size=%llu, Tag=%u, RootHeap=%u, Time=%f, CallstackId=%u"), CurrentTraceThreadId, Address, InSize, Tag, RootHeap, Time, CallstackId);
+			INSIGHTS_DEBUG_WATCH_FOUND;
 		}
 	}
-#endif
+#endif // INSIGHTS_DEBUG_WATCH
 
 #if INSIGHTS_VALIDATE_ALLOC_EVENTS
 	FAllocationItem* AllocationPtr = LiveAllocs[RootHeap]->FindRef(Address);
@@ -1129,20 +1151,21 @@ void FAllocationsProvider::EditAlloc(double Time, uint64 Owner, uint64 Address,
 	if (!AllocationPtr)
 	{
 		AllocationPtr = LiveAllocs[RootHeap]->AddNew(Address);
-		INSIGHTS_SLOW_CHECK(AllocationPtr->Address == Address)
-
 		FAllocationItem& Allocation = *AllocationPtr;
+
+		INSIGHTS_SLOW_CHECK(Allocation.Address == Address);
+		Allocation.SizeAndAlignment = FAllocationItem::PackSizeAndAlignment(InSize, InAlignment);
 		Allocation.StartEventIndex = EventIndex[RootHeap];
 		Allocation.EndEventIndex = (uint32)-1;
 		Allocation.StartTime = Time;
 		Allocation.EndTime = std::numeric_limits<double>::infinity();
-		Allocation.Owner = Owner;
-		//Allocation.Address = Address;
-		Allocation.SizeAndAlignment = FAllocationItem::PackSizeAndAlignment(InSize, InAlignment);
-		Allocation.Callstack = nullptr;
+		Allocation.CallstackId = CallstackId;
+		Allocation.FreeCallstackId = 0; // no callstack yet
+		Allocation.MetadataId = 0; //TODO: MetadataProvider.PinAndGetId() // pins the metadata stack and returns an id for it
+		Allocation.Reserved = 0;
 		Allocation.Tag = Tag;
-		Allocation.Flags = EMemoryTraceHeapAllocationFlags::None;
 		Allocation.RootHeap = RootHeap;
+		Allocation.Flags = EMemoryTraceHeapAllocationFlags::None;
 
 		UpdateHistogramByAllocSize(InSize);
 
@@ -1170,7 +1193,7 @@ void FAllocationsProvider::EditAlloc(double Time, uint64 Owner, uint64 Address,
 
 ////////////////////////////////////////////////////////////////////////////////////////////////////
 
-void FAllocationsProvider::EditFree(double Time, uint64 Address, HeapId RootHeap)
+void FAllocationsProvider::EditFree(double Time, uint32 CallstackId, uint64 Address, HeapId RootHeap)
 {
 	Lock.WriteAccessCheck();
 
@@ -1181,7 +1204,7 @@ void FAllocationsProvider::EditFree(double Time, uint64 Address, HeapId RootHeap
 
 	if (Address == 0)
 	{
-		UE_LOG(LogTraceServices, Warning, TEXT("[MemAlloc] Free for address 0 : RootHeap=%u, Time=%f"), RootHeap, Time);
+		UE_LOG(LogTraceServices, Warning, TEXT("[MemAlloc] Free for address 0 : RootHeap=%u, Time=%f, CallstackId=%u"), RootHeap, Time, CallstackId);
 		return;
 	}
 
@@ -1201,8 +1224,8 @@ void FAllocationsProvider::EditFree(double Time, uint64 Address, HeapId RootHeap
 	{
 		if (GWatchAddresses[AddrIndex] == Address)
 		{
-			UE_LOG(LogTraceServices, Warning, TEXT("[MemAlloc] Free 0x%llX : RootHeap=%u, Time=%f"), Address, RootHeap, Time);
-			break;
+			UE_LOG(LogTraceServices, Warning, TEXT("[MemAlloc][%u] Free 0x%llX : RootHeap=%u, Time=%f, CallstackId=%u"), CurrentTraceThreadId, Address, RootHeap, Time, CallstackId);
+			INSIGHTS_DEBUG_WATCH_FOUND;
 		}
 	}
 #endif // INSIGHTS_DEBUG_WATCH
@@ -1224,6 +1247,8 @@ void FAllocationsProvider::EditFree(double Time, uint64 Address, HeapId RootHeap
 		AllocationPtr->EndEventIndex = EventIndex[RootHeap];
 		AllocationPtr->EndTime = Time;
 
+		AllocationPtr->FreeCallstackId = CallstackId;
+
 		const uint64 OldSize = AllocationPtr->GetSize();
 
 		SbTree[RootHeap]->AddAlloc(AllocationPtr); // SbTree takes ownership of AllocationPtr
@@ -1265,11 +1290,11 @@ void FAllocationsProvider::EditUnmarkAllocationAsHeap(double Time, uint64 Addres
 	{
 		if (GWatchAddresses[AddrIndex] == Address)
 		{
-			UE_LOG(LogTraceServices, Warning, TEXT("[MemAlloc] HeapUnmarkAlloc 0x%llX : Heap=%u, Time=%f"), Address, Heap, Time);
-			break;
+			UE_LOG(LogTraceServices, Warning, TEXT("[MemAlloc][%u] HeapUnmarkAlloc 0x%llX : Heap=%u, Time=%f"), CurrentTraceThreadId, Address, Heap, Time);
+			INSIGHTS_DEBUG_WATCH_FOUND;
 		}
 	}
-#endif
+#endif // INSIGHTS_DEBUG_WATCH
 
 	HeapId RootHeap = FindRootHeap(Heap);
 
@@ -1279,19 +1304,20 @@ void FAllocationsProvider::EditUnmarkAllocationAsHeap(double Time, uint64 Addres
 	{
 		const uint64 Size = Alloc->GetSize();
 		const uint32 Alignment = Alloc->GetAlignment();
-		const uint32 Owner = Alloc->Owner;
-		const uint32 ThreadId = 0;
-		const uint8 Tracker = 1;
+		const uint32 CallstackId = Alloc->CallstackId;
+		const uint32 FreeCallstackId = 0; // unknown
 
 		// Re-add this allocation to the Live allocs.
 		LiveAllocs[RootHeap]->Add(Alloc); // the Live allocs takes ownership of Alloc
 
-		// We cannot just unmark the allocation as heap, there is no timestamp support, instead fake a free
-		// and allocation. Make sure the new allocation retains the tag from the original.
-		EditPushTagFromPtr(ThreadId, Tracker, Address);
-		EditFree(Time, Address, RootHeap);
-		EditAlloc(Time, Address, Owner, Size, Alignment, 0, 1, RootHeap);
-		EditPopTagFromPtr(ThreadId, Tracker);
+		// We cannot just unmark the allocation as heap, there is no timestamp support, instead fake a "free"
+		// event and an "alloc" event. Make sure the new allocation retains the tag from the original.
+		CurrentTracker = 1;
+		EditPushTagFromPtr(CurrentSystemThreadId, CurrentTracker, Address);
+		EditFree(Time, FreeCallstackId, Address, RootHeap);
+		EditAlloc(Time, CallstackId, Address, Size, Alignment, RootHeap);
+		EditPopTagFromPtr(CurrentSystemThreadId, CurrentTracker);
+		CurrentTracker = 0;
 	}
 	else
 	{
@@ -1314,11 +1340,11 @@ void FAllocationsProvider::EditMarkAllocationAsHeap(double Time, uint64 Address,
 	{
 		if (GWatchAddresses[AddrIndex] == Address)
 		{
-			UE_LOG(LogTraceServices, Warning, TEXT("[MemAlloc] HeapMarkAlloc 0x%llX : Heap=%u, Flags=%u, Time=%f"), Address, Heap, uint32(Flags), Time);
-			break;
+			UE_LOG(LogTraceServices, Warning, TEXT("[MemAlloc][%u] HeapMarkAlloc 0x%llX : Heap=%u, Flags=%u, Time=%f"), CurrentTraceThreadId, Address, Heap, uint32(Flags), Time);
+			INSIGHTS_DEBUG_WATCH_FOUND;
 		}
 	}
-#endif
+#endif // INSIGHTS_DEBUG_WATCH
 
 	HeapId RootHeap = FindRootHeap(Heap);
 
@@ -1561,14 +1587,42 @@ void FAllocationsProvider::EditOnAnalysisCompleted(double Time)
 			Tree->Validate();
 		}
 	}
+
+	if (AllocErrors > 0)
+	{
+		UE_LOG(LogTraceServices, Error, TEXT("[MemAlloc] ALLOC event errors: %u"), AllocErrors);
+	}
+	if (FreeErrors > 0)
+	{
+		UE_LOG(LogTraceServices, Error, TEXT("[MemAlloc] FREE event errors: %u"), FreeErrors);
+	}
+	if (HeapErrors > 0)
+	{
+		UE_LOG(LogTraceServices, Error, TEXT("[MemAlloc] HEAP event errors: %u"), HeapErrors);
+	}
+	if (MiscErrors > 0)
+	{
+		UE_LOG(LogTraceServices, Error, TEXT("[MemAlloc] Other errors: %u"), MiscErrors);
+	}
+	if (TagTracker.GetNumErrors() > 0)
+	{
+		UE_LOG(LogTraceServices, Error, TEXT("[MemAlloc] TagTracker errors: %u"), TagTracker.GetNumErrors());
+	}
+
+	UE_LOG(LogTraceServices, Log, TEXT("[MemAlloc] Analysis Completed"));
 }
 
 ////////////////////////////////////////////////////////////////////////////////////////////////////
 
 void FAllocationsProvider::EnumerateRootHeaps(TFunctionRef<void(HeapId Id, const FHeapSpec&)> Callback) const
 {
-	Algo::ForEachIf(HeapSpecs, [](const FHeapSpec& Spec) { return Spec.Parent == nullptr && Spec.Name != nullptr; },
-	                [&](const FHeapSpec& Spec) { Callback(Spec.Id, Spec); });
+	for (const FHeapSpec& Spec : HeapSpecs)
+	{
+		if (Spec.Parent == nullptr && Spec.Name != nullptr)
+		{
+			Callback(Spec.Id, Spec);
+		}
+	}
 }
 
 ////////////////////////////////////////////////////////////////////////////////////////////////////
@@ -1904,6 +1958,7 @@ const IAllocationsProvider* ReadAllocationsProvider(const IAnalysisSession& Sess
 
 #undef INSIGHTS_SLOW_CHECK
 #undef INSIGHTS_DEBUG_WATCH
+#undef INSIGHTS_DEBUG_WATCH_FOUND
 #undef INSIGHTS_USE_SHORT_LIVING_ALLOCS
 #undef INSIGHTS_USE_LAST_ALLOC
 #undef INSIGHTS_VALIDATE_ALLOC_EVENTS
diff --git a/Engine/Source/Developer/TraceServices/Private/Model/AllocationsProvider.h b/Engine/Source/Developer/TraceServices/Private/Model/AllocationsProvider.h
--- a/Engine/Source/Developer/TraceServices/Private/Model/AllocationsProvider.h
+++ b/Engine/Source/Developer/TraceServices/Private/Model/AllocationsProvider.h
@@ -42,7 +42,7 @@ private:
 	static constexpr uint32 TrackerIdShift = 24;
 	static constexpr uint32 TrackerIdMask = 0xFF000000;
 	static constexpr uint32 PtrTagMask = 0x80000000;
-	
+
 	struct ThreadState
 	{
 		TArray<TagIdType> TagStack;
@@ -65,6 +65,8 @@ public:
 	void PopTagFromPtr(uint32 ThreadId, uint8 Tracker);
 	bool HasTagFromPtrScope(uint32 ThreadId, uint8 Tracker) const;
 
+	uint32 GetNumErrors() const { return NumErrors; }
+
 private:
 	inline uint32 GetTrackerThreadId(uint32 ThreadId, uint8 Tracker) const
 	{
@@ -93,14 +95,16 @@ struct FAllocationItem
 	uint32 GetAlignment() const { return UnpackAlignment(SizeAndAlignment); }
 	bool IsHeap() const { return EnumHasAnyFlags(Flags, EMemoryTraceHeapAllocationFlags::Heap); }
 
+	uint64 Address;
+	uint64 SizeAndAlignment; // (Alignment << AlignmentShift) | Size
 	uint32 StartEventIndex;
 	uint32 EndEventIndex;
 	double StartTime;
 	double EndTime;
-	uint64 Owner;
-	uint64 Address;
-	uint64 SizeAndAlignment; // (Alignment << AlignmentShift) | Size
-	mutable const FCallstack* Callstack;
+	uint32 CallstackId;
+	uint32 FreeCallstackId;
+	uint32 MetadataId;
+	uint32 Reserved;
 	TagIdType Tag;
 	uint8 RootHeap;
 	EMemoryTraceHeapAllocationFlags Flags;
@@ -211,7 +215,7 @@ public:
 
 	uint32 Num() const { return TotalAllocCount; }
 	uint32 PeakCount() const { return MaxAllocCount; }
-	
+
 	// Finds the allocation with specified address.
 	// Returns the found allocation or nullptr if not found.
 	FORCEINLINE FAllocationItem* FindRef(uint64 Address);
@@ -324,11 +328,11 @@ public:
 
 	void EditHeapSpec(HeapId Id, HeapId ParentId, const FStringView& Name, EMemoryTraceHeapFlags Flags);
 
-	void EditAlloc(double Time, uint64 Owner, uint64 Address, uint64 Size, uint32 Alignment, uint32 ThreadId, uint8 Tracker, HeapId RootHeap);
-	void EditFree(double Time, uint64 Address, HeapId RootHeap);
+	void EditAlloc(double Time, uint32 CallstackId, uint64 Address, uint64 Size, uint32 Alignment, HeapId RootHeap);
+	void EditFree(double Time, uint32 CallstackId, uint64 Address, HeapId RootHeap);
 	void EditMarkAllocationAsHeap(double Time, uint64 Address, HeapId Heap, EMemoryTraceHeapAllocationFlags Flags);
 	void EditUnmarkAllocationAsHeap(double Time, uint64 Address, HeapId Heap);
-	
+
 	void EditAddTagSpec(TagIdType Tag, TagIdType ParentTag, const TCHAR* Display) { EditAccessCheck(); TagTracker.AddTagSpec(Tag, ParentTag, Display); }
 	void EditPushTag(uint32 ThreadId, uint8 Tracker, TagIdType Tag);
 	void EditPopTag(uint32 ThreadId, uint8 Tracker);
@@ -337,6 +341,12 @@ public:
 
 	void EditOnAnalysisCompleted(double Time);
 
+	void SetCurrentThreadId(uint32 InThreadId, uint32 InSystemThreadId)
+	{
+		CurrentTraceThreadId = InThreadId;
+		CurrentSystemThreadId = InSystemThreadId;
+	}
+
 	//////////////////////////////////////////////////
 
 private:
@@ -351,16 +361,20 @@ private:
 
 	mutable FAllocationsProviderLock Lock;
 
-	//Number of supported root heaps
+	// Number of supported root heaps
 	constexpr static uint8 MaxRootHeaps = 16;
-	
+
 	double InitTime = 0;
 	uint8 MinAlignment = 0;
 	uint8 SizeShift = 0;
 	uint8 SummarySizeShift = 0;
 	bool bInitialized = false;
 
+	uint32 CurrentTraceThreadId = 0;
+	uint32 CurrentSystemThreadId = 0;
+
 	FTagTracker TagTracker;
+	uint8 CurrentTracker = 0;
 
 	uint32 EventIndex[MaxRootHeaps];
 
@@ -401,7 +415,7 @@ private:
 	TPagedArray<uint32> MaxLiveAllocationsTimeline;
 	TPagedArray<uint32> AllocEventsTimeline;
 	TPagedArray<uint32> FreeEventsTimeline;
-	
+
 	TArray<FHeapSpec> HeapSpecs;
 };
 
diff --git a/Engine/Source/Developer/TraceServices/Private/Model/AllocationsQuery.cpp b/Engine/Source/Developer/TraceServices/Private/Model/AllocationsQuery.cpp
--- a/Engine/Source/Developer/TraceServices/Private/Model/AllocationsQuery.cpp
+++ b/Engine/Source/Developer/TraceServices/Private/Model/AllocationsQuery.cpp
@@ -118,7 +118,6 @@ void FAllocationsQuery::Run()
 		UE_LOG(LogTraceServices, Log, TEXT("[MemAlloc] Enqueue %u live allocs..."), NumLiveAllocs);
 		TotalAllocationCount += NumLiveAllocs;
 
-		QueryCallstacks(LiveAllocsResult);
 		Results.Enqueue(LiveAllocsResult);
 	}
 	else
@@ -170,7 +169,6 @@ void FAllocationsQuery::Run()
 			UE_LOG(LogTraceServices, Log, TEXT("[MemAlloc] Enqueue %u allocs..."), NumAllocs);
 			TotalAllocationCount += NumAllocs;
 
-			QueryCallstacks(CellResult);
 			Results.Enqueue(CellResult);
 		}
 		else
@@ -402,20 +400,4 @@ void FAllocationsQuery::QueryLiveAllocs(TArray<const FAllocationItem*>& OutAlloc
 
 ////////////////////////////////////////////////////////////////////////////////////////////////////
 
-void FAllocationsQuery::QueryCallstacks(FAllocationsImpl* Result) const
-{
-	for (auto Item : Result->Items)
-	{
-		// Callstacks could have been resolved in previous queries,
-		// check before querying again
-		if (!Item->Callstack)
-		{
-			Item->Callstack = CallstacksProvider.GetCallstack(Item->Owner);
-			check(Item->Callstack != nullptr);
-		}
-	}
-}
-
-////////////////////////////////////////////////////////////////////////////////////////////////////
-
 } // namespace TraceServices
diff --git a/Engine/Source/Developer/TraceServices/Private/Model/AllocationsQuery.h b/Engine/Source/Developer/TraceServices/Private/Model/AllocationsQuery.h
--- a/Engine/Source/Developer/TraceServices/Private/Model/AllocationsQuery.h
+++ b/Engine/Source/Developer/TraceServices/Private/Model/AllocationsQuery.h
@@ -34,7 +34,6 @@ public:
 private:
 	void Run();
 	void QueryLiveAllocs(TArray<const FAllocationItem*>& OutAllocs) const;
-	void QueryCallstacks(FAllocationsImpl* Result) const;
 
 private:
 	const FAllocationsProvider& AllocationsProvider;
diff --git a/Engine/Source/Developer/TraceServices/Private/Model/CallstacksProvider.cpp b/Engine/Source/Developer/TraceServices/Private/Model/CallstacksProvider.cpp
--- a/Engine/Source/Developer/TraceServices/Private/Model/CallstacksProvider.cpp
+++ b/Engine/Source/Developer/TraceServices/Private/Model/CallstacksProvider.cpp
@@ -15,10 +15,9 @@ namespace TraceServices
 static const FResolvedSymbol GNeverResolveSymbol(ESymbolQueryResult::NotLoaded, nullptr, nullptr, nullptr, 0);
 static const FResolvedSymbol GNotFoundSymbol(ESymbolQueryResult::NotFound, TEXT("Unknown"), nullptr, nullptr, 0);
 static const FResolvedSymbol GNoSymbol(ESymbolQueryResult::NotFound, TEXT("No callstack recorded"), nullptr, nullptr, 0);
-static constexpr FStackFrame GNotFoundStackFrame = { 0, &GNotFoundSymbol};
-static constexpr FStackFrame GNoStackFrame = {0, &GNoSymbol};	
+static constexpr FStackFrame GNotFoundStackFrame = { 0, &GNotFoundSymbol };
+static constexpr FStackFrame GNoStackFrame = { 0, &GNoSymbol };
 static const FCallstack GNotFoundCallstack(&GNotFoundStackFrame, 1);
-static const FCallstack GNoCallstack(&GNoStackFrame, 1);
 
 /////////////////////////////////////////////////////////////////////
 #ifdef TRACE_CALLSTACK_STATS
@@ -37,12 +36,15 @@ FCallstacksProvider::FCallstacksProvider(IAnalysisSession& InSession)
 	, Callstacks(InSession.GetLinearAllocator(), CallstacksPerPage)
 	, Frames(InSession.GetLinearAllocator(), FramesPerPage)
 {
+	// Let the first callstack to be an empty/undefined callstack.
+	FCallstack& FirstCallstack = Callstacks.PushBack();
+	FirstCallstack.Init(&GNoStackFrame, 1);
 }
 
 /////////////////////////////////////////////////////////////////////
-void FCallstacksProvider::AddCallstack(uint64 InCallstackId, const uint64* InFrames, uint8 InFrameCount)
+void FCallstacksProvider::AddCallstack(uint32 InCallstackId, const uint64* InFrames, uint8 InFrameCount)
 {
-	if (InFrameCount == 0)
+	if (InCallstackId == 0)
 	{
 		return;
 	}
@@ -53,7 +55,7 @@ void FCallstacksProvider::AddCallstack(uint64 InCallstackId, const uint64* InFra
 	GCallstackStats.FrameCountHistogram[InFrameCount]++;
 #endif
 
-	// The module provider is created on the fly so we want to cache it 
+	// The module provider is created on the fly so we want to cache it
 	// once it's available. Note that the module provider is conditionally
 	// created so EditProvider() may return a null pointer.
 	if (!ModuleProvider)
@@ -61,25 +63,28 @@ void FCallstacksProvider::AddCallstack(uint64 InCallstackId, const uint64* InFra
 		ModuleProvider = Session.EditProvider<IModuleProvider>(GetModuleProviderName());
 	}
 
-	// Make sure all the frames fit on one page by appending dummy entries.
-	const uint64 PageHeadroom = Frames.GetPageSize() - (Frames.Num() % Frames.GetPageSize());
-	if (PageHeadroom < InFrameCount)
+	if (InFrameCount > 0)
 	{
-		FRWScopeLock WriteLock(EntriesLock, SLT_Write);
-		uint64 EntriesToAdd = PageHeadroom + 1; // Fill page and allocate one on next
-		do { Frames.PushBack(); } while (--EntriesToAdd);
+		// Make sure all the frames fit on one page by appending dummy entries.
+		const uint64 PageHeadroom = Frames.GetPageSize() - (Frames.Num() % Frames.GetPageSize());
+		if (PageHeadroom < InFrameCount)
+		{
+			FRWScopeLock WriteLock(EntriesLock, SLT_Write);
+			uint64 EntriesToAdd = PageHeadroom + 1; // Fill page and allocate one on next
+			do { Frames.PushBack(); } while (--EntriesToAdd);
+		}
 	}
 
-	// Append the incoming frames
+	// Append the incoming frames.
 	const uint64 FirstFrame = Frames.Num();
 	for (uint32 FrameIdx = 0; FrameIdx < InFrameCount; ++FrameIdx)
 	{
 		FStackFrame& F = Frames.PushBack();
 		F.Addr = InFrames[FrameIdx];
-		
+
 		if (ModuleProvider)
 		{
-			// This will return immediately. The result will be empty if the symbol 
+			// This will return immediately. The result will be empty if the symbol
 			// has not been encountered before, and resolution has been queued up.
 			F.Symbol = ModuleProvider->GetSymbol(InFrames[FrameIdx]);
 		}
@@ -91,29 +96,73 @@ void FCallstacksProvider::AddCallstack(uint64 InCallstackId, const uint64* InFra
 
 	{
 		FRWScopeLock WriteLock(EntriesLock, SLT_Write);
-		FCallstack* Callstack = &Callstacks.EmplaceBack(&Frames[FirstFrame], InFrameCount);
-		CallstackEntries.Add(InCallstackId, Callstack);
+		FCallstack* Callstack = nullptr;
+		if (InCallstackId < Callstacks.Num())
+		{
+			Callstack = &Callstacks[InCallstackId];
+		}
+		else
+		{
+			while (InCallstackId >= Callstacks.Num())
+			{
+				Callstack = &Callstacks.PushBack();
+			}
+		}
+		check(Callstack->Num() == 0); // adding same callstack id twice?
+		Callstack->Init(&Frames[FirstFrame], InFrameCount);
 	}
 }
 
 /////////////////////////////////////////////////////////////////////
-void FCallstacksProvider::AddCallstack(uint32 InCallstackId, const uint64* InFrames, uint8 InFrameCount)
+uint32 FCallstacksProvider::AddCallstackWithHash(uint64 InCallstackHash, const uint64* InFrames, uint8 InFrameCount)
+{
+	if (InCallstackHash == 0)
+	{
+		return 0;
+	}
+
+	uint32 CallstackId;
+	{
+		FRWScopeLock WriteLock(EntriesLock, SLT_Write);
+		CallstackId = Callstacks.Num();
+		CallstackMap.Add(InCallstackHash, CallstackId);
+	}
+
+	AddCallstack(CallstackId, InFrames, InFrameCount);
+	return CallstackId;
+}
+
+/////////////////////////////////////////////////////////////////////
+uint32 FCallstacksProvider::GetCallstackIdForHash(uint64 InCallstackHash) const
 {
-	// In order to maintain backward compatibility to older format where we sent the runtime hash value
-	// we keep the callstack entry key as uint64.
-	AddCallstack(uint64(InCallstackId), InFrames, InFrameCount);
+	if (InCallstackHash == 0)
+	{
+		return 0;
+	}
+	FRWScopeLock ReadLock(EntriesLock, SLT_ReadOnly);
+	const uint32* CallstackIdPtr = CallstackMap.Find(InCallstackHash);
+	if (CallstackIdPtr)
+	{
+		return *CallstackIdPtr;
+	}
+	else
+	{
+		return 0;
+	}
 }
 
 /////////////////////////////////////////////////////////////////////
 const FCallstack* FCallstacksProvider::GetCallstack(uint64 CallstackId) const
 {
 	FRWScopeLock ReadLock(EntriesLock, SLT_ReadOnly);
-	if (CallstackId == 0)
+	if (CallstackId < Callstacks.Num())
 	{
-		return &GNoCallstack;
+		return &Callstacks[CallstackId];
+	}
+	else
+	{
+		return &GNotFoundCallstack;
 	}
-	const FCallstack* const* FindResult = CallstackEntries.Find(CallstackId);
-	return FindResult ? *FindResult : &GNotFoundCallstack;
 }
 
 /////////////////////////////////////////////////////////////////////
@@ -125,8 +174,14 @@ void FCallstacksProvider::GetCallstacks(const TArrayView<uint64>& CallstackIds,
 	FRWScopeLock ReadLock(EntriesLock, SLT_ReadOnly);
 	for (uint64 CallstackId : CallstackIds)
 	{
-		const FCallstack* const* FindResult = CallstackEntries.Find(CallstackId);
-		OutCallstacks[OutIdx] = FindResult ? *FindResult : &GNotFoundCallstack;
+		if (CallstackId < Callstacks.Num())
+		{
+			OutCallstacks[OutIdx] = &Callstacks[CallstackId];
+		}
+		else
+		{
+			OutCallstacks[OutIdx] = &GNotFoundCallstack;
+		}
 		OutIdx++;
 	}
 }
diff --git a/Engine/Source/Developer/TraceServices/Private/Model/CallstacksProvider.h b/Engine/Source/Developer/TraceServices/Private/Model/CallstacksProvider.h
--- a/Engine/Source/Developer/TraceServices/Private/Model/CallstacksProvider.h
+++ b/Engine/Source/Developer/TraceServices/Private/Model/CallstacksProvider.h
@@ -20,10 +20,13 @@ public:
 	explicit FCallstacksProvider(IAnalysisSession& Session);
 	virtual ~FCallstacksProvider() {}
 
-	const FCallstack*	GetCallstack(uint64 CallstackId) const override;
-	void				GetCallstacks(const TArrayView<uint64>& CallstackIds, FCallstack const** OutCallstacks) const override;
-	void				AddCallstack(uint64 CallstackId, const uint64* Frames, uint8 FrameCount);
-	void				AddCallstack(uint32 CallstackId, const uint64* Frames, uint8 FrameCount);
+	const FCallstack* GetCallstack(uint64 CallstackId) const override;
+	void GetCallstacks(const TArrayView<uint64>& CallstackIds, FCallstack const** OutCallstacks) const override;
+	void AddCallstack(uint32 CallstackId, const uint64* Frames, uint8 FrameCount);
+
+	// Backward compatibility with legacy memory trace format (5.0-EA).
+	uint32 AddCallstackWithHash(uint64 CallstackHash, const uint64* Frames, uint8 FrameCount);
+	uint32 GetCallstackIdForHash(uint64 CallstackHash) const override;
 
 private:
 	enum
@@ -35,9 +38,9 @@ private:
 	mutable FRWLock					EntriesLock;
 	IAnalysisSession&				Session;
 	mutable IModuleProvider*		ModuleProvider;
-	TMap<uint64,const FCallstack*>	CallstackEntries;
-	TPagedArray<FCallstack>			Callstacks;
+	TPagedArray<FCallstack>			Callstacks; // CallstackId is an index in this array; Callstacks[0] is an empty callstack
 	TPagedArray<FStackFrame>		Frames;
+	TMap<uint64, uint32>			CallstackMap; // (CallstackHash --> CallstackId) map, used for backward compatibility with legacy trace format
 };
 
 } // namespace TraceServices
diff --git a/Engine/Source/Developer/TraceServices/Public/TraceServices/Model/AllocationsProvider.h b/Engine/Source/Developer/TraceServices/Public/TraceServices/Model/AllocationsProvider.h
--- a/Engine/Source/Developer/TraceServices/Public/TraceServices/Model/AllocationsProvider.h
+++ b/Engine/Source/Developer/TraceServices/Public/TraceServices/Model/AllocationsProvider.h
@@ -99,7 +99,9 @@ public:
 		uint64 GetAddress() const;
 		uint64 GetSize() const;
 		uint32 GetAlignment() const;
-		const struct FCallstack* GetCallstack() const;
+		uint32 GetCallstackId() const;
+		uint32 GetFreeCallstackId() const;
+		uint32 GetMetadataId() const;
 		TagIdType GetTag() const;
 		HeapId GetRootHeap() const;
 		bool IsHeap() const;
diff --git a/Engine/Source/Developer/TraceServices/Public/TraceServices/Model/Callstack.h b/Engine/Source/Developer/TraceServices/Public/TraceServices/Model/Callstack.h
--- a/Engine/Source/Developer/TraceServices/Public/TraceServices/Model/Callstack.h
+++ b/Engine/Source/Developer/TraceServices/Public/TraceServices/Model/Callstack.h
@@ -19,14 +19,19 @@ struct TRACESERVICES_API FStackFrame
 /////////////////////////////////////////////////////////////////////
 struct TRACESERVICES_API FCallstack
 {
+	/** Creates an empty callstack. */
+	FCallstack();
+	/** Creates an callstack initialised with a certain number of stack frames. */
 	FCallstack(const FStackFrame* FirstEntry, uint8 FrameCount);
-	/** Get the number of stack frames in callstack. */
+	/** Initialises the callstack with a certain number of stack frames. */
+	void Init(const FStackFrame* FirstEntry, uint8 FrameCount);
+	/** Gets the number of stack frames in callstack. */
 	uint32 Num() const;
 	/** Gets the address at a given stack depth. */
 	uint64 Addr(uint8 Depth) const;
 	/** Gets the cached symbol name at a given stack depth. */
 	const TCHAR* Name(uint8 Depth) const;
-	/** Gets the entire frame at given depth */
+	/** Gets the entire frame at given depth. */
 	const FStackFrame* Frame(uint8 Depth) const;
 
 private:
@@ -50,7 +55,7 @@ public:
 
 	/**
 	  * Queries a callstack id.
-	  * @param CallstackId Callstack id to query
+	  * @param CallstackId Callstack id to query.
 	  * @return Callstack information. If id is not found a callstack with zero stack depth is returned.
 	  */
 	virtual const FCallstack* GetCallstack(uint64 CallstackId) const = 0;
@@ -61,6 +66,13 @@ public:
 	  * @param OutCallstacks Output list of callstacks. Caller is responsible for allocating space according to CallstackIds length.
 	  */
 	virtual void GetCallstacks(const TArrayView<uint64>& CallstackIds, FCallstack const** OutCallstacks) const = 0;
+
+	/**
+	  * Gets a callstack id for a registerd callstack hash.
+	  * @param CallstackHash The callstack hash.
+	  * @returns The callstack id.
+	  */
+	virtual uint32 GetCallstackIdForHash(uint64 CallstackHash) const = 0;
 };
 
 /////////////////////////////////////////////////////////////////////
@@ -68,8 +80,20 @@ public:
 TRACESERVICES_API FName GetCallstacksProviderName();
 TRACESERVICES_API const ICallstacksProvider* ReadCallstacksProvider(const IAnalysisSession& Session);
 
+/////////////////////////////////////////////////////////////////////
+inline FCallstack::FCallstack()
+: CallstackLenIndex(0)
+{
+}
+
 /////////////////////////////////////////////////////////////////////
 inline FCallstack::FCallstack(const FStackFrame* InFirstFrame, uint8 InFrameCount)
+{
+	Init(InFirstFrame, InFrameCount);
+}
+
+/////////////////////////////////////////////////////////////////////
+inline void FCallstack::Init(const FStackFrame* InFirstFrame, uint8 InFrameCount)
 {
 	check((uint64(InFirstFrame) & EntryLenMask) == 0);
 	CallstackLenIndex = (uint64(InFrameCount) << EntryLenShift) | (~EntryLenMask & uint64(InFirstFrame));
