[Insights]
  - Changed CallstackId to 32 bits also for CallstacksProvider API.
  - Also fixed static analysis warning (false positive).

diff --git a/Engine/Source/Developer/TraceServices/Private/Model/CallstacksProvider.cpp b/Engine/Source/Developer/TraceServices/Private/Model/CallstacksProvider.cpp
--- a/Engine/Source/Developer/TraceServices/Private/Model/CallstacksProvider.cpp
+++ b/Engine/Source/Developer/TraceServices/Private/Model/CallstacksProvider.cpp
@@ -108,7 +108,7 @@ void FCallstacksProvider::AddCallstack(uint32 InCallstackId, const uint64* InFra
 				Callstack = &Callstacks.PushBack();
 			}
 		}
-		check(Callstack->Num() == 0); // adding same callstack id twice?
+		check(Callstack && Callstack->Num() == 0); // adding same callstack id twice?
 		Callstack->Init(&Frames[FirstFrame], InFrameCount);
 	}
 }
@@ -152,7 +152,7 @@ uint32 FCallstacksProvider::GetCallstackIdForHash(uint64 InCallstackHash) const
 }
 
 /////////////////////////////////////////////////////////////////////
-const FCallstack* FCallstacksProvider::GetCallstack(uint64 CallstackId) const
+const FCallstack* FCallstacksProvider::GetCallstack(uint32 CallstackId) const
 {
 	FRWScopeLock ReadLock(EntriesLock, SLT_ReadOnly);
 	if (CallstackId < Callstacks.Num())
@@ -166,7 +166,7 @@ const FCallstack* FCallstacksProvider::GetCallstack(uint64 CallstackId) const
 }
 
 /////////////////////////////////////////////////////////////////////
-void FCallstacksProvider::GetCallstacks(const TArrayView<uint64>& CallstackIds, FCallstack const** OutCallstacks) const
+void FCallstacksProvider::GetCallstacks(const TArrayView<uint32>& CallstackIds, FCallstack const** OutCallstacks) const
 {
 	uint64 OutIdx(0);
 	check(OutCallstacks != nullptr);
diff --git a/Engine/Source/Developer/TraceServices/Private/Model/CallstacksProvider.h b/Engine/Source/Developer/TraceServices/Private/Model/CallstacksProvider.h
--- a/Engine/Source/Developer/TraceServices/Private/Model/CallstacksProvider.h
+++ b/Engine/Source/Developer/TraceServices/Private/Model/CallstacksProvider.h
@@ -20,8 +20,8 @@ public:
 	explicit FCallstacksProvider(IAnalysisSession& Session);
 	virtual ~FCallstacksProvider() {}
 
-	const FCallstack* GetCallstack(uint64 CallstackId) const override;
-	void GetCallstacks(const TArrayView<uint64>& CallstackIds, FCallstack const** OutCallstacks) const override;
+	const FCallstack* GetCallstack(uint32 CallstackId) const override;
+	void GetCallstacks(const TArrayView<uint32>& CallstackIds, FCallstack const** OutCallstacks) const override;
 	void AddCallstack(uint32 CallstackId, const uint64* Frames, uint8 FrameCount);
 
 	// Backward compatibility with legacy memory trace format (5.0-EA).
diff --git a/Engine/Source/Developer/TraceServices/Public/TraceServices/Model/Callstack.h b/Engine/Source/Developer/TraceServices/Public/TraceServices/Model/Callstack.h
--- a/Engine/Source/Developer/TraceServices/Public/TraceServices/Model/Callstack.h
+++ b/Engine/Source/Developer/TraceServices/Public/TraceServices/Model/Callstack.h
@@ -56,16 +56,16 @@ public:
 	/**
 	  * Queries a callstack id.
 	  * @param CallstackId Callstack id to query.
-	  * @return Callstack information. If id is not found a callstack with zero stack depth is returned.
+	  * @return Callstack information. If id is not found a special "Unknown" callstack is returned.
 	  */
-	virtual const FCallstack* GetCallstack(uint64 CallstackId) const = 0;
+	virtual const FCallstack* GetCallstack(uint32 CallstackId) const = 0;
 
 	/**
 	  * Queries a set of callstack ids.
 	  * @param CallstackIds List of callstack ids to query
 	  * @param OutCallstacks Output list of callstacks. Caller is responsible for allocating space according to CallstackIds length.
 	  */
-	virtual void GetCallstacks(const TArrayView<uint64>& CallstackIds, FCallstack const** OutCallstacks) const = 0;
+	virtual void GetCallstacks(const TArrayView<uint32>& CallstackIds, FCallstack const** OutCallstacks) const = 0;
 
 	/**
 	  * Gets a callstack id for a registerd callstack hash.
