Stats Trace: Add stat group tracking

diff --git a/Engine/Source/Developer/TraceServices/Private/Analyzers/StatsTraceAnalysis.cpp b/Engine/Source/Developer/TraceServices/Private/Analyzers/StatsTraceAnalysis.cpp
--- a/Engine/Source/Developer/TraceServices/Private/Analyzers/StatsTraceAnalysis.cpp
+++ b/Engine/Source/Developer/TraceServices/Private/Analyzers/StatsTraceAnalysis.cpp
@@ -29,118 +29,130 @@ bool FStatsAnalyzer::OnEvent(uint16 RouteId, EStyle Style, const FOnEventContext
 	const auto& EventData = Context.EventData;
 	switch (RouteId)
 	{
-	case RouteId_Spec:
-	{
-		uint32 StatId = EventData.GetValue<uint32>("Id");
-		IEditableCounter* Counter = CountersMap.FindRef(StatId);
-		if (!Counter)
-		{
-			Counter = CounterProvider.CreateCounter();
-			CountersMap.Add(StatId, Counter);
-		}
-
-		FString Name;
-		FString Description;
-		if (EventData.GetString("Name", Name))
-		{
-			EventData.GetString("Description", Description);
-		}
-		else
-		{
-			Name = reinterpret_cast<const ANSICHAR*>(EventData.GetAttachment());
-			Description = reinterpret_cast<const TCHAR*>(EventData.GetAttachment() + Name.Len() + 1);
-		}
-
-
-		ECounterDisplayHint DisplayHint = CounterDisplayHint_None;
-		if (EventData.GetValue<bool>("IsMemory"))
-		{
-			DisplayHint = CounterDisplayHint_Memory;
-		}
-		Counter->SetName(Session.StoreString(*Name));
-		Counter->SetDescription(Session.StoreString(*Description));
-		Counter->SetIsFloatingPoint(EventData.GetValue<bool>("IsFloatingPoint"));
-		Counter->SetDisplayHint(DisplayHint);
-		break;
-	}
-	case RouteId_EventBatch:
-	{
-		uint32 ThreadId = FTraceAnalyzerUtils::GetThreadIdField(Context);
-		TSharedRef<FThreadState> ThreadState = GetThreadState(ThreadId);
-		TArrayView<const uint8> DataView = FTraceAnalyzerUtils::LegacyAttachmentArray("Data", Context);
-		uint64 BufferSize = DataView.Num();
-		const uint8* BufferPtr = DataView.GetData();
-		const uint8* BufferEnd = BufferPtr + BufferSize;
-		while (BufferPtr < BufferEnd)
+		case RouteId_Spec:
 		{
-			enum EOpType
-			{
-				Increment = 0,
-				Decrement = 1,
-				AddInteger = 2,
-				SetInteger = 3,
-				AddFloat = 4,
-				SetFloat = 5,
-			};
-
-			uint64 DecodedIdAndOp = FTraceAnalyzerUtils::Decode7bit(BufferPtr);
-			uint32 StatId = DecodedIdAndOp >> 3;
+			uint32 StatId = EventData.GetValue<uint32>("Id");
 			IEditableCounter* Counter = CountersMap.FindRef(StatId);
 			if (!Counter)
 			{
 				Counter = CounterProvider.CreateCounter();
 				CountersMap.Add(StatId, Counter);
 			}
-			uint8 Op = DecodedIdAndOp & 0x7;
-			uint64 CycleDiff = FTraceAnalyzerUtils::Decode7bit(BufferPtr);
-			uint64 Cycle = ThreadState->LastCycle + CycleDiff;
-			double Time = Context.EventTime.AsSeconds(Cycle);
-			ThreadState->LastCycle = Cycle;
-			switch (Op)
-			{
-			case Increment:
+
+			FString Name;
+			FString Description;
+			FString Group;
+			if (EventData.GetString("Name", Name))
 			{
-				Counter->AddValue(Time, int64(1));
-				break;
+				EventData.GetString("Description", Description);
+				EventData.GetString("Group", Group);
 			}
-			case Decrement:
+			else
 			{
-				Counter->AddValue(Time, int64(-1));
-				break;
+				Name = reinterpret_cast<const ANSICHAR*>(EventData.GetAttachment());
+				Description = reinterpret_cast<const TCHAR*>(EventData.GetAttachment() + Name.Len() + 1);
 			}
-			case AddInteger:
+
+
+			ECounterDisplayHint DisplayHint = CounterDisplayHint_None;
+			if (EventData.GetValue<bool>("IsMemory"))
 			{
-				int64 Amount = FTraceAnalyzerUtils::DecodeZigZag(BufferPtr);
-				Counter->AddValue(Time, Amount);
-				break;
+				DisplayHint = CounterDisplayHint_Memory;
 			}
-			case SetInteger:
+
+			if (Name.IsEmpty())
 			{
-				int64 Value = FTraceAnalyzerUtils::DecodeZigZag(BufferPtr);
-				Counter->SetValue(Time, Value);
-				break;
+				UE_LOG(LogTraceServices, Warning, TEXT("Got an invalid/empty counter name from StatsTrace events."))
 			}
-			case AddFloat:
+				
+			Counter->SetName(Session.StoreString(*Name));
+			if (!Group.IsEmpty())
 			{
-				double Amount;
-				memcpy(&Amount, BufferPtr, sizeof(double));
-				BufferPtr += sizeof(double);
-				Counter->AddValue(Time, Amount);
-				break;
+				Counter->SetGroup(Session.StoreString(*Group));
 			}
-			case SetFloat:
+			Counter->SetDescription(Session.StoreString(*Description));
+			Counter->SetIsFloatingPoint(EventData.GetValue<bool>("IsFloatingPoint"));
+			Counter->SetDisplayHint(DisplayHint);
+			break;
+		}
+		case RouteId_EventBatch:
+		{
+			uint32 ThreadId = FTraceAnalyzerUtils::GetThreadIdField(Context);
+			TSharedRef<FThreadState> ThreadState = GetThreadState(ThreadId);
+			TArrayView<const uint8> DataView = FTraceAnalyzerUtils::LegacyAttachmentArray("Data", Context);
+			uint64 BufferSize = DataView.Num();
+			const uint8* BufferPtr = DataView.GetData();
+			const uint8* BufferEnd = BufferPtr + BufferSize;
+			while (BufferPtr < BufferEnd)
 			{
-				double Value;
-				memcpy(&Value, BufferPtr, sizeof(double));
-				BufferPtr += sizeof(double);
-				Counter->SetValue(Time, Value);
-				break;
-			}
+				enum EOpType
+				{
+					Increment = 0,
+					Decrement = 1,
+					AddInteger = 2,
+					SetInteger = 3,
+					AddFloat = 4,
+					SetFloat = 5,
+				};
+
+				uint64 DecodedIdAndOp = FTraceAnalyzerUtils::Decode7bit(BufferPtr);
+				uint32 StatId = DecodedIdAndOp >> 3;
+				IEditableCounter* Counter = CountersMap.FindRef(StatId);
+				if (!Counter)
+				{
+					Counter = CounterProvider.CreateCounter();
+					CountersMap.Add(StatId, Counter);
+				}
+				uint8 Op = DecodedIdAndOp & 0x7;
+				uint64 CycleDiff = FTraceAnalyzerUtils::Decode7bit(BufferPtr);
+				uint64 Cycle = ThreadState->LastCycle + CycleDiff;
+				double Time = Context.EventTime.AsSeconds(Cycle);
+				ThreadState->LastCycle = Cycle;
+				switch (Op)
+				{
+				case Increment:
+				{
+					Counter->AddValue(Time, int64(1));
+					break;
+				}
+				case Decrement:
+				{
+					Counter->AddValue(Time, int64(-1));
+					break;
+				}
+				case AddInteger:
+				{
+					int64 Amount = FTraceAnalyzerUtils::DecodeZigZag(BufferPtr);
+					Counter->AddValue(Time, Amount);
+					break;
+				}
+				case SetInteger:
+				{
+					int64 Value = FTraceAnalyzerUtils::DecodeZigZag(BufferPtr);
+					Counter->SetValue(Time, Value);
+					break;
+				}
+				case AddFloat:
+				{
+					double Amount;
+					memcpy(&Amount, BufferPtr, sizeof(double));
+					BufferPtr += sizeof(double);
+					Counter->AddValue(Time, Amount);
+					break;
+				}
+				case SetFloat:
+				{
+					double Value;
+					memcpy(&Value, BufferPtr, sizeof(double));
+					BufferPtr += sizeof(double);
+					Counter->SetValue(Time, Value);
+					break;
+				}
+				}
 			}
+			check(BufferPtr == BufferEnd);
+			break;
 		}
-		check(BufferPtr == BufferEnd);
-		break;
-	}
 	}
 
 	return true;
diff --git a/Engine/Source/Developer/TraceServices/Private/Model/CountersPrivate.h b/Engine/Source/Developer/TraceServices/Private/Model/CountersPrivate.h
--- a/Engine/Source/Developer/TraceServices/Private/Model/CountersPrivate.h
+++ b/Engine/Source/Developer/TraceServices/Private/Model/CountersPrivate.h
@@ -196,6 +196,8 @@ public:
 	FCounter(ILinearAllocator& Allocator, const TArray64<double>& FrameStartTimes);
 	virtual const TCHAR* GetName() const override { return Name; }
 	virtual void SetName(const TCHAR* InName) override { Name = InName; }
+	virtual const TCHAR* GetGroup() const override { return Group; }
+	virtual void SetGroup(const TCHAR* InGroup) override { Group = InGroup; }
 	virtual const TCHAR* GetDescription() const override { return Description; }
 	virtual void SetDescription(const TCHAR* InDescription) override { Description = InDescription; }
 	virtual bool IsFloatingPoint() const override { return bIsFloatingPoint; }
@@ -215,6 +217,7 @@ private:
 	TCounterData<int64> IntCounterData;
 	TCounterData<double> DoubleCounterData;
 	const TCHAR* Name = nullptr;
+	const TCHAR* Group = nullptr;
 	const TCHAR* Description = nullptr;
 	uint64 ModCount = 0;
 	ECounterDisplayHint DisplayHint = CounterDisplayHint_None;
diff --git a/Engine/Source/Developer/TraceServices/Private/Model/LoadTimeProfiler.cpp b/Engine/Source/Developer/TraceServices/Private/Model/LoadTimeProfiler.cpp
--- a/Engine/Source/Developer/TraceServices/Private/Model/LoadTimeProfiler.cpp
+++ b/Engine/Source/Developer/TraceServices/Private/Model/LoadTimeProfiler.cpp
@@ -28,6 +28,11 @@ const TCHAR* FLoadTimeProfilerProvider::FLoaderFrameCounter::GetName() const
 	return TEXT("");
 }
 
+const TCHAR* FLoadTimeProfilerProvider::FLoaderFrameCounter::GetGroup() const
+{
+	return nullptr;
+}
+
 const TCHAR* FLoadTimeProfilerProvider::FLoaderFrameCounter::GetDescription() const
 {
 	return TEXT("N/A");
diff --git a/Engine/Source/Developer/TraceServices/Private/Model/LoadTimeProfilerPrivate.h b/Engine/Source/Developer/TraceServices/Private/Model/LoadTimeProfilerPrivate.h
--- a/Engine/Source/Developer/TraceServices/Private/Model/LoadTimeProfilerPrivate.h
+++ b/Engine/Source/Developer/TraceServices/Private/Model/LoadTimeProfilerPrivate.h
@@ -87,6 +87,7 @@ private:
 		virtual ~FLoaderFrameCounter() {}
 
 		virtual const TCHAR* GetName() const override;
+		virtual const TCHAR* GetGroup() const override;
 		virtual const TCHAR* GetDescription() const override;
 		virtual bool IsFloatingPoint() const override;
 		virtual ECounterDisplayHint GetDisplayHint() const override;
diff --git a/Engine/Source/Developer/TraceServices/Public/TraceServices/Model/Counters.h b/Engine/Source/Developer/TraceServices/Public/TraceServices/Model/Counters.h
--- a/Engine/Source/Developer/TraceServices/Public/TraceServices/Model/Counters.h
+++ b/Engine/Source/Developer/TraceServices/Public/TraceServices/Model/Counters.h
@@ -20,6 +20,7 @@ public:
 	virtual ~ICounter() = default;
 
 	virtual const TCHAR* GetName() const = 0;
+	virtual const TCHAR* GetGroup() const = 0;
 	virtual const TCHAR* GetDescription() const = 0;
 	virtual bool IsFloatingPoint() const = 0;
 	virtual ECounterDisplayHint GetDisplayHint() const = 0;
@@ -34,6 +35,7 @@ public:
 	virtual ~IEditableCounter() = default;
 
 	virtual void SetName(const TCHAR* Name) = 0;
+	virtual void SetGroup(const TCHAR* Group) = 0;
 	virtual void SetDescription(const TCHAR* Description) = 0;
 	virtual void SetIsFloatingPoint(bool bIsFloatingPoint) = 0;
 	virtual void SetDisplayHint(ECounterDisplayHint DisplayHint) = 0;
diff --git a/Engine/Source/Runtime/Core/Private/Stats/Stats2.cpp b/Engine/Source/Runtime/Core/Private/Stats/Stats2.cpp
--- a/Engine/Source/Runtime/Core/Private/Stats/Stats2.cpp
+++ b/Engine/Source/Runtime/Core/Private/Stats/Stats2.cpp
@@ -619,7 +619,7 @@ public:
 		{
 			ANSICHAR NameBuffer[1024];
 			StatShortName.GetPlainANSIString(NameBuffer);
-			FStatsTrace::DeclareStat(Stat, NameBuffer, *StatDescription, InStatType == EStatDataType::ST_double, MemoryRegion != FPlatformMemory::MCR_Invalid, bShouldClearEveryFrame);
+			FStatsTrace::DeclareStat(Stat, NameBuffer, *StatDescription, InGroup, InStatType == EStatDataType::ST_double, MemoryRegion != FPlatformMemory::MCR_Invalid, bShouldClearEveryFrame);
 		}
 #endif
 		return TStatId(Result);
diff --git a/Engine/Source/Runtime/Core/Private/Stats/StatsTrace.cpp b/Engine/Source/Runtime/Core/Private/Stats/StatsTrace.cpp
--- a/Engine/Source/Runtime/Core/Private/Stats/StatsTrace.cpp
+++ b/Engine/Source/Runtime/Core/Private/Stats/StatsTrace.cpp
@@ -19,6 +19,7 @@ UE_TRACE_EVENT_BEGIN(Stats, Spec, NoSync|Important)
 	UE_TRACE_EVENT_FIELD(bool, ShouldClearEveryFrame)
 	UE_TRACE_EVENT_FIELD(UE::Trace::AnsiString, Name)
 	UE_TRACE_EVENT_FIELD(UE::Trace::AnsiString, Description)
+	UE_TRACE_EVENT_FIELD(UE::Trace::AnsiString, Group)
 UE_TRACE_EVENT_END()
 
 UE_TRACE_EVENT_BEGIN(Stats, EventBatch)
@@ -103,18 +104,20 @@ void FStatsTraceInternal::EndEncodeOp(FThreadState* ThreadState, uint8* BufferPt
 	ThreadState->BufferSize = uint16(BufferPtr - ThreadState->Buffer);
 }
 
-void FStatsTrace::DeclareStat(const FName& Stat, const ANSICHAR* Name, const TCHAR* Description, bool IsFloatingPoint, bool IsMemory, bool ShouldClearEveryFrame)
+void FStatsTrace::DeclareStat(const FName& Stat, const ANSICHAR* Name, const TCHAR* Description, const ANSICHAR* Group, bool IsFloatingPoint, bool IsMemory, bool ShouldClearEveryFrame)
 {
 	uint32 NameLen = FCStringAnsi::Strlen(Name);
 	uint32 DescriptionLen = FCString::Strlen(Description);
-
-	UE_TRACE_LOG(Stats, Spec, StatsChannel, NameLen + DescriptionLen)
+	uint32 GroupLen = FCStringAnsi::Strlen(Group);
+		
+	UE_TRACE_LOG(Stats, Spec, StatsChannel, NameLen + DescriptionLen + GroupLen)
 		<< Spec.Id(Stat.GetComparisonIndex().ToUnstableInt())
 		<< Spec.IsFloatingPoint(IsFloatingPoint)
 		<< Spec.IsMemory(IsMemory)
 		<< Spec.ShouldClearEveryFrame(ShouldClearEveryFrame)
 		<< Spec.Name(Name, NameLen)
-		<< Spec.Description(Description, DescriptionLen);
+		<< Spec.Description(Description, DescriptionLen)
+		<< Spec.Group(Group, GroupLen);
 }
 
 void FStatsTrace::Increment(const FName& Stat)
diff --git a/Engine/Source/Runtime/Core/Public/Stats/Stats2.h b/Engine/Source/Runtime/Core/Public/Stats/Stats2.h
--- a/Engine/Source/Runtime/Core/Public/Stats/Stats2.h
+++ b/Engine/Source/Runtime/Core/Public/Stats/Stats2.h
@@ -1749,6 +1749,7 @@ public:
 	/**
 	 * Returns a pointer to a bool (valid forever) that determines if this group is active
 	 * This should be CACHED. We will get a few calls from different stats and different threads and stuff, but once things are "warmed up", this should NEVER be called.
+	 * This function will also register any stats with the StatsTrace system 
 	 * @param InGroup, group to look up
 	 * @param InCategory, the category the group belongs to
 	 * @param bDefaultEnable, If this is the first time this group has been set up, this sets the default enable value for this group.
diff --git a/Engine/Source/Runtime/Core/Public/Stats/StatsTrace.h b/Engine/Source/Runtime/Core/Public/Stats/StatsTrace.h
--- a/Engine/Source/Runtime/Core/Public/Stats/StatsTrace.h
+++ b/Engine/Source/Runtime/Core/Public/Stats/StatsTrace.h
@@ -19,7 +19,7 @@ class FName;
 
 struct FStatsTrace
 {
-	CORE_API static void DeclareStat(const FName& Stat, const ANSICHAR* Name, const TCHAR* Description, bool IsFloatingPoint, bool IsMemory, bool ShouldClearEveryFrame);
+	CORE_API static void DeclareStat(const FName& Stat, const ANSICHAR* Name, const TCHAR* Description, const ANSICHAR* Group, bool IsFloatingPoint, bool IsMemory, bool ShouldClearEveryFrame);
 	CORE_API static void Increment(const FName& Stat);
 	CORE_API static void Decrement(const FName& Stat);
 	CORE_API static void Add(const FName& Stat, int64 Amount);
