[Insights]
  - Improved performance of memory analysis (processing of live allocations in AllocationsProvider) with up to 50% (i.e. twice as fast). This is achived by using a better hash for TMap for long living allocations. Also switched the short living allocations to use only the linked list, with a reduced size (64 allocs).
  - Delayed initialization of SystemMemory heap to when the allocations provider inits (to save memory usage when traces does not have mem alloc events).
  - Added option to enable an initial "reserve" number for TMap (one used for long living allocations). Benchmarks shows further performance increase when using large reserve numbers, but mem cost becomes too high. So kept this 0 (unset) by default.
  - Added error checking for very large mem traces (when total number of mem events could exceed 32 bit limit).

Benchmarks:
  * analysis of a 32s session: from 17.6s to 8.3s (-53%)
  * analysis of a 7min session: from 88s to 46s (-48%)

diff --git a/Engine/Source/Developer/TraceServices/Private/Model/AllocationsProvider.cpp b/Engine/Source/Developer/TraceServices/Private/Model/AllocationsProvider.cpp
--- a/Engine/Source/Developer/TraceServices/Private/Model/AllocationsProvider.cpp
+++ b/Engine/Source/Developer/TraceServices/Private/Model/AllocationsProvider.cpp
@@ -23,11 +23,16 @@
 //#define INSIGHTS_DEBUG_WATCH_FOUND return // log and ignore events
 //#define INSIGHTS_DEBUG_WATCH_FOUND UE_DEBUG_BREAK(); break
 
+// Initial reserved number for long living allocations.
+#define INSIGHTS_LLA_RESERVE 0
+//#define INSIGHTS_LLA_RESERVE (64 * 1024)
+
 // Use optimized path for the short living allocations.
 // If enabled, caches the short living allocations, as there is a very high chance to be freed in the next few "free" events.
 // ~66% of all allocs are expected to have an "event distance" < 64 events
 // ~70% of all allocs are expected to have an "event distance" < 512 events
 #define INSIGHTS_USE_SHORT_LIVING_ALLOCS 1
+#define INSIGHTS_SLA_USE_ADDRESS_MAP 0
 
 // Use optimized path for the last alloc.
 // If enabled, caches the last added alloc, as there is a high chance to be freed in the next "free" event.
@@ -381,6 +386,10 @@ IAllocationsProvider::FQueryResult IAllocationsProvider::FQueryStatus::NextResul
 
 FShortLivingAllocs::FShortLivingAllocs()
 {
+#if INSIGHTS_SLA_USE_ADDRESS_MAP
+	AddressMap.Reserve(MaxAllocCount);
+#endif
+
 	AllNodes = new FNode[MaxAllocCount];
 
 	// Build the "unused" simple linked list.
@@ -403,7 +412,9 @@ FShortLivingAllocs::~FShortLivingAllocs()
 
 void FShortLivingAllocs::Reset()
 {
+#if INSIGHTS_SLA_USE_ADDRESS_MAP
 	AddressMap.Reset();
+#endif
 
 	FNode* Node = LastAddedAllocNode;
 	while (Node != nullptr)
@@ -425,14 +436,13 @@ void FShortLivingAllocs::Reset()
 
 FAllocationItem* FShortLivingAllocs::FindRef(uint64 Address)
 {
-#if 1
+#if INSIGHTS_SLA_USE_ADDRESS_MAP
 	FNode** NodePtr = AddressMap.Find(Address);
 	return NodePtr ? (*NodePtr)->Alloc : nullptr;
 #else
 	// Search linearly the list of allocations, backward, starting with most recent one.
 	// As the probability of finding a match for a "free" event is higher for the recent allocs,
 	// this could actually be faster than doing a O(log n) search in AddressMap.
-	// TODO: This needs to be tested for various MaxAllocCount values.
 	FNode* Node = LastAddedAllocNode;
 	while (Node != nullptr)
 	{
@@ -462,13 +472,17 @@ FAllocationItem* FShortLivingAllocs::AddAndRemoveOldest(FAllocationItem* Alloc)
 
 		// Remove the oldest allocation.
 		FAllocationItem* RemovedAlloc = OldestAllocNode->Alloc;
+#if INSIGHTS_SLA_USE_ADDRESS_MAP
 		AddressMap.Remove(RemovedAlloc->Address);
+#endif
 		OldestAllocNode = OldestAllocNode->Next;
 		INSIGHTS_SLOW_CHECK(OldestAllocNode != nullptr);
 		OldestAllocNode->Prev = nullptr;
 
 		// Add the new node.
+#if INSIGHTS_SLA_USE_ADDRESS_MAP
 		AddressMap.Add(Alloc->Address, NewNode);
+#endif
 		NewNode->Alloc = Alloc;
 		NewNode->Next = nullptr;
 		NewNode->Prev = LastAddedAllocNode;
@@ -486,7 +500,9 @@ FAllocationItem* FShortLivingAllocs::AddAndRemoveOldest(FAllocationItem* Alloc)
 		FirstUnusedNode = FirstUnusedNode->Next;
 
 		// Add the new node.
+#if INSIGHTS_SLA_USE_ADDRESS_MAP
 		AddressMap.Add(Alloc->Address, NewNode);
+#endif
 		NewNode->Alloc = Alloc;
 		NewNode->Next = nullptr;
 		NewNode->Prev = LastAddedAllocNode;
@@ -508,41 +524,58 @@ FAllocationItem* FShortLivingAllocs::AddAndRemoveOldest(FAllocationItem* Alloc)
 
 FAllocationItem* FShortLivingAllocs::Remove(uint64 Address)
 {
-	FNode* RemovedNode;
-	if (AddressMap.RemoveAndCopyValue(Address, RemovedNode))
+#if INSIGHTS_SLA_USE_ADDRESS_MAP
+	FNode* Node = nullptr;
+	if (!AddressMap.RemoveAndCopyValue(Address, Node))
 	{
-		INSIGHTS_SLOW_CHECK(AllocCount > 0);
-		--AllocCount;
-
-		FAllocationItem* RemovedAlloc = RemovedNode->Alloc;
-		INSIGHTS_SLOW_CHECK(RemovedAlloc->Address == Address);
-
-		// Remove node.
-		if (RemovedNode == OldestAllocNode)
-		{
-			OldestAllocNode = OldestAllocNode->Next;
-		}
-		if (RemovedNode == LastAddedAllocNode)
-		{
-			LastAddedAllocNode = LastAddedAllocNode->Prev;
-		}
-		if (RemovedNode->Prev)
-		{
-			RemovedNode->Prev->Next = RemovedNode->Next;
-		}
-		if (RemovedNode->Next)
+		return nullptr;
+	}
+	INSIGHTS_SLOW_CHECK(Node && Node->Alloc && Node->Alloc->Address == Address);
+#else
+	// Search linearly the list of allocations, backward, starting with most recent one.
+	// As the probability of finding a match for a "free" event is higher for the recent allocs,
+	// this could actually be faster than doing a O(log n) search in AddressMap.
+	FNode* Node = LastAddedAllocNode;
+	while (Node != nullptr)
+	{
+		if (Node->Alloc->Address == Address)
 		{
-			RemovedNode->Next->Prev = RemovedNode->Prev;
+			break;
 		}
+		Node = Node->Prev;
+	}
+	if (!Node)
+	{
+		return nullptr;
+	}
+#endif // INSIGHTS_SLA_USE_ADDRESS_MAP
 
-		// Add the removed node to the "unused" list.
-		RemovedNode->Next = FirstUnusedNode;
-		FirstUnusedNode = RemovedNode;
+	INSIGHTS_SLOW_CHECK(AllocCount > 0);
+	--AllocCount;
 
-		return RemovedAlloc;
+	// Remove node.
+	if (Node == OldestAllocNode)
+	{
+		OldestAllocNode = OldestAllocNode->Next;
+	}
+	if (Node == LastAddedAllocNode)
+	{
+		LastAddedAllocNode = LastAddedAllocNode->Prev;
+	}
+	if (Node->Prev)
+	{
+		Node->Prev->Next = Node->Next;
+	}
+	if (Node->Next)
+	{
+		Node->Next->Prev = Node->Prev;
 	}
 
-	return nullptr;
+	// Add the removed node to the "unused" list.
+	Node->Next = FirstUnusedNode;
+	FirstUnusedNode = Node;
+
+	return Node->Alloc;
 }
 
 ////////////////////////////////////////////////////////////////////////////////////////////////////
@@ -722,6 +755,9 @@ FAllocationItem* FHeapAllocs::FindRange(uint64 Address) const
 
 FLiveAllocCollection::FLiveAllocCollection()
 {
+#if INSIGHTS_LLA_RESERVE
+	LongLivingAllocs.Reserve(INSIGHTS_LLA_RESERVE);
+#endif
 }
 
 ////////////////////////////////////////////////////////////////////////////////////////////////////
@@ -1001,13 +1037,7 @@ FAllocationsProvider::FAllocationsProvider(IAnalysisSession& InSession)
 	, AllocEventsTimeline(Session.GetLinearAllocator(), 1024)
 	, FreeEventsTimeline(Session.GetLinearAllocator(), 1024)
 {
-	FMemory::Memset(EventIndex, 0, sizeof(EventIndex));
-	FMemory::Memset(SbTree, 0, sizeof(SbTree));
-	FMemory::Memset(LiveAllocs, 0, sizeof(LiveAllocs));
 	HeapSpecs.AddZeroed(256);
-
-	// Create system root heap structures for backwards compatibility (before heap description events)
-	AddHeapSpec(EMemoryTraceRootHeap::SystemMemory, 0, TEXT("System memory"), EMemoryTraceHeapFlags::Root);
 }
 
 ////////////////////////////////////////////////////////////////////////////////////////////////////
@@ -1044,6 +1074,9 @@ void FAllocationsProvider::EditInit(double InTime, uint8 InMinAlignment)
 	InitTime = InTime;
 	MinAlignment = InMinAlignment;
 
+	// Create system root heap structures for backwards compatibility (before heap description events)
+	AddHeapSpec(EMemoryTraceRootHeap::SystemMemory, 0, TEXT("System memory"), EMemoryTraceHeapFlags::Root);
+
 	bInitialized = true;
 
 	AdvanceTimelines(InTime);
@@ -1188,7 +1221,15 @@ void FAllocationsProvider::EditAlloc(double Time, uint32 CallstackId, uint64 Add
 #endif
 
 	++AllocCount;
-	++EventIndex[RootHeap];
+	if (EventIndex[RootHeap] != ~0u)
+	{
+		++EventIndex[RootHeap];
+	}
+	else
+	{
+		UE_LOG(LogTraceServices, Error, TEXT("[MemAlloc] Too many events!"));
+		bInitialized = false; // ignore further events
+	}
 }
 
 ////////////////////////////////////////////////////////////////////////////////////////////////////
@@ -1276,7 +1317,15 @@ void FAllocationsProvider::EditFree(double Time, uint32 CallstackId, uint64 Addr
 	}
 
 	++FreeCount;
-	++EventIndex[RootHeap];
+	if (EventIndex[RootHeap] != ~0u)
+	{
+		++EventIndex[RootHeap];
+	}
+	else
+	{
+		UE_LOG(LogTraceServices, Error, TEXT("[MemAlloc] Too many events!"));
+		bInitialized = false; // ignore further events
+	}
 }
 
 ////////////////////////////////////////////////////////////////////////////////////////////////////
@@ -1488,7 +1537,7 @@ void FAllocationsProvider::EditPushTag(uint32 ThreadId, uint8 Tracker, TagIdType
 void FAllocationsProvider::EditPopTag(uint32 ThreadId, uint8 Tracker)
 {
 	EditAccessCheck();
-	
+
 	TagTracker.PopTag(ThreadId, Tracker);
 }
 
@@ -1588,6 +1637,8 @@ void FAllocationsProvider::EditOnAnalysisCompleted(double Time)
 		}
 	}
 
+	//TODO: shrink live allocs buffers
+
 	if (AllocErrors > 0)
 	{
 		UE_LOG(LogTraceServices, Error, TEXT("[MemAlloc] ALLOC event errors: %u"), AllocErrors);
@@ -1609,7 +1660,12 @@ void FAllocationsProvider::EditOnAnalysisCompleted(double Time)
 		UE_LOG(LogTraceServices, Error, TEXT("[MemAlloc] TagTracker errors: %u"), TagTracker.GetNumErrors());
 	}
 
-	UE_LOG(LogTraceServices, Log, TEXT("[MemAlloc] Analysis Completed"));
+	uint64 TotalEventCount = 0;
+	for (uint32 RootHeap = 0; RootHeap < MaxRootHeaps; ++RootHeap)
+	{
+		TotalEventCount += EventIndex[RootHeap];
+	}
+	UE_LOG(LogTraceServices, Log, TEXT("[MemAlloc] Analysis Completed (%llu events, %llu allocs, %llu frees)"), TotalEventCount, AllocCount, FreeCount);
 }
 
 ////////////////////////////////////////////////////////////////////////////////////////////////////
@@ -1959,6 +2015,8 @@ const IAllocationsProvider* ReadAllocationsProvider(const IAnalysisSession& Sess
 #undef INSIGHTS_SLOW_CHECK
 #undef INSIGHTS_DEBUG_WATCH
 #undef INSIGHTS_DEBUG_WATCH_FOUND
+#undef INSIGHTS_LLA_RESERVE
 #undef INSIGHTS_USE_SHORT_LIVING_ALLOCS
+#undef INSIGHTS_SLA_USE_ADDRESS_MAP
 #undef INSIGHTS_USE_LAST_ALLOC
 #undef INSIGHTS_VALIDATE_ALLOC_EVENTS
diff --git a/Engine/Source/Developer/TraceServices/Private/Model/AllocationsProvider.h b/Engine/Source/Developer/TraceServices/Private/Model/AllocationsProvider.h
--- a/Engine/Source/Developer/TraceServices/Private/Model/AllocationsProvider.h
+++ b/Engine/Source/Developer/TraceServices/Private/Model/AllocationsProvider.h
@@ -124,7 +124,8 @@ private:
 		FNode* Prev;
 	};
 
-	static const int32 MaxAllocCount = 8 * 1024; // max number of short living allocations
+	//static const int32 MaxAllocCount = 8 * 1024; // max number of short living allocations
+	static const int32 MaxAllocCount = 64; // max number of short living allocations
 
 public:
 	FShortLivingAllocs();
@@ -151,7 +152,7 @@ public:
 	FAllocationItem* FindRange(uint64 Address) const;
 
 private:
-	TMap<uint64, FNode*> AddressMap; // map for short living allocations: Address -> FNode*
+	//TMap<uint64, FNode*> AddressMap; // map for short living allocations: Address -> FNode*
 	FNode* AllNodes = nullptr; // preallocated array of nodes (MaxAllocCount nodes)
 	FNode* LastAddedAllocNode = nullptr; // the last added alloc; double linked list: Prev -> .. -> OldestAlloc
 	FNode* OldestAllocNode = nullptr; // the oldest alloc; double linked list: Next -> .. -> LastAddedAlloc
@@ -263,10 +264,33 @@ public:
 	// Enumerates all allocations (including heap allocations).
 	void Enumerate(TFunctionRef<void(const FAllocationItem& Alloc)> Callback) const;
 
+private:
+	template<typename ValueType>
+	struct TAddressKeyFuncs : BaseKeyFuncs<TPair<uint64, ValueType>, uint64, false>
+	{
+		typedef typename TTypeTraits<uint64>::ConstPointerType KeyInitType;
+		typedef const TPairInitializer<typename TTypeTraits<uint64>::ConstInitType, typename TTypeTraits<ValueType>::ConstInitType>& ElementInitType;
+
+		static FORCEINLINE uint64 GetSetKey(ElementInitType Element)
+		{
+			return Element.Key;
+		}
+
+		static FORCEINLINE bool Matches(uint64 A, uint64 B)
+		{
+			return A == B;
+		}
+
+		static FORCEINLINE uint32 GetKeyHash(uint64 Key)
+		{
+			return uint32(Key >> 6);
+		}
+	};
+
 private:
 	FAllocationItem* LastAlloc = nullptr; // last allocation
 	FShortLivingAllocs ShortLivingAllocs; // short living allocs
-	TMap<uint64, FAllocationItem*> LongLivingAllocs; // long living allocations
+	TMap<uint64, FAllocationItem*, FDefaultSetAllocator, TAddressKeyFuncs<FAllocationItem*>> LongLivingAllocs; // long living allocations
 	FHeapAllocs HeapAllocs; // heap allocs
 
 	uint32 TotalAllocCount = 0;
@@ -376,13 +400,15 @@ private:
 	FTagTracker TagTracker;
 	uint8 CurrentTracker = 0;
 
-	uint32 EventIndex[MaxRootHeaps];
+	TArray<FHeapSpec> HeapSpecs;
+
+	uint32 EventIndex[MaxRootHeaps] = { 0 };
+	FSbTree* SbTree[MaxRootHeaps] = { nullptr };
+	FLiveAllocCollection* LiveAllocs[MaxRootHeaps] = { nullptr };
 
 	uint64 AllocCount = 0;
 	uint64 FreeCount = 0;
 
-	FLiveAllocCollection* LiveAllocs[MaxRootHeaps];
-
 	uint64 MiscErrors = 0;
 	uint64 HeapErrors = 0;
 	uint64 AllocErrors = 0;
@@ -394,8 +420,6 @@ private:
 	uint32 MaxEventDistance = 0;
 	uint32 EventDistanceHistogramPow2[33] = { 0 };
 
-	FSbTree* SbTree[MaxRootHeaps];
-
 	uint64 TotalAllocatedMemory = 0;
 	uint32 TotalLiveAllocations = 0;
 
@@ -415,8 +439,6 @@ private:
 	TPagedArray<uint32> MaxLiveAllocationsTimeline;
 	TPagedArray<uint32> AllocEventsTimeline;
 	TPagedArray<uint32> FreeEventsTimeline;
-
-	TArray<FHeapSpec> HeapSpecs;
 };
 
 ////////////////////////////////////////////////////////////////////////////////////////////////////
