[Insights]
  - Added "Discovered" status for symbol resolving of a module.
  - Fixed deadlock for analysis thread when session completes. This bug also didn't allowed Insights app to terminate when closed. The bug occurs for relatively short mem trace sessions when the session analysis completes before callstack resolving fully loads all modules.
  - Fixed issue where number of Discovered symbols for a module was lower than Resolved + Failed. Note: This issue might still occur when a module is reloaded, but should not occur anymore for normal callstack resolving.

diff --git a/Engine/Source/Developer/TraceServices/Private/Model/ModuleProvider.cpp b/Engine/Source/Developer/TraceServices/Private/Model/ModuleProvider.cpp
--- a/Engine/Source/Developer/TraceServices/Private/Model/ModuleProvider.cpp
+++ b/Engine/Source/Developer/TraceServices/Private/Model/ModuleProvider.cpp
@@ -65,7 +65,7 @@ private:
 		uint32 FileOffset;
 		uint32 Line;
 	};
-	
+
 	mutable FRWLock				ModulesLock;
 	TPagedArray<FModule>		Modules;
 
@@ -80,7 +80,7 @@ private:
 	uint32 NumCachedSymbols;
 	// Number of discovered symbols
 	std::atomic<uint32> SymbolsDiscovered;
-	
+
 	IAnalysisSession&			Session;
 	FString						Platform;
 	TUniquePtr<SymbolResolverType>	Resolver;
@@ -128,21 +128,25 @@ const FResolvedSymbol* TModuleProvider<SymbolResolverType>::GetSymbol(uint64 Add
 		}
 	}
 
-	FResolvedSymbol* ResolvedSymbol;
+	FResolvedSymbol* ResolvedSymbol = nullptr;
 	{
-		// Add a pending entry to our cache
 		FWriteScopeLock _(SymbolsLock);
-		if (SymbolCacheLookup.Contains(Address))
+
+		// Attempt again to read from the cached symbols.
+		const FResolvedSymbol* CachedResolvedSymbol = SymbolCacheLookup.FindRef(Address);
+		if (CachedResolvedSymbol)
 		{
-			return SymbolCacheLookup[Address];
+			return CachedResolvedSymbol;
 		}
+
+		// Add a pending entry to our cache.
 		ResolvedSymbol = &SymbolCache.EmplaceBack(ESymbolQueryResult::Pending, nullptr, nullptr, nullptr, 0);
 		SymbolCacheLookup.Add(Address, ResolvedSymbol);
 		++SymbolsDiscovered;
 	}
-
-	// If not in cache yet, queue it up in the resolver
 	check(ResolvedSymbol);
+
+	// If not in cache yet, queue it up in the resolver.
 	Resolver->QueueSymbolResolve(Address, ResolvedSymbol);
 
 	return ResolvedSymbol;
@@ -158,7 +162,7 @@ uint32 TModuleProvider<SymbolResolverType>::GetNumModules() const
 
 /////////////////////////////////////////////////////////////////////
 template <typename SymbolResolverType>
-void TModuleProvider<SymbolResolverType>::EnumerateModules(uint32 Start, TFunctionRef<void(const FModule& Module)> Callback) const 
+void TModuleProvider<SymbolResolverType>::EnumerateModules(uint32 Start, TFunctionRef<void(const FModule& Module)> Callback) const
 {
 	FReadScopeLock _(ModulesLock);
 	for (uint32 i = Start; i < Modules.Num(); ++i)
@@ -179,7 +183,7 @@ FGraphEventRef TModuleProvider<SymbolResolverType>::LoadSymbolsForModuleUsingPat
 		if (Resolver && !FullPath.IsEmpty() && Module->Status.load() != EModuleStatus::Loaded)
 		{
 			// Setup a task to queue and watch the queued module. If it succeeds in resolving with the new path
-			// re-resolve any cached symbols 
+			// re-resolve any cached symbols
 			LoadSymbolsTask = FFunctionGraphTask::CreateAndDispatchWhenReady([this, Module, FullPath]()
 			{
 				auto ReloadModuleFn = [this] (const FModule* InModule, const TCHAR* InPath)
@@ -191,7 +195,7 @@ FGraphEventRef TModuleProvider<SymbolResolverType>::LoadSymbolsForModuleUsingPat
 					auto ReresolveOnSuccess = [this, DiscoveredSymbols, ModuleBegin, ModuleEnd] (TArray<TTuple<uint64,FResolvedSymbol*>>& OutSymbols)
 					{
 						OutSymbols.Reserve(DiscoveredSymbols);
-					
+
 						FReadScopeLock _(SymbolsLock);
 						for (auto Pair : SymbolCacheLookup)
 						{
@@ -205,15 +209,16 @@ FGraphEventRef TModuleProvider<SymbolResolverType>::LoadSymbolsForModuleUsingPat
 					};
 
 					Resolver->QueueModuleReload(InModule, InPath, ReresolveOnSuccess);
+
 					// Wait for the resolver to do it's work
 					while (InModule->Status.load() == EModuleStatus::Pending)
 					{
-						FPlatformProcess::Sleep(.1);
+						FPlatformProcess::Sleep(0.1f);
 					}
 
 					return InModule->Status.load();
 				};
-			
+
 				UE_LOG(LogTraceServices, Display, TEXT("Queing symbol loading using path %s."), *FullPath);
 
 				// Load the requested module
@@ -231,7 +236,7 @@ FGraphEventRef TModuleProvider<SymbolResolverType>::LoadSymbolsForModuleUsingPat
 							return;
 						}
 						const EModuleStatus ModuleStatus = OtherModule.Status.load();
-						if (&OtherModule != Module && ModuleStatus != EModuleStatus::Loaded && ModuleStatus != EModuleStatus::Pending)
+						if (&OtherModule != Module && ModuleStatus >= EModuleStatus::FailedStatusStart)
 						{
 							ReloadModuleFn(&OtherModule, *Directory);
 						}
@@ -278,12 +283,12 @@ void TModuleProvider<SymbolResolverType>::OnModuleLoad(const FStringView& Module
 	const TCHAR* FullName = Session.StoreString(Module);
 
 	FWriteScopeLock _(ModulesLock);
-	FModule& NewModule = Modules.EmplaceBack(Name, FullName, Base, Size, EModuleStatus::Pending);
+	FModule& NewModule = Modules.EmplaceBack(Name, FullName, Base, Size, EModuleStatus::Discovered);
 
 	// The number of cached symbols for this module is equal to the number
 	// of matching addresses in the cache.
 	NewModule.Stats.Cached = GetNumCachedSymbolsFromModule(Base, Size);
-	
+
 	Resolver->QueueModuleLoad(ImageId, ImageIdSize, &NewModule);
 }
 
@@ -308,7 +313,7 @@ void TModuleProvider<SymbolResolverType>::SaveSymbolsToCache(IAnalysisCache& Cac
 	// Create a temporary reverse lookup for symbol -> address
 	TMap<const FResolvedSymbol*, uint64> SymbolReverseLookup;
 	SymbolReverseLookup.Reserve(SymbolCacheLookup.Num());
-	Algo::Transform(SymbolCacheLookup, SymbolReverseLookup, [](const TTuple<uint64, const FResolvedSymbol*>& Pair) { 
+	Algo::Transform(SymbolCacheLookup, SymbolReverseLookup, [](const TTuple<uint64, const FResolvedSymbol*>& Pair) {
 		return TTuple<const FResolvedSymbol*, uint64>(Pair.Value, Pair.Key);
 	});
 
diff --git a/Engine/Source/Developer/TraceServices/Private/Model/SymslibResolver.cpp b/Engine/Source/Developer/TraceServices/Private/Model/SymslibResolver.cpp
--- a/Engine/Source/Developer/TraceServices/Private/Model/SymslibResolver.cpp
+++ b/Engine/Source/Developer/TraceServices/Private/Model/SymslibResolver.cpp
@@ -10,6 +10,7 @@
 #include "Containers/StringView.h"
 #include "GenericPlatform/GenericPlatformFile.h"
 #include "HAL/PlatformFileManager.h"
+#include "HAL/PlatformProcess.h"
 #include "HAL/PlatformTLS.h"
 #include "Logging/LogMacros.h"
 #include "Misc/ConfigCacheIni.h"
@@ -71,17 +72,17 @@ namespace
 			return SearchPath;
 		}
 
-		
+
 		// Extract only filename part in case Path is absolute path
 		FString FileName = FPaths::GetCleanFilename(File);
-		
-		// Look in the search path 
+
+		// Look in the search path
 		FString Result = FPaths::Combine(SearchPath, FileName);
 		if (PlatformFile->FileExists(*Result))
 		{
 			return Result;
 		}
-		
+
 		// In case File is relative
 		if (FPaths::IsRelative(SearchPath))
 		{
@@ -124,11 +125,11 @@ namespace
 		{
 			SearchPathBase = SearchPath;
 		}
-		
+
 		// Extract only filename part in case Path is absolute path
 		FString FileName = FPaths::GetCleanFilename(File);
-		
-		// Look in the search path 
+
+		// Look in the search path
 		FString Result = FPaths::Combine(SearchPathBase, FileName);
 		if (PlatformFile->FileExists(*Result))
 		{
@@ -162,9 +163,9 @@ namespace
 				return Result;
 			}
 		}
-		
+
 		const FString File = FPaths::GetCleanFilename(FilePath);
-		
+
 		// Look in some well known directories
 		for (const TCHAR* Folder : EngineFolders)
 		{
@@ -173,7 +174,7 @@ namespace
 			{
 				return Result;
 			}
-			
+
 		}
 
 		// Look in project directory
@@ -231,7 +232,7 @@ namespace
 			BinaryPath = FindFileInEngineFolder(FilePath, SearchPath, Platform, Project);
 		}
 		bFileFound = !BinaryPath.IsEmpty();
-		
+
 		if (!bFileFound)
 		{
 			UE_LOG(LogSymslib, Verbose, TEXT("Binary file '%s' not found"), *FilePath);
@@ -261,7 +262,7 @@ namespace
 
 		return true;
 	}
-	
+
 	static bool LoadDebug(SYMS_Arena* Arena, SYMS_ParseBundle& Bundle, TArray<FAutoMappedFile>& Files, const FString& SearchPath)
 	{
 		if (syms_bin_is_dbg(Bundle.bin))
@@ -324,7 +325,7 @@ namespace
 			UE_LOG(LogSymslib, Verbose, TEXT("Debug file '%s' not found"), *FilePath);
 			return false;
 		}
-		
+
 		FAutoMappedFile& File = Files.AddDefaulted_GetRef();
 		if (!File.Load(*DebugPath))
 		{
@@ -471,7 +472,7 @@ private:
 	uint32 BlockRemaining = 0;
 	uint32 BlockUsed = 0;
 };
-	
+
 /////////////////////////////////////////////////////////////////////
 FSymslibResolver::FSymslibResolver(IAnalysisSession& InSession)
 	: Modules(InSession.GetLinearAllocator(), 128)
@@ -551,7 +552,7 @@ void FSymslibResolver::QueueModuleLoad(const uint8* ImageId, uint32 ImageIdSize,
 {
 	FWriteScopeLock _(ModulesLock);
 	check(!CleanupTask); // No queued addresses should be queued after cleanup task is started
-	
+
 	// Add module
 	FModuleEntry* Entry = &Modules.PushBack();
 	Entry->Module = Module;
@@ -568,7 +569,7 @@ void FSymslibResolver::QueueModuleLoad(const uint8* ImageId, uint32 ImageIdSize,
 			LoadModuleTracked(Entry, FStringView());
 			--TasksInFlight;
 		}, TStatId{}, nullptr, ENamedThreads::AnyBackgroundThreadNormalTask);
-	
+
 	// Sort list according to base address
 	SortedModules.Add(Entry);
 	Algo::Sort(SortedModules, [](const FModuleEntry* Lhs, const FModuleEntry* Rhs) { return Lhs->Module->Base < Rhs->Module->Base; });
@@ -588,16 +589,15 @@ void FSymslibResolver::QueueModuleReload(const FModule* Module, const TCHAR* InP
 			break;
 		}
 	}
-	
+
 	if (Entry)
 	{
 		// No use in trying reload already loaded modules
-		EModuleStatus PreviousStatus = Entry->Module->Status.load();
-		if (!Entry->Module->Status.compare_exchange_strong(PreviousStatus, EModuleStatus::Pending))
+		if (Entry->Module->Status == EModuleStatus::Loaded)
 		{
 			return;
 		}
-		
+
 		FString PathStr(InPath);
 		FPaths::NormalizeDirectoryName(PathStr);
 		const TCHAR* Path = Session.StoreString(PathStr);
@@ -607,25 +607,27 @@ void FSymslibResolver::QueueModuleReload(const FModule* Module, const TCHAR* InP
 		{
 			ModuleReloadTask->Wait();
 		}
-		
-		FFunctionGraphTask::CreateAndDispatchWhenReady([this, Entry, Path, PreviousStatus, ResolveOnSuccess]
+
+		FFunctionGraphTask::CreateAndDispatchWhenReady([this, Entry, Path, ResolveOnSuccess]
 		{
 			FScopeLock CleanupExclusive(&CleanupLock);
-			
+
+			EModuleStatus PreviousStatus = Entry->Module->Status.load();
+
 			LoadModuleTracked(Entry, Path);
 
 			if (PreviousStatus >= EModuleStatus::FailedStatusStart)
 			{
 				--ModulesFailed;
 			}
-			
+
 			if (Entry->Module->Status.load() == EModuleStatus::Loaded)
 			{
 				// If an additional search path was specified and successful, add to the permanent
 				// list of search paths.
 				{
 					FWriteScopeLock _(SymbolSearchPathsLock);
-					//Get the base path if path points to a file 
+					// Get the base path if path points to a file
 					IPlatformFile* PlatformFile = &FPlatformFileManager::Get().GetPlatformFile();
 					const FString Directory = PlatformFile->DirectoryExists(Path) ? Path : FPaths::GetPath(Path);
 					if (!Directory.IsEmpty())
@@ -633,14 +635,22 @@ void FSymslibResolver::QueueModuleReload(const FModule* Module, const TCHAR* InP
 						SymbolSearchPaths.AddUnique(Directory);
 					}
 				}
-				
+
+				// Reset stats for reloaded module.
+				// Note: If there are symbols pending to resolve (ex. busy spinning in ResolveSymbol()), the Stats.Discovered
+				//       might be already incremented (see ResolveSymbolTracked()). So, resetting the stats here might result
+				//       in slightly wrong stats where Discovered < Resolved + Failed.
+				Entry->Module->Stats.Discovered.store(0u);
+				Entry->Module->Stats.Resolved.store(0u);
+				Entry->Module->Stats.Failed.store(0u);
+
 				// Ask the caller for a list of symbols that should be resolved, now that module is properly loaded and
 				// resolve them immediately.
 				{
 					SymbolArray SymbolsToResolve;
-					FSymbolStringAllocator StringAllocator(Session.GetLinearAllocator(), (2 << 12) / sizeof(TCHAR) );
+					FSymbolStringAllocator StringAllocator(Session.GetLinearAllocator(), (2 << 12) / sizeof(TCHAR));
 					ResolveOnSuccess(SymbolsToResolve);
-					for(TTuple<uint64, FResolvedSymbol*> Pair : SymbolsToResolve)
+					for (TTuple<uint64, FResolvedSymbol*> Pair : SymbolsToResolve)
 					{
 						ResolveSymbolTracked(Pair.Get<0>(), Pair.Get<1>(), StringAllocator);
 					}
@@ -653,7 +663,7 @@ void FSymslibResolver::QueueModuleReload(const FModule* Module, const TCHAR* InP
 					// Need to wait for tasks, since there could be queued symbol resolves for this module waiting
 					// in the queue.
 					WaitForTasks();
-					
+
 					Algo::ForEach(Entry->Instance.Arenas, syms_arena_release);
 					Entry->Instance.Arenas.Empty();
 				}
@@ -676,9 +686,8 @@ void FSymslibResolver::OnAnalysisComplete()
 	{
 		 FScopeLock _(&SymbolsQueueLock);
 		 DispatchQueuedAddresses();
-		 ResolveQueue.Reset();
 	}
-	
+
 	CleanupTask = FFunctionGraphTask::CreateAndDispatchWhenReady([this]
 	{
 		FScopeLock CleanupExclusive(&CleanupLock);
@@ -696,10 +705,9 @@ void FSymslibResolver::OnAnalysisComplete()
 				Module.Instance.Arenas.Empty();
 			}
 		}
-		
+
 		UE_LOG(LogSymslib, Display, TEXT("Allocated %.02f Mb of strings, %.02f Mb wasted."),
 		       SymbolBytesAllocated / float(1024*1024), SymbolBytesWasted / float(1024*1024));
-		
 	});
 }
 
@@ -734,7 +742,6 @@ void FSymslibResolver::MaybeDispatchQueuedAddresses()
 	if (!bModulesInFlight && (ResolveQueue.Num() >= QueuedAddressLength))
 	{
 		DispatchQueuedAddresses();
-		ResolveQueue.Reset();
 	}
 }
 
@@ -743,27 +750,27 @@ void FSymslibResolver::MaybeDispatchQueuedAddresses()
  */
 void FSymslibResolver::DispatchQueuedAddresses()
 {
-	if (ResolveQueue.IsEmpty())
+	if (!ResolveQueue.IsEmpty())
 	{
-		return;
-	}
+		TArray<FQueuedAddress> WorkingSet(ResolveQueue);
+
+		uint32 Stride = (WorkingSet.Num() - 1) / SymbolTasksInParallel + 1;
+		constexpr uint32 MinStride = 4;
+		Stride = FMath::Max(Stride, MinStride);
+		const uint32 ActualSymbolTasksInParallel = (WorkingSet.Num() + Stride - 1) / Stride;
+		TasksInFlight += ActualSymbolTasksInParallel;
+
+		// Use background priority in order to not interfere with Slate.
+		ParallelFor(ActualSymbolTasksInParallel, [this, &WorkingSet, Stride](uint32 Index) {
+			const uint32 StartIndex = Index * Stride;
+			const uint32 EndIndex = FMath::Min(StartIndex + Stride, (uint32)WorkingSet.Num());
+			TArrayView<FQueuedAddress> QueuedWork(&WorkingSet[StartIndex], EndIndex - StartIndex);
+			ResolveSymbols(QueuedWork);
+			--TasksInFlight;
+		}, EParallelForFlags::BackgroundPriority);
 
-	TArray<FQueuedAddress> WorkingSet(ResolveQueue);
-
-	uint32 Stride = (WorkingSet.Num() - 1) / SymbolTasksInParallel + 1;
-	constexpr uint32 MinStride = 4;
-	Stride = FMath::Max(Stride, MinStride);
-	const uint32 ActualSymbolTasksInParallel = (WorkingSet.Num() + Stride - 1) / Stride;
-	TasksInFlight += ActualSymbolTasksInParallel;
-
-	// Use background priority in order to not not interfere with Slate
-	ParallelFor(ActualSymbolTasksInParallel, [this, &WorkingSet, Stride](uint32 Index) {
-		const uint32 StartIndex = Index * Stride;
-		const uint32 EndIndex = FMath::Min(StartIndex + Stride, (uint32)WorkingSet.Num());
-		TArrayView<FQueuedAddress> QueuedWork(&WorkingSet[StartIndex], EndIndex - StartIndex);
-		ResolveSymbols(QueuedWork);
-		--TasksInFlight;
-	}, EParallelForFlags::BackgroundPriority);
+		ResolveQueue.Reset();
+	}
 }
 
 void FSymslibResolver::ResolveSymbols(TArrayView<FQueuedAddress>& QueuedWork)
@@ -809,6 +816,8 @@ void FSymslibResolver::UpdateResolvedSymbol(FResolvedSymbol* Symbol, ESymbolQuer
 
 void FSymslibResolver::LoadModuleTracked(FModuleEntry* Entry, FStringView OverrideSearchPath)
 {
+	EModuleStatus PreviousStatus = Entry->Module->Status.exchange(EModuleStatus::Pending);
+
 	// Build search paths
 	FString ModuleNamePath;
 	TArray<FStringView> SearchPaths;
@@ -819,7 +828,7 @@ void FSymslibResolver::LoadModuleTracked(FModuleEntry* Entry, FStringView Overri
 		{
 			SearchPaths.AddUnique(SearchPath);
 		}
-		
+
 		ModuleNamePath = FPaths::GetPath(Entry->Module->FullName);
 		FPaths::NormalizeDirectoryName(ModuleNamePath);
 		if (!ModuleNamePath.IsEmpty())
@@ -827,50 +836,34 @@ void FSymslibResolver::LoadModuleTracked(FModuleEntry* Entry, FStringView Overri
 			SearchPaths.AddUnique(ModuleNamePath);
 		}
 	}
-	else 
+	else
 	{
 		SearchPaths.Add(OverrideSearchPath);
 	}
 
 	// Exhaust all the search paths
-	EModuleStatus Status = EModuleStatus::Pending;
+	EModuleStatus Status = EModuleStatus::Failed;
 	TStringBuilder<128> StatusMessage;
-	while(Status != EModuleStatus::Loaded && !CancelTasks.load() && SearchPaths.Num() > 0)
+	while (Status != EModuleStatus::Loaded && !CancelTasks.load() && SearchPaths.Num() > 0)
 	{
 		StatusMessage.Reset();
-		Status = LoadModule(Entry,SearchPaths.Top(), StatusMessage);
+		Status = LoadModule(Entry, SearchPaths.Top(), StatusMessage);
 		UE_LOG(LogSymslib, Display, TEXT("%s"), StatusMessage.ToString());
 		SearchPaths.Pop();
 	}
 
-	// If we are still pending after the loop we have failed
-	if (Status == EModuleStatus::Pending)
-	{
-		Status = EModuleStatus::Failed;
-	}
-	
 	if (Status == EModuleStatus::Loaded)
 	{
-		// Reset this modules stats
-		Entry->Module->Stats.Discovered.store(0u);
-		Entry->Module->Stats.Resolved.store(0u);
-		Entry->Module->Stats.Failed.store(0u);
 		++ModulesLoaded;
 	}
 	else
 	{
 		++ModulesFailed;
 	}
-	
+
 	// Make the final status visible to the world
 	Entry->Module->StatusMessage = Session.StoreString(StatusMessage.ToView());
 	Entry->Module->Status.store(Status);
-
-	// If this is the last module in flight dispatch pending queries
-	{
-		FScopeLock _(&SymbolsQueueLock);
-		MaybeDispatchQueuedAddresses();
-	}
 }
 
 EModuleStatus FSymslibResolver::LoadModule(FModuleEntry* Entry, FStringView SearchPathView, FStringBuilderBase& OutStatusMessage) const
@@ -1034,7 +1027,7 @@ EModuleStatus FSymslibResolver::LoadModule(FModuleEntry* Entry, FStringView Sear
 	// In order to avoid listing all files, use only the last. On dwarf/elf this will be the main binary, on Windows
 	// this will be the pdb (the first is the exe/dll).
 	OutStatusMessage.Appendf(TEXT("%s loaded with %u symbols."), *Files.Last().FilePath, SymbolCount.load());
-	return EModuleStatus::Loaded; 
+	return EModuleStatus::Loaded;
 }
 
 void FSymslibResolver::ResolveSymbolTracked(uint64 Address, FResolvedSymbol* Target, FSymbolStringAllocator& StringAllocator)
@@ -1044,7 +1037,7 @@ void FSymslibResolver::ResolveSymbolTracked(uint64 Address, FResolvedSymbol* Tar
 	{
 		return;
 	}
-	
+
 	FModuleEntry* Entry = GetModuleForAddress(Address);
 	if (!Entry)
 	{
@@ -1068,9 +1061,9 @@ void FSymslibResolver::ResolveSymbolTracked(uint64 Address, FResolvedSymbol* Tar
 bool FSymslibResolver::ResolveSymbol(uint64 Address, FResolvedSymbol* Target, FSymbolStringAllocator& StringAllocator, FModuleEntry* Entry) const
 {
 	EModuleStatus Status = Entry->Module->Status.load();
-	while (Status == EModuleStatus::Pending)
+	while ((Status == EModuleStatus::Pending || Status == EModuleStatus::Discovered) && !CancelTasks.load())
 	{
-		//todo: yield
+		FPlatformProcess::YieldThread();
 		Status = Entry->Module->Status.load();
 	}
 
@@ -1160,10 +1153,10 @@ bool FSymslibResolver::ResolveSymbol(uint64 Address, FResolvedSymbol* Target, FS
 		SourceFilePersistent,
 		SourceFileLine
 	);
-		
+
 	return true;
 }
-	
+
 void FSymslibResolver::WaitForTasks()
 {
 	uint32 OutstandingTasks = TasksInFlight.load(std::memory_order_acquire);
diff --git a/Engine/Source/Developer/TraceServices/Public/TraceServices/Model/Modules.h b/Engine/Source/Developer/TraceServices/Public/TraceServices/Model/Modules.h
--- a/Engine/Source/Developer/TraceServices/Public/TraceServices/Model/Modules.h
+++ b/Engine/Source/Developer/TraceServices/Public/TraceServices/Model/Modules.h
@@ -44,7 +44,7 @@ inline const TCHAR* QueryResultToString(ESymbolQueryResult Result)
 ////////////////////////////////////////////////////////////////////////////////
 /**
   * Represent a resolved symbol. The resolve status and string values may change
-  * over time, but string pointers returned from the methods are guaranteed to live 
+  * over time, but string pointers returned from the methods are guaranteed to live
   * during the entire analysis session.
   */
 struct FResolvedSymbol
@@ -54,13 +54,13 @@ struct FResolvedSymbol
 	const TCHAR* Name;
 	const TCHAR* File;
 	uint16 Line;
-	
+
 	inline ESymbolQueryResult GetResult() const
 	{
 		return Result.load(std::memory_order_acquire);
 	}
 
-	FResolvedSymbol(ESymbolQueryResult InResult, const TCHAR* InModule, const TCHAR* InName, const TCHAR* InFile, uint16 InLine) 
+	FResolvedSymbol(ESymbolQueryResult InResult, const TCHAR* InModule, const TCHAR* InName, const TCHAR* InFile, uint16 InLine)
 		: Result(InResult)
 		, Module(InModule)
 		, Name(InName)
@@ -72,6 +72,7 @@ struct FResolvedSymbol
 ////////////////////////////////////////////////////////////////////////////////
 enum class EModuleStatus
 {
+	Discovered,			// Symbols are discovered for this module
 	Pending,			// Module is pending load
 	Loaded,				// Module has been successfully loaded
 	VersionMismatch,	// Debug data was found, but did not match traced version
@@ -80,11 +81,12 @@ enum class EModuleStatus
 	StatusNum,
 	FailedStatusStart = VersionMismatch
 };
-	
+
 ////////////////////////////////////////////////////////////////////////////////
 inline const TCHAR* ModuleStatusToString(EModuleStatus Result)
 {
 	static const TCHAR* DisplayStrings[] = {
+		TEXT("Discovered"),
 		TEXT("Pending..."),
 		TEXT("Loaded"),
 		TEXT("Version mismatch"),
@@ -94,7 +96,7 @@ inline const TCHAR* ModuleStatusToString(EModuleStatus Result)
 	static_assert(UE_ARRAY_COUNT(DisplayStrings) == (uint8) EModuleStatus::StatusNum, "Missing QueryResult");
 	return DisplayStrings[(uint8)Result];
 }
-	
+
 ////////////////////////////////////////////////////////////////////////////////
 /**
  * Represents information about a module (engine/game/system dll or monolithic binary)
@@ -122,7 +124,7 @@ struct FModule
 		std::atomic<uint32>		Resolved;
 		std::atomic<uint32>		Failed;
 	} Stats;
-	
+
 	FModule(const TCHAR* InName, const TCHAR* InFullName, uint64 InBase, uint32 InSize, EModuleStatus InStatus)
 		: Name(InName)
 		, FullName(InFullName)
@@ -132,7 +134,7 @@ struct FModule
 		, StatusMessage(nullptr)
 	{}
 };
-	
+
 ////////////////////////////////////////////////////////////////////////////////
 class IModuleProvider : public IProvider
 {
@@ -149,8 +151,8 @@ public:
 
 	virtual ~IModuleProvider() = default;
 
-	/** Queries the name of the symbol at address. This function returns immediately, 
-	 * but the lookup is async. See \ref FResolvedSymbol for details. It assumed that 
+	/** Queries the name of the symbol at address. This function returns immediately,
+	 * but the lookup is async. See \ref FResolvedSymbol for details. It assumed that
 	 * all calls to this function happens before analysis has ended.
 	 */
 	virtual const FResolvedSymbol* GetSymbol(uint64 Address) = 0;
