[Insights] Fix crash when receiving allocations without callstacks due to late connect.

diff --git a/Engine/Source/Developer/TraceServices/Private/Model/CallstacksProvider.cpp b/Engine/Source/Developer/TraceServices/Private/Model/CallstacksProvider.cpp
--- a/Engine/Source/Developer/TraceServices/Private/Model/CallstacksProvider.cpp
+++ b/Engine/Source/Developer/TraceServices/Private/Model/CallstacksProvider.cpp
@@ -106,9 +106,10 @@ void FCallstacksProvider::AddCallstack(uint32 InCallstackId, const uint64* InFra
 			while (InCallstackId >= Callstacks.Num())
 			{
 				Callstack = &Callstacks.PushBack();
+				Callstack->Init(&GNoStackFrame, 1);
 			}
 		}
-		check(Callstack && Callstack->Num() == 0); // adding same callstack id twice?
+		check(Callstack && (Callstack->Num() == 0 || Callstack->Frame(0) == &GNoStackFrame)); // adding same callstack id twice?
 		Callstack->Init(&Frames[FirstFrame], InFrameCount);
 	}
 }
