[Insights] Hierarchical tag view in memory queries.

* Adds logic to sort tags into correct hierarchy in analysis, regardless if the tag was declared as an LLMTag enum, name based tag or declared before it's parent has been declared.
* Add new class to break down memory allocs into hiearchical tags tree.

diff --git a/Engine/Source/Developer/TraceInsights/Private/Insights/MemoryProfiler/ViewModels/MemAllocGroupingByTag.cpp b/Engine/Source/Developer/TraceInsights/Private/Insights/MemoryProfiler/ViewModels/MemAllocGroupingByTag.cpp
new file mode 100644
--- /dev/null
+++ b/Engine/Source/Developer/TraceInsights/Private/Insights/MemoryProfiler/ViewModels/MemAllocGroupingByTag.cpp
@@ -0,0 +1,130 @@
+﻿// Copyright Epic Games, Inc. All Rights Reserved.
+#include "MemAllocGroupingByTag.h"
+#include "MemAllocNode.h"
+
+#define LOCTEXT_NAMESPACE "Insights::FMemAllocGroupingByTag"
+
+namespace Insights
+{
+
+////////////////////////////////////////////////////////////////////////////////////////////////////
+
+INSIGHTS_IMPLEMENT_RTTI(FMemAllocGroupingByTag);
+	
+////////////////////////////////////////////////////////////////////////////////////////////////////
+
+FMemAllocGroupingByTag::FMemAllocGroupingByTag(const TraceServices::IAllocationsProvider& InTagProvider)
+	: FTreeNodeGrouping(
+		LOCTEXT("Grouping_ByTag_ShortName", "Tag"),
+		LOCTEXT("Grouping_ByTag_TitleName", "By Tag"),
+		LOCTEXT("Grouping_Tag_Desc", "Creates a tree based on Tag hierarchy."),
+		TEXT("Icons.Group.TreeItem"),
+		nullptr)
+	, TagProvider(InTagProvider)
+{
+}
+
+////////////////////////////////////////////////////////////////////////////////////////////////////
+
+void FMemAllocGroupingByTag::GroupNodes(const TArray<FTableTreeNodePtr>& Nodes,
+                                        FTableTreeNode& ParentGroup, TWeakPtr<FTable> InParentTable,
+                                        std::atomic<bool>& bCancelGrouping) const
+{
+	using namespace TraceServices;
+	
+	ParentGroup.ClearChildren();
+
+	struct FTagParentRel
+	{
+		TagIdType Id;
+		TagIdType Parent;
+	};
+	TMap<TagIdType,FTableTreeNodePtr> TagIdToNode;
+	TArray<FTagParentRel> TagParentFixup;
+	TSet<TagIdType> NonEmptyIds;
+
+	// Create tag nodes
+	TagProvider.EnumerateTags([&](const TCHAR* Name, const TCHAR* FullName, TagIdType Id, TagIdType ParentId)
+	{
+		const auto Node = MakeShared<FTableTreeNode>(FName(Name), InParentTable);
+		TagIdToNode.Emplace(Id, Node);
+		TagParentFixup.Emplace(FTagParentRel {Id, ParentId});
+	});
+
+	// Untagged node is used if we cant find the id.
+	const auto UntaggedNode = TagIdToNode.Find(0);
+	NonEmptyIds.Add(0);
+
+	// Add nodes to correct tag group
+	for (FTableTreeNodePtr NodePtr : Nodes)
+	{
+		if (bCancelGrouping)
+		{
+			return;
+		}
+
+		if (NodePtr->IsGroup())
+		{
+			ParentGroup.AddChildAndSetGroupPtr(NodePtr);
+			continue;
+		}
+
+		const FMemAllocNode& MemAllocNode = static_cast<const FMemAllocNode&>(*NodePtr);
+		const FMemoryAlloc* Alloc = MemAllocNode.GetMemAlloc();
+		if (Alloc)
+		{
+			const TagIdType TagId = Alloc->GetTagId();
+			const auto TagNode = TagIdToNode.Find(TagId);
+			if (TagNode)
+			{
+				(*TagNode)->AddChildAndSetGroupPtr(NodePtr);
+				NonEmptyIds.Add(TagId);
+			}
+			else if (UntaggedNode)
+			{
+				(*UntaggedNode)->AddChildAndSetGroupPtr(NodePtr);
+			}
+		}
+	}
+
+	// Sort by parent id
+	Algo::SortBy(TagParentFixup, &FTagParentRel::Parent);
+	
+	// Fixup parent relationships and filter out empty nodes
+	for (auto TagParentRel : TagParentFixup)
+	{
+		if (bCancelGrouping)
+		{
+			return;
+		}
+		
+		if (!NonEmptyIds.Contains(TagParentRel.Id))
+		{
+			continue;
+		}
+		
+		if (TagParentRel.Parent == ~0)
+		{
+			// Add to parent directly
+			ParentGroup.AddChildAndSetGroupPtr(TagIdToNode[TagParentRel.Id]);
+		}
+		else
+		{
+			// Try to find the parent node and add it as a child
+			if (const auto ParentNode = TagIdToNode.Find(TagParentRel.Parent))
+			{
+				(*ParentNode)->AddChildAndSetGroupPtr(TagIdToNode[TagParentRel.Id]);
+			}
+
+			// Indicate that parent also is non-empty. Since the fixup array is sorted by parent id the root nodes
+			// be processed last.
+			NonEmptyIds.Add(TagParentRel.Parent);
+		}
+	}
+}
+	
+////////////////////////////////////////////////////////////////////////////////////////////////////
+
+#undef LOCTEXT_NAMESPACE
+
+}
diff --git a/Engine/Source/Developer/TraceInsights/Private/Insights/MemoryProfiler/ViewModels/MemAllocGroupingByTag.h b/Engine/Source/Developer/TraceInsights/Private/Insights/MemoryProfiler/ViewModels/MemAllocGroupingByTag.h
new file mode 100644
--- /dev/null
+++ b/Engine/Source/Developer/TraceInsights/Private/Insights/MemoryProfiler/ViewModels/MemAllocGroupingByTag.h
@@ -0,0 +1,28 @@
+﻿// Copyright Epic Games, Inc. All Rights Reserved.
+#pragma once
+
+#include "CoreMinimal.h"
+#include "Insights/Common/SimpleRtti.h"
+#include "Insights/Table/ViewModels/TableTreeNode.h"
+#include "Insights/Table/ViewModels/TreeNodeGrouping.h"
+#include "TraceServices/Model/AllocationsProvider.h"
+
+namespace Insights
+{
+	
+////////////////////////////////////////////////////////////////////////////////////////////////////
+class FMemAllocGroupingByTag : public FTreeNodeGrouping
+{
+	INSIGHTS_DECLARE_RTTI(FMemAllocGroupingByTag, FTreeNodeGrouping);
+
+public:
+	FMemAllocGroupingByTag(const TraceServices::IAllocationsProvider& TagProvider);
+
+	virtual void GroupNodes(const TArray<FTableTreeNodePtr>& Nodes, FTableTreeNode& ParentGroup, TWeakPtr<FTable> InParentTable, std::atomic<bool>& bCancelGrouping) const override;
+
+private:
+	const TraceServices::IAllocationsProvider& TagProvider;
+};
+	
+////////////////////////////////////////////////////////////////////////////////////////////////////
+}
diff --git a/Engine/Source/Developer/TraceInsights/Private/Insights/MemoryProfiler/ViewModels/MemoryAlloc.h b/Engine/Source/Developer/TraceInsights/Private/Insights/MemoryProfiler/ViewModels/MemoryAlloc.h
--- a/Engine/Source/Developer/TraceInsights/Private/Insights/MemoryProfiler/ViewModels/MemoryAlloc.h
+++ b/Engine/Source/Developer/TraceInsights/Private/Insights/MemoryProfiler/ViewModels/MemoryAlloc.h
@@ -30,6 +30,7 @@ public:
 	uint64 GetPage() const { return Address & ~(4llu*1024-1); }
 	int64 GetSize() const { return Size; }
 	const TCHAR* GetTag() const { return Tag; }
+	TraceServices::TagIdType GetTagId() const { return TagId; }
 	const TraceServices::FCallstack* GetCallstack() const { return Callstack; }
 	FText GetFullCallstack() const;
 	HeapId GetRootHeap() const { return RootHeap; };
@@ -42,6 +43,7 @@ private:
 	uint64 Address;
 	int64 Size;
 	const TCHAR* Tag;
+	TraceServices::TagIdType TagId;
 	const TraceServices::FCallstack* Callstack;
 	const TraceServices::FCallstack* FreeCallstack;
 	HeapId RootHeap;
diff --git a/Engine/Source/Developer/TraceInsights/Private/Insights/MemoryProfiler/Widgets/SMemAllocTableTreeView.cpp b/Engine/Source/Developer/TraceInsights/Private/Insights/MemoryProfiler/Widgets/SMemAllocTableTreeView.cpp
--- a/Engine/Source/Developer/TraceInsights/Private/Insights/MemoryProfiler/Widgets/SMemAllocTableTreeView.cpp
+++ b/Engine/Source/Developer/TraceInsights/Private/Insights/MemoryProfiler/Widgets/SMemAllocTableTreeView.cpp
@@ -24,6 +24,7 @@
 #include "Insights/MemoryProfiler/ViewModels/MemAllocGroupingByCallstack.h"
 #include "Insights/MemoryProfiler/ViewModels/MemAllocGroupingByHeap.h"
 #include "Insights/MemoryProfiler/ViewModels/MemAllocGroupingBySize.h"
+#include "Insights/MemoryProfiler/ViewModels/MemAllocGroupingByTag.h"
 #include "Insights/MemoryProfiler/ViewModels/MemAllocNode.h"
 #include "Insights/MemoryProfiler/ViewModels/MemAllocTable.h"
 #include "Insights/MemoryProfiler/ViewModels/MemorySharedState.h"
@@ -363,6 +364,7 @@ void SMemAllocTableTreeView::UpdateQuery(TraceServices::IAllocationsProvider::EQ
 					Alloc.Address = Allocation->GetAddress();
 					Alloc.Size = int64(Allocation->GetSize());
 					Alloc.Tag = Provider.GetTagName(Allocation->GetTag());
+					Alloc.TagId = Allocation->GetTag();
 					if (CallstacksProvider)
 					{
 						Alloc.Callstack = CallstacksProvider->GetCallstack(Allocation->GetCallstackId());
@@ -746,8 +748,7 @@ void SMemAllocTableTreeView::InitAvailableViewPresets()
 			const TSharedPtr<FTreeNodeGrouping>* TagGrouping = InAvailableGroupings.FindByPredicate(
 				[](TSharedPtr<FTreeNodeGrouping>& Grouping)
 				{
-					return Grouping->Is<FTreeNodeGroupingByUniqueValue>() &&
-						   Grouping->As<FTreeNodeGroupingByUniqueValue>().GetColumnId() == FMemAllocTableColumns::TagColumnId;
+					return Grouping->Is<FMemAllocGroupingByTag>();
 				});
 			if (TagGrouping)
 			{
@@ -1091,6 +1092,7 @@ void SMemAllocTableTreeView::InternalCreateGroupings()
 		if (AllocationsProvider)
 		{
 			AvailableGroupings.Insert(MakeShared<FMemAllocGroupingByHeap>(*AllocationsProvider), Index++);
+			AvailableGroupings.Insert(MakeShared<FMemAllocGroupingByTag>(*AllocationsProvider), Index++);
 		}
 	}
 }
diff --git a/Engine/Source/Developer/TraceServices/Private/Model/AllocationsProvider.cpp b/Engine/Source/Developer/TraceServices/Private/Model/AllocationsProvider.cpp
--- a/Engine/Source/Developer/TraceServices/Private/Model/AllocationsProvider.cpp
+++ b/Engine/Source/Developer/TraceServices/Private/Model/AllocationsProvider.cpp
@@ -2,17 +2,17 @@
 
 #include "AllocationsProvider.h"
 
+#include <limits>
+
 #include "AllocationsQuery.h"
 #include "SbTree.h"
-#include "Algo/ForEach.h"
 #include "Common/Utils.h"
 #include "Containers/ArrayView.h"
+#include "Misc/PathViews.h"
 #include "ProfilingDebugging/MemoryTrace.h"
 #include "TraceServices/Containers/Allocators.h"
 #include "TraceServices/Model/Callstack.h"
 
-#include <limits>
-
 ////////////////////////////////////////////////////////////////////////////////////////////////////
 
 #define INSIGHTS_SLOW_CHECK(expr) //check(expr)
@@ -122,15 +122,77 @@ void FAllocationsProviderLock::EndWrite()
 // FTagTracker
 ////////////////////////////////////////////////////////////////////////////////////////////////////
 
+FTagTracker::FTagTracker(IAnalysisSession& InSession)
+	: Session(InSession)
+{
+}
+
+////////////////////////////////////////////////////////////////////////////////////////////////////
+
 void FTagTracker::AddTagSpec(TagIdType InTag, TagIdType InParentTag, const TCHAR* InDisplay)
 {
 	if (ensure(!TagMap.Contains(InTag)))
 	{
-		TagMap.Emplace(InTag, TagEntry{ InDisplay, InParentTag });
+		FStringView Display(InDisplay);
+		FString DisplayName;
+		TStringBuilder<128> FullName;
+		if (Display.Contains(TEXT("/")))
+		{
+			DisplayName = FPathViews::GetPathLeaf(Display);
+			FullName = Display;
+			// It is possible to define a child tag in runtime using only a string, even if the parent tag does not yet
+			// exist. We need to find the correct parent or store it to the side until the parent tag is announced.
+			if (InParentTag == ~0)
+			{
+				const FStringView Parent = FPathViews::GetPathLeaf(FPathViews::GetPath(Display));
+				for (auto EntryPair : TagMap)
+				{
+					const auto Entry = EntryPair.Get<1>();
+					const uint32 Id = EntryPair.Get<0>();
+					if (Parent.Equals(Entry.Display))
+					{
+						InParentTag = Id;
+						break;
+					}
+				}
+				// If parent tag is still unknown here, create a temporary entry here
+				if (InParentTag == ~0)
+				{
+					PendingTags.Emplace(InTag, Parent);
+				}
+			}
+		}
+		else
+		{
+			DisplayName = Display;
+			BuildTagPath(FullName, Display, InParentTag);
+		}
+
+		const FTagEntry& Entry = TagMap.Emplace(InTag, FTagEntry{
+			Session.StoreString(DisplayName),
+			Session.StoreString(FullName.ToString()),
+			InParentTag
+		});
+
+		// Check if this new tag has been referenced before by a child tag
+		for (auto Pending : PendingTags)
+		{
+			const FString& Name = Pending.Get<1>();
+			const TagIdType ReferencingId = Pending.Get<0>();
+			if (Name.Equals(DisplayName))
+			{
+				TagMap[ReferencingId].ParentTag = InTag;
+			}
+		}
+		
 		if (!InDisplay || *InDisplay == TEXT('\0'))
 		{
 			UE_LOG(LogTraceServices, Warning, TEXT("[MemAlloc] Tag with id %u has invalid display name (ParentTag=%u)!"), InTag, InParentTag);
 		}
+		else
+		{
+			UE_LOG(LogTraceServices, Verbose, TEXT("[MemAlloc] Added Tag '%s' ('%s')"), Entry.Display, Entry.FullPath);
+		}
 	}
 	else
 	{
@@ -144,10 +206,22 @@ void FTagTracker::AddTagSpec(TagIdType InTag, TagIdType InParentTag, const TCHAR
 
 ////////////////////////////////////////////////////////////////////////////////////////////////////
 
+void FTagTracker::BuildTagPath(FStringBuilderBase& OutString, FStringView Name, TagIdType ParentTagId)
+{
+	if (const FTagEntry* ParentEntry = TagMap.Find(ParentTagId))
+	{
+		BuildTagPath(OutString, ParentEntry->Display, ParentEntry->ParentTag);
+		OutString << TEXT("/");
+	}
+	OutString << Name;
+}
+
+////////////////////////////////////////////////////////////////////////////////////////////////////
+
 void FTagTracker::PushTag(uint32 InThreadId, uint8 InTracker, TagIdType InTag)
 {
 	const uint32 TrackerThreadId = GetTrackerThreadId(InThreadId, InTracker);
-	ThreadState& State = TrackerThreadStates.FindOrAdd(TrackerThreadId);
+	FThreadState& State = TrackerThreadStates.FindOrAdd(TrackerThreadId);
 	State.TagStack.Push(InTag);
 }
 
@@ -156,7 +230,7 @@ void FTagTracker::PushTag(uint32 InThreadId, uint8 InTracker, TagIdType InTag)
 void FTagTracker::PopTag(uint32 InThreadId, uint8 InTracker)
 {
 	const uint32 TrackerThreadId = GetTrackerThreadId(InThreadId, InTracker);
-	ThreadState* State = TrackerThreadStates.Find(TrackerThreadId);
+	FThreadState* State = TrackerThreadStates.Find(TrackerThreadId);
 	if (State && !State->TagStack.IsEmpty())
 	{
 		INSIGHTS_SLOW_CHECK((State->TagStack.Top() & 0x80000000) == 0);
@@ -177,7 +251,7 @@ void FTagTracker::PopTag(uint32 InThreadId, uint8 InTracker)
 TagIdType FTagTracker::GetCurrentTag(uint32 InThreadId, uint8 InTracker) const
 {
 	const uint32 TrackerThreadId = GetTrackerThreadId(InThreadId, InTracker);
-	const ThreadState* State = TrackerThreadStates.Find(TrackerThreadId);
+	const FThreadState* State = TrackerThreadStates.Find(TrackerThreadId);
 	if (!State || State->TagStack.IsEmpty())
 	{
 		return 0; // Untagged
@@ -189,16 +263,28 @@ TagIdType FTagTracker::GetCurrentTag(uint32 InThreadId, uint8 InTracker) const
 
 const TCHAR* FTagTracker::GetTagString(TagIdType InTag) const
 {
-	const TagEntry* Entry = TagMap.Find(InTag);
+	const FTagEntry* Entry = TagMap.Find(InTag);
 	return Entry ? Entry->Display : TEXT("Unknown");
 }
 
 ////////////////////////////////////////////////////////////////////////////////////////////////////
 
+void FTagTracker::EnumerateTags(TFunctionRef<void (const TCHAR*, const TCHAR*, TagIdType, TagIdType)> Callback) const
+{
+	for(const auto& EntryPair : TagMap)
+	{
+		const TagIdType Id = EntryPair.Get<0>();
+		const FTagEntry& Entry = EntryPair.Get<1>();
+		Callback(Entry.Display, Entry.FullPath, Id, Entry.ParentTag);
+	}
+}
+
+////////////////////////////////////////////////////////////////////////////////////////////////////
+
 void FTagTracker::PushTagFromPtr(uint32 InThreadId, uint8 InTracker, TagIdType InTag)
 {
 	const uint32 TrackerThreadId = GetTrackerThreadId(InThreadId, InTracker);
-	ThreadState& State = TrackerThreadStates.FindOrAdd(TrackerThreadId);
+	FThreadState& State = TrackerThreadStates.FindOrAdd(TrackerThreadId);
 	State.TagStack.Push(InTag | PtrTagMask);
 }
 
@@ -207,7 +293,7 @@ void FTagTracker::PushTagFromPtr(uint32 InThreadId, uint8 InTracker, TagIdType I
 void FTagTracker::PopTagFromPtr(uint32 InThreadId, uint8 InTracker)
 {
 	const uint32 TrackerThreadId = GetTrackerThreadId(InThreadId, InTracker);
-	ThreadState* State = TrackerThreadStates.Find(TrackerThreadId);
+	FThreadState* State = TrackerThreadStates.Find(TrackerThreadId);
 	if (State && !State->TagStack.IsEmpty())
 	{
 		INSIGHTS_SLOW_CHECK((State->TagStack.Top() & PtrTagMask) != 0);
@@ -228,7 +314,7 @@ void FTagTracker::PopTagFromPtr(uint32 InThreadId, uint8 InTracker)
 bool FTagTracker::HasTagFromPtrScope(uint32 InThreadId, uint8 InTracker) const
 {
 	const uint32 TrackerThreadId = GetTrackerThreadId(InThreadId, InTracker);
-	const ThreadState* State = TrackerThreadStates.Find(TrackerThreadId);
+	const FThreadState* State = TrackerThreadStates.Find(TrackerThreadId);
 	return State && (State->TagStack.Num() > 0) && ((State->TagStack.Top() & PtrTagMask) != 0);
 }
 
@@ -1029,6 +1115,7 @@ void FLiveAllocCollection::Enumerate(TFunctionRef<void(const FAllocationItem& Al
 
 FAllocationsProvider::FAllocationsProvider(IAnalysisSession& InSession)
 	: Session(InSession)
+	, TagTracker(InSession)
 	, Timeline(Session.GetLinearAllocator(), 1024)
 	, MinTotalAllocatedMemoryTimeline(Session.GetLinearAllocator(), 1024)
 	, MaxTotalAllocatedMemoryTimeline(Session.GetLinearAllocator(), 1024)
@@ -1938,6 +2025,13 @@ void FAllocationsProvider::EnumerateFreeEventsTimeline(int32 StartIndex, int32 E
 
 ////////////////////////////////////////////////////////////////////////////////////////////////////
 
+void FAllocationsProvider::EnumerateTags(TFunctionRef<void(const TCHAR*, const TCHAR*, TagIdType, TagIdType)> Callback) const
+{
+	TagTracker.EnumerateTags(Callback);
+}
+
+////////////////////////////////////////////////////////////////////////////////////////////////////
+
 void FAllocationsProvider::DebugPrint() const
 {
 	for(const FSbTree* Tree : SbTree)
diff --git a/Engine/Source/Developer/TraceServices/Private/Model/AllocationsProvider.h b/Engine/Source/Developer/TraceServices/Private/Model/AllocationsProvider.h
--- a/Engine/Source/Developer/TraceServices/Private/Model/AllocationsProvider.h
+++ b/Engine/Source/Developer/TraceServices/Private/Model/AllocationsProvider.h
@@ -43,23 +43,26 @@ private:
 	static constexpr uint32 TrackerIdMask = 0xFF000000;
 	static constexpr uint32 PtrTagMask = 0x80000000;
 
-	struct ThreadState
+	struct FThreadState
 	{
 		TArray<TagIdType> TagStack;
 	};
 
-	struct TagEntry
+	struct FTagEntry
 	{
 		const TCHAR* Display;
+		const TCHAR* FullPath;
 		TagIdType ParentTag;
 	};
 
 public:
+	explicit FTagTracker(IAnalysisSession& Session);
 	void AddTagSpec(TagIdType Tag, TagIdType ParentTag, const TCHAR* Display);
 	void PushTag(uint32 ThreadId, uint8 Tracker, TagIdType Tag);
 	void PopTag(uint32 ThreadId, uint8 Tracker);
 	TagIdType GetCurrentTag(uint32 ThreadId, uint8 Tracker) const;
 	const TCHAR* GetTagString(TagIdType Tag) const;
+	void EnumerateTags(TFunctionRef<void(const TCHAR*, const TCHAR*, TagIdType, TagIdType)> Callback) const;
 
 	void PushTagFromPtr(uint32 ThreadId, uint8 Tracker, TagIdType Tag);
 	void PopTagFromPtr(uint32 ThreadId, uint8 Tracker);
@@ -68,13 +71,16 @@ public:
 	uint32 GetNumErrors() const { return NumErrors; }
 
 private:
-	inline uint32 GetTrackerThreadId(uint32 ThreadId, uint8 Tracker) const
+	void BuildTagPath(FStringBuilderBase& OutString, FStringView Name, TagIdType ParentTagId);
+	static inline uint32 GetTrackerThreadId(uint32 ThreadId, uint8 Tracker)
 	{
 		return (Tracker << TrackerIdShift) | (~TrackerIdMask & ThreadId);
 	}
 
-	TMap<uint32, ThreadState> TrackerThreadStates;
-	TMap<TagIdType, TagEntry> TagMap;
+	IAnalysisSession& Session;
+	TMap<uint32, FThreadState> TrackerThreadStates;
+	TMap<TagIdType, FTagEntry> TagMap;
+	TArray<TTuple<TagIdType, FString>> PendingTags;
 	uint32 NumErrors = 0;
 };
 
@@ -330,6 +336,7 @@ public:
 	virtual void EnumerateMaxLiveAllocationsTimeline(int32 StartIndex, int32 EndIndex, TFunctionRef<void(double Time, double Duration, uint32 Value)> Callback) const override;
 	virtual void EnumerateAllocEventsTimeline(int32 StartIndex, int32 EndIndex, TFunctionRef<void(double Time, double Duration, uint32 Value)> Callback) const override;
 	virtual void EnumerateFreeEventsTimeline(int32 StartIndex, int32 EndIndex, TFunctionRef<void(double Time, double Duration, uint32 Value)> Callback) const override;
+	virtual void EnumerateTags(TFunctionRef<void(const TCHAR*, const TCHAR*, TagIdType, TagIdType)> Callback) const override;
 
 	virtual FQueryHandle StartQuery(const FQueryParams& Params) const override;
 	virtual void CancelQuery(FQueryHandle Query) const override;
diff --git a/Engine/Source/Developer/TraceServices/Public/TraceServices/Model/AllocationsProvider.h b/Engine/Source/Developer/TraceServices/Public/TraceServices/Model/AllocationsProvider.h
--- a/Engine/Source/Developer/TraceServices/Public/TraceServices/Model/AllocationsProvider.h
+++ b/Engine/Source/Developer/TraceServices/Public/TraceServices/Model/AllocationsProvider.h
@@ -181,6 +181,9 @@ public:
 	// Enumerates the Free Events timeline points in the inclusive index interval [StartIndex, EndIndex].
 	virtual void EnumerateFreeEventsTimeline(int32 StartIndex, int32 EndIndex, TFunctionRef<void(double Time, double Duration, uint32 Value)> Callback) const = 0;
 
+	// Enumerates the discovered tags.
+	virtual void EnumerateTags(TFunctionRef<void(const TCHAR*, const TCHAR*, TagIdType, TagIdType)> Callback) const = 0;
+
 	virtual FQueryHandle StartQuery(const FQueryParams& Params) const = 0;
 	virtual void CancelQuery(FQueryHandle Query) const = 0;
 	virtual const FQueryStatus PollQuery(FQueryHandle Query) const = 0;
