UnrealInsights: Improved callstack butterfly, recursive calls will now be merged instead of displayed as individual nodes in the tree

diff --git a/Engine/Source/Developer/TraceServices/Private/Model/TimingProfiler.cpp b/Engine/Source/Developer/TraceServices/Private/Model/TimingProfiler.cpp
--- a/Engine/Source/Developer/TraceServices/Private/Model/TimingProfiler.cpp
+++ b/Engine/Source/Developer/TraceServices/Private/Model/TimingProfiler.cpp
@@ -437,16 +437,17 @@ ITimingProfilerButterfly* FTimingProfilerProvider::CreateButterfly(double Interv
 
 	struct FLocalStackEntry
 	{
-		FTimingProfilerButterflyNode* Node;
-		double StartTime;
-		double ExclusiveTime;
-		uint32 CurrentCallstackHash;
+		FTimingProfilerButterflyNode* Node = nullptr;
+		double StartTime = 0.0;
+		double ExclusiveTime = 0.0;
+		uint32 CurrentCallstackHash = 0;
+		bool bIsRecursive = false;
 	};
 
 	TArray<FLocalStackEntry> CurrentCallstack;
 	CurrentCallstack.Reserve(1024);
 
-	TMap<FTimingProfilerCallstackKey, FTimingProfilerButterflyNode*> CallstackNodeMap;
+	TMap<FTimingProfilerCallstackKey, TTuple<FTimingProfilerButterflyNode*, bool>> CallstackNodeMap;
 
 	double LastTime = 0.0;
 	for (const TimelineInternal* Timeline : IncludedTimelines)
@@ -471,31 +472,45 @@ ITimingProfilerButterfly* FTimingProfilerProvider::CreateButterfly(double Interv
 
 				FLocalStackEntry& StackEntry = CurrentCallstack.AddDefaulted_GetRef();
 				StackEntry.StartTime = Time;
-				StackEntry.ExclusiveTime = 0.0;
 				StackEntry.CurrentCallstackHash = ParentCallstackHash * 17 + Timer->Id;
 
 				CurrentCallstackKey.TimerStack.Push(Timer->Id);
 				CurrentCallstackKey.Hash = StackEntry.CurrentCallstackHash;
 
-				FTimingProfilerButterflyNode** FindIt = CallstackNodeMap.Find(CurrentCallstackKey);
+				TTuple<FTimingProfilerButterflyNode*, bool>* FindIt = CallstackNodeMap.Find(CurrentCallstackKey);
 				if (FindIt)
 				{
-					StackEntry.Node = *FindIt;
+					StackEntry.Node = FindIt->Get<0>();
+					StackEntry.bIsRecursive = FindIt->Get<1>();
 				}
 				else
 				{
-					StackEntry.Node = &Butterfly->Nodes.PushBack();
-					CallstackNodeMap.Add(CurrentCallstackKey, StackEntry.Node);
-					Butterfly->TimerCallstacksMap[Timer->Id].Add(StackEntry.Node);
-					StackEntry.Node->InclusiveTime = 0.0;
-					StackEntry.Node->ExclusiveTime = 0.0;
-					StackEntry.Node->Count = 0;
-					StackEntry.Node->Timer = Timer;
-					StackEntry.Node->Parent = ParentNode;
-					if (ParentNode)
+					for (int32 StackIndex = 0, StackEnd = CurrentCallstack.Num() - 1; StackIndex < StackEnd; ++StackIndex)
 					{
-						ParentNode->Children.Add(StackEntry.Node);
+						if (CurrentCallstack[StackIndex].Node->Timer == Timer)
+						{
+							StackEntry.Node = CurrentCallstack[StackIndex].Node;
+							StackEntry.bIsRecursive = true;
+							break;
+						}
 					}
+
+					if (!StackEntry.Node)
+					{
+						StackEntry.Node = &Butterfly->Nodes.PushBack();
+						StackEntry.Node->InclusiveTime = 0.0;
+						StackEntry.Node->ExclusiveTime = 0.0;
+						StackEntry.Node->Count = 0;
+						StackEntry.Node->Timer = Timer;
+						Butterfly->TimerCallstacksMap[Timer->Id].Add(StackEntry.Node);
+
+						StackEntry.Node->Parent = ParentNode;
+						if (ParentNode)
+						{
+							ParentNode->Children.Add(StackEntry.Node);
+						}
+					}
+					CallstackNodeMap.Add(CurrentCallstackKey, MakeTuple(StackEntry.Node, StackEntry.bIsRecursive));
 				}
 			}
 			else
@@ -504,8 +519,10 @@ ITimingProfilerButterfly* FTimingProfilerProvider::CreateButterfly(double Interv
 				double InclusiveTime = Time - StackEntry.StartTime;
 				check(InclusiveTime >= 0.0);
 				check(StackEntry.ExclusiveTime >= 0.0 && StackEntry.ExclusiveTime <= InclusiveTime);
-
-				StackEntry.Node->InclusiveTime += InclusiveTime;
+				if (!StackEntry.bIsRecursive)
+				{
+					StackEntry.Node->InclusiveTime += InclusiveTime;
+				}
 				StackEntry.Node->ExclusiveTime += StackEntry.ExclusiveTime;
 				++StackEntry.Node->Count;
 
