[Insights] Fixed NonUnity compile errors.

diff --git a/Engine/Source/Developer/TraceServices/Private/Model/DbgHelpResolver.h b/Engine/Source/Developer/TraceServices/Private/Model/DbgHelpResolver.h
--- a/Engine/Source/Developer/TraceServices/Private/Model/DbgHelpResolver.h
+++ b/Engine/Source/Developer/TraceServices/Private/Model/DbgHelpResolver.h
@@ -1,5 +1,8 @@
 // Copyright Epic Games, Inc. All Rights Reserved.
 #pragma once
+
+#include "HAL/Platform.h"
+
 #if PLATFORM_WINDOWS
 
 #include "Containers/Array.h"
@@ -20,7 +23,7 @@ class FDbgHelpResolver : public FRunnable
 {
 public:
 	typedef TArray<TTuple<uint64, FResolvedSymbol*>> SymbolArray;
-	
+
 	void Start();
 	FDbgHelpResolver(IAnalysisSession& InSession);
 	~FDbgHelpResolver();
@@ -31,7 +34,7 @@ public:
 	void EnumerateSymbolSearchPaths(TFunctionRef<void(FStringView Path)> Callback) const;
 	void OnAnalysisComplete();
 
-private:	
+private:
 	struct FModuleEntry
 	{
 		uint64 Base;
@@ -59,7 +62,6 @@ private:
 		MaxNameLen = 512,
 	};
 
-
 	bool SetupSyms();
 	void FreeSyms() const;
 	virtual uint32 Run() override;
diff --git a/Engine/Source/Developer/TraceServices/Private/Model/SymslibResolver.h b/Engine/Source/Developer/TraceServices/Private/Model/SymslibResolver.h
--- a/Engine/Source/Developer/TraceServices/Private/Model/SymslibResolver.h
+++ b/Engine/Source/Developer/TraceServices/Private/Model/SymslibResolver.h
@@ -1,6 +1,8 @@
 // Copyright Epic Games, Inc. All Rights Reserved.
 #pragma once
 
+#include "HAL/Platform.h"
+
 #if PLATFORM_WINDOWS
 
 #include "Async/MappedFileHandle.h"
@@ -21,7 +23,7 @@ class FSymslibResolver
 {
 public:
 	typedef TArray<TTuple<uint64, FResolvedSymbol*>> SymbolArray;
-	
+
 	FSymslibResolver(IAnalysisSession& InSession);
 	~FSymslibResolver();
 	void QueueModuleLoad(const uint8* ImageId, uint32 ImageIdSize, FModule* Module);
@@ -123,7 +125,6 @@ private:
 	FString Project;
 };
 
-
 /////////////////////////////////////////////////////////////////////
 
 } // namespace TraceServices
