M - Support for Unreal Insights to allow Suspend/Resume of Scope stacks which is Used by Coroutines.

diff --git a/Engine/Source/Developer/TraceServices/Private/Analyzers/CpuProfilerTraceAnalysis.cpp b/Engine/Source/Developer/TraceServices/Private/Analyzers/CpuProfilerTraceAnalysis.cpp
--- a/Engine/Source/Developer/TraceServices/Private/Analyzers/CpuProfilerTraceAnalysis.cpp
+++ b/Engine/Source/Developer/TraceServices/Private/Analyzers/CpuProfilerTraceAnalysis.cpp
@@ -1,6 +1,8 @@
 // Copyright Epic Games, Inc. All Rights Reserved.
 #include "CpuProfilerTraceAnalysis.h"
 #include "AnalysisServicePrivate.h"
+#include "CborWriter.h"
+#include "Serialization/MemoryWriter.h"
 #include "Common/Utils.h"
 #include "Model/ThreadsPrivate.h"
 
@@ -33,6 +35,8 @@ void FCpuProfilerAnalyzer::OnAnalysisBegin(const FOnAnalysisContext& Context)
 	Builder.RouteEvent(RouteId_EventBatch, "CpuProfiler", "EventBatch");
 	Builder.RouteEvent(RouteId_EndThread, "CpuProfiler", "EndThread");
 	Builder.RouteEvent(RouteId_EndCapture, "CpuProfiler", "EndCapture");
+	Builder.RouteEvent(RouteId_EventBatchV2, "CpuProfiler", "EventBatchV2");
+	Builder.RouteEvent(RouteId_EndCaptureV2, "CpuProfiler", "EndCaptureV2");
 }
 
 bool FCpuProfilerAnalyzer::OnEvent(uint16 RouteId, EStyle Style, const FOnEventContext& Context)
@@ -113,6 +117,8 @@ bool FCpuProfilerAnalyzer::OnEvent(uint16 RouteId, EStyle Style, const FOnEventC
 
 	case RouteId_EventBatch:
 	case RouteId_EndCapture:
+	case RouteId_EventBatchV2:
+	case RouteId_EndCaptureV2:
 	{
 		const uint32 ThreadId = FTraceAnalyzerUtils::GetThreadIdField(Context);
 		FThreadState& ThreadState = GetThreadState(ThreadId);
@@ -128,12 +134,21 @@ bool FCpuProfilerAnalyzer::OnEvent(uint16 RouteId, EStyle Style, const FOnEventC
 		const uint32 BufferSize = DataView.Num();
 		const uint8* BufferPtr = DataView.GetData();
 
-		uint64 LastCycle = ProcessBuffer(Context.EventTime, ThreadState, BufferPtr, BufferSize);
-		if (LastCycle != 0)
+		uint64 LastCycle;
+		if (RouteId == RouteId_EventBatch || RouteId == RouteId_EndCapture)
+		{
+			LastCycle = ProcessBuffer(Context.EventTime, ThreadState, BufferPtr, BufferSize);
+		}
+		else
+		{
+			LastCycle = ProcessBufferV2(Context.EventTime, ThreadState, BufferPtr, BufferSize);
+		}
+
+		if (LastCycle)
 		{
 			double LastTimestamp = Context.EventTime.AsSeconds(LastCycle);
 			Session.UpdateDurationSeconds(LastTimestamp);
-			if (RouteId == RouteId_EndCapture)
+			if (RouteId == RouteId_EndCapture || RouteId == RouteId_EndCaptureV2)
 			{
 				for (int32 i = ThreadState.ScopeStack.Num(); i--;)
 				{
@@ -227,21 +242,7 @@ uint64 FCpuProfilerAnalyzer::ProcessBuffer(const FEventTime& EventTime, FThreadS
 		if (DecodedCycle & 1ull)
 		{
 			uint32 SpecId = FTraceAnalyzerUtils::Decode7bit(BufferPtr);
-			uint32* FindIt = SpecIdToTimerIdMap.Find(SpecId);
-			uint32 TimerId;
-			if (!FindIt)
-			{
-				// Adds a timer with an "unknown" name.
-				const TCHAR* TimerName = Session.StoreString(*FString::Printf(TEXT("<unknown %u>"), SpecId));
-				// The name might be updated when an EventSpec event is received (for this SpecId),
-				// so we do not want to merge by name all the unknown timers.
-				constexpr bool bMergeByName = false;
-				TimerId = DefineTimer(SpecId, TimerName, nullptr, 0, bMergeByName);
-			}
-			else
-			{
-				TimerId = *FindIt;
-			}
+			uint32 TimerId = GetTimerId(SpecId);
 
 			FEventScopeState& ScopeState = ThreadState.ScopeStack.AddDefaulted_GetRef();
 			ScopeState.StartCycle = ActualCycle;
@@ -293,6 +294,225 @@ uint64 FCpuProfilerAnalyzer::ProcessBuffer(const FEventTime& EventTime, FThreadS
 	return LastCycle;
 }
 
+uint64 FCpuProfilerAnalyzer::ProcessBufferV2(const FEventTime& EventTime, FThreadState& ThreadState, const uint8* BufferPtr, uint32 BufferSize)
+{
+	uint64 LastCycle = ThreadState.LastCycle;
+
+	int32 RemainingPending = ThreadState.PendingEvents.Num();
+	const FPendingEvent* PendingCursor = ThreadState.PendingEvents.GetData();
+
+	const uint8* BufferEnd = BufferPtr + BufferSize;
+	while (BufferPtr < BufferEnd)
+	{
+		uint64 DecodedCycle = FTraceAnalyzerUtils::Decode7bit(BufferPtr);
+		uint64 ActualCycle = (DecodedCycle >> 2);
+
+		// ActualCycle larger or equal to LastCycle means we have a new
+		// base value.
+		if (ActualCycle < LastCycle)
+		{
+			ActualCycle += LastCycle;
+		}
+
+		// If we late connect we will be joining the cycle stream mid-flow and
+		// will have missed out on it's base timestamp. Reconstruct it here.
+		check(EventTime.GetTimestamp() == 0);
+		uint64 BaseCycle = EventTime.AsCycle64();
+		if (ActualCycle < BaseCycle)
+		{
+			ActualCycle += BaseCycle;
+		}
+
+		// Dispatch pending events that are younger than the one we've just decoded
+		for (; RemainingPending > 0; RemainingPending--, PendingCursor++)
+		{
+			bool bEnter = true;
+			uint64 PendingCycle = PendingCursor->Cycle;
+			if (int64(PendingCycle) < 0)
+			{
+				PendingCycle = ~PendingCycle;
+				bEnter = false;
+			}
+
+			if (PendingCycle > ActualCycle)
+			{
+				break;
+			}
+
+			if (PendingCycle < LastCycle)
+			{
+				PendingCycle = LastCycle;
+			}
+
+			double PendingTime = EventTime.AsSeconds(PendingCycle);
+			if (bEnter)
+			{
+				FTimingProfilerEvent Event;
+				Event.TimerIndex = PendingCursor->TimerId;
+				ThreadState.Timeline->AppendBeginEvent(PendingTime, Event);
+			}
+			else
+			{
+				ThreadState.Timeline->AppendEndEvent(PendingTime);
+			}
+		}
+
+		const double ActualTime = EventTime.AsSeconds(ActualCycle);
+
+		if (DecodedCycle & 2ull)
+		{
+			constexpr uint32 CoroutineSpecId        = (1u << 31u) - 1u;
+			constexpr uint32 CoroutineUnknownSpecId = (1u << 31u) - 2u;
+
+			if (DecodedCycle & 1ull)
+			{
+				uint64 CoroutineId = FTraceAnalyzerUtils::Decode7bit(BufferPtr);
+				uint32 TimerScopeDepth = FTraceAnalyzerUtils::Decode7bit(BufferPtr);
+
+				// Begins a "CoroTask" scoped timer.
+				{
+					if (CoroutineTimerId == ~0)
+					{
+						CoroutineTimerId = DefineNewTimerChecked(CoroutineSpecId, TEXT("Coroutine"));
+					}
+
+					TArray<uint8> CborData;
+					{
+						CborData.Reserve(256);
+						FMemoryWriter MemoryWriter(CborData, false, true);
+						FCborWriter CborWriter(&MemoryWriter, ECborEndianness::StandardCompliant);
+						CborWriter.WriteContainerStart(ECborCode::Map, 2); // 2 is the FieldCount
+						CborWriter.WriteValue("Id", 2);
+						CborWriter.WriteValue(CoroutineId);
+						CborWriter.WriteValue("C", 1); // continuation?
+						CborWriter.WriteValue(false);
+					}
+					uint32 MetadataTimerId = TimingProfilerProvider.AddMetadata(CoroutineTimerId, MoveTemp(CborData));
+
+					FEventScopeState& ScopeState = ThreadState.ScopeStack.AddDefaulted_GetRef();
+					ScopeState.StartCycle = ActualCycle;
+					ScopeState.EventTypeId = MetadataTimerId;
+
+					FTimingProfilerEvent Event;
+					Event.TimerIndex = MetadataTimerId;
+					ThreadState.Timeline->AppendBeginEvent(ActualTime, Event);
+				}
+
+				// Begins the cpu scoped timers (suspended in previous coroutine execution).
+				{
+					if (CoroutineUnknownTimerId == ~0)
+					{
+						CoroutineUnknownTimerId = DefineNewTimerChecked(CoroutineUnknownSpecId, TEXT("<unknown>"));
+					}
+
+					for (uint32 i = 0; i < TimerScopeDepth; ++i)
+					{
+						FEventScopeState& ScopeState = ThreadState.ScopeStack.AddDefaulted_GetRef();
+						ScopeState.StartCycle = ActualCycle;
+						ScopeState.EventTypeId = CoroutineUnknownTimerId;
+
+						FTimingProfilerEvent Event;
+						Event.TimerIndex = CoroutineUnknownTimerId;
+						ThreadState.Timeline->AppendBeginEvent(ActualTime, Event);
+					}
+				}
+			}
+			else
+			{
+				uint32 TimerScopeDepth = FTraceAnalyzerUtils::Decode7bit(BufferPtr);
+
+				if (TimerScopeDepth != 0)
+				{
+					// Ends (suspends) the cpu scoped timers.
+					for (uint32 i = 0; i < TimerScopeDepth; ++i)
+					{
+						// If we receive mismatched end events ignore them for now.
+						// This can happen for example because tracing connects to the store after events were traced. Those events can be lost.
+						if (ThreadState.ScopeStack.Num() > 0)
+						{
+							ThreadState.ScopeStack.Pop();
+							ThreadState.Timeline->AppendEndEvent(ActualTime);
+						}
+					}
+
+					// Update the "continuation" (suspended or destroyed) metadata flag.
+					if (ThreadState.ScopeStack.Num() > 0)
+					{
+						uint32 MetadataTimerId = ThreadState.ScopeStack.Top().EventTypeId;
+						TArrayView<const uint8> Metadata = TimingProfilerProvider.GetMetadata(MetadataTimerId);
+						// Change the last byte in metadata to "true".
+						const_cast<uint8*>(Metadata.GetData())[Metadata.Num() - 1] = (uint8)(ECborCode::Prim | ECborCode::True);
+					}
+				}
+
+				// Ends the "CoroTask" scoped timer.
+				{
+					// If we receive mismatched end events ignore them for now.
+					// This can happen for example because tracing connects to the store after events were traced. Those events can be lost.
+					if (ThreadState.ScopeStack.Num() > 0)
+					{
+						ThreadState.ScopeStack.Pop();
+						ThreadState.Timeline->AppendEndEvent(ActualTime);
+					}
+				}
+			}
+		}
+		else
+		{
+			if (DecodedCycle & 1ull)
+			{
+				uint32 SpecId = FTraceAnalyzerUtils::Decode7bit(BufferPtr);
+				uint32 TimerId = GetTimerId(SpecId);
+
+				FEventScopeState& ScopeState = ThreadState.ScopeStack.AddDefaulted_GetRef();
+				ScopeState.StartCycle = ActualCycle;
+				ScopeState.EventTypeId = TimerId;
+
+				FTimingProfilerEvent Event;
+				Event.TimerIndex = TimerId;
+				ThreadState.Timeline->AppendBeginEvent(ActualTime, Event);
+				++TotalScopeCount;
+			}
+			else
+			{
+				// If we receive mismatched end events ignore them for now.
+				// This can happen for example because tracing connects to the store after events were traced. Those events can be lost.
+				if (ThreadState.ScopeStack.Num() > 0)
+				{
+					ThreadState.ScopeStack.Pop();
+					ThreadState.Timeline->AppendEndEvent(ActualTime);
+				}
+			}
+		}
+
+		LastCycle = ActualCycle;
+	}
+	check(BufferPtr == BufferEnd);
+
+	// Dispatch remaining pending events.
+	for (; RemainingPending > 0; RemainingPending--, PendingCursor++)
+	{
+		uint64 PendingCycle = PendingCursor->Cycle;
+		if (int64(PendingCycle) < 0)
+		{
+			PendingCycle = ~PendingCycle;
+			double PendingTime = EventTime.AsSeconds(PendingCycle);
+			ThreadState.Timeline->AppendEndEvent(PendingTime);
+		}
+		else
+		{
+			FTimingProfilerEvent Event;
+			Event.TimerIndex = PendingCursor->TimerId;
+			double PendingTime = EventTime.AsSeconds(PendingCycle);
+			ThreadState.Timeline->AppendBeginEvent(PendingTime, Event);
+		}
+	}
+
+	ThreadState.PendingEvents.Reset();
+	ThreadState.LastCycle = LastCycle;
+	return LastCycle;
+}
+
 void FCpuProfilerAnalyzer::OnCpuScopeEnter(const FOnEventContext& Context)
 {
 	if (Context.EventTime.GetTimestamp() == 0)
@@ -405,6 +625,29 @@ uint32 FCpuProfilerAnalyzer::DefineTimer(uint32 SpecId, const TCHAR* Name, const
 	}
 }
 
+uint32 FCpuProfilerAnalyzer::DefineNewTimerChecked(uint32 SpecId, const TCHAR* TimerName, const TCHAR* File, uint32 Line)
+{
+	TimerName = Session.StoreString(TimerName);
+	uint32 NewTimerId = TimingProfilerProvider.AddCpuTimer(TimerName, File, Line);
+	SpecIdToTimerIdMap.Add(SpecId, NewTimerId);
+	return NewTimerId;
+}
+
+uint32 FCpuProfilerAnalyzer::GetTimerId(uint32 SpecId)
+{
+	if (uint32* FindIt = SpecIdToTimerIdMap.Find(SpecId))
+	{
+		return *FindIt;
+	}
+	else
+	{
+		// Adds a timer with an "unknown" name.
+		// The "unknown" timers are not merged by name, becasue the actual name
+		// might be updated when an EventSpec event is received (for this SpecId).
+		return DefineNewTimerChecked(SpecId, *FString::Printf(TEXT("<unknown %u>"), SpecId));
+	}
+}
+
 FCpuProfilerAnalyzer::FThreadState& FCpuProfilerAnalyzer::GetThreadState(uint32 ThreadId)
 {
 	FThreadState* ThreadState = ThreadStatesMap.FindRef(ThreadId);
diff --git a/Engine/Source/Developer/TraceServices/Private/Analyzers/CpuProfilerTraceAnalysis.h b/Engine/Source/Developer/TraceServices/Private/Analyzers/CpuProfilerTraceAnalysis.h
--- a/Engine/Source/Developer/TraceServices/Private/Analyzers/CpuProfilerTraceAnalysis.h
+++ b/Engine/Source/Developer/TraceServices/Private/Analyzers/CpuProfilerTraceAnalysis.h
@@ -3,6 +3,7 @@
 #pragma once
 
 #include "Trace/Analyzer.h"
+#include "Containers/Map.h"
 #include "Containers/UnrealString.h"
 #include "Model/TimingProfilerPrivate.h"
 
@@ -44,16 +45,21 @@ private:
 
 	void OnCpuScopeEnter(const FOnEventContext& Context);
 	void OnCpuScopeLeave(const FOnEventContext& Context);
-	uint32 DefineTimer(uint32 SpecId, const TCHAR* ScopeName, const TCHAR* File, uint32 Line, bool bMergeByName); // returns the TimerId
+	uint32 DefineTimer(uint32 SpecId, const TCHAR* TimerName, const TCHAR* File, uint32 Line, bool bMergeByName); // returns the TimerId
+	uint32 DefineNewTimerChecked(uint32 SpecId, const TCHAR* TimerName, const TCHAR* File = nullptr, uint32 Line = 0); // returns the TimerId
+	uint32 GetTimerId(uint32 SpecId);
 	FThreadState& GetThreadState(uint32 ThreadId);
 	uint64 ProcessBuffer(const FEventTime& EventTime, FThreadState& ThreadState, const uint8* BufferPtr, uint32 BufferSize);
+	uint64 ProcessBufferV2(const FEventTime& EventTime, FThreadState& ThreadState, const uint8* BufferPtr, uint32 BufferSize);
 
 	enum : uint16
 	{
 		RouteId_EventSpec,
 		RouteId_EventBatch,
+		RouteId_EventBatchV2,
 		RouteId_EndThread,
 		RouteId_EndCapture,
+		RouteId_EndCaptureV2,
 		RouteId_CpuScope,
 		RouteId_ChannelAnnounce,
 		RouteId_ChannelToggle,
@@ -65,6 +71,8 @@ private:
 	TMap<uint32, FThreadState*> ThreadStatesMap;
 	TMap<uint32, uint32> SpecIdToTimerIdMap;
 	TMap<const TCHAR*, uint32> ScopeNameToTimerIdMap;
+	uint32 CoroutineTimerId = ~0;
+	uint32 CoroutineUnknownTimerId = ~0;
 	uint64 TotalEventSize = 0;
 	uint64 TotalScopeCount = 0;
 	double BytesPerScope = 0.0;
diff --git a/Engine/Source/Developer/TraceServices/TraceServices.Build.cs b/Engine/Source/Developer/TraceServices/TraceServices.Build.cs
--- a/Engine/Source/Developer/TraceServices/TraceServices.Build.cs
+++ b/Engine/Source/Developer/TraceServices/TraceServices.Build.cs
@@ -8,6 +8,7 @@ public class TraceServices : ModuleRules
 	{
 		PublicDependencyModuleNames.AddRange(
 			new string[] {
+				"Cbor",
 				"Core",
 				"TraceAnalysis",
 				"SymsLib"
diff --git a/Engine/Source/Editor/UnrealEd/Private/Commandlets/SummarizeTraceCommandlet.cpp b/Engine/Source/Editor/UnrealEd/Private/Commandlets/SummarizeTraceCommandlet.cpp
--- a/Engine/Source/Editor/UnrealEd/Private/Commandlets/SummarizeTraceCommandlet.cpp
+++ b/Engine/Source/Editor/UnrealEd/Private/Commandlets/SummarizeTraceCommandlet.cpp
@@ -91,6 +91,9 @@ public:
 
 	/** Invoked when the trace stream has been fully consumed/processed. */
 	virtual void OnCpuScopeAnalysisEnd() {};
+
+	static constexpr uint32 CoroutineSpecId        = (1u << 31u) - 1u;
+	static constexpr uint32 CoroutineUnknownSpecId = (1u << 31u) - 2u;
 };
 
 /**
@@ -154,6 +157,7 @@ private:
 	// Internal implementation.
 	void OnEventSpec(const FOnEventContext& Context);
 	void OnBatch(const FOnEventContext& Context);
+	void OnBatchV2(const FOnEventContext& Context);
 	void OnCpuScopeSpec(uint32 ScopeId, const FString* ScopeName);
 	void OnCpuScopeEnter(uint32 ScopeId, uint32 ThreadId, double Timestamp);
 	void OnCpuScopeExit(uint32 ThreadId, double Timestamp);
@@ -172,6 +176,8 @@ private:
 
 	// Scope name lookup function, cached for efficiency.
 	TFunction<const FString*(uint32 ScopeId)> LookupScopeNameFn;
+
+
 };
 
 enum
@@ -180,6 +186,8 @@ enum
 	RouteId_CpuProfiler_EventSpec,
 	RouteId_CpuProfiler_EventBatch,
 	RouteId_CpuProfiler_EndCapture,
+	RouteId_CpuProfiler_EventBatchV2,
+	RouteId_CpuProfiler_EndCaptureV2
 };
 
 FCpuScopeStreamProcessor::FCpuScopeStreamProcessor()
@@ -197,6 +205,8 @@ void FCpuScopeStreamProcessor::OnAnalysisBegin(const FOnAnalysisContext& Context
 	Context.InterfaceBuilder.RouteEvent(RouteId_CpuProfiler_EventSpec, "CpuProfiler", "EventSpec");
 	Context.InterfaceBuilder.RouteEvent(RouteId_CpuProfiler_EventBatch, "CpuProfiler", "EventBatch");
 	Context.InterfaceBuilder.RouteEvent(RouteId_CpuProfiler_EndCapture, "CpuProfiler", "EndCapture");
+	Context.InterfaceBuilder.RouteEvent(RouteId_CpuProfiler_EventBatchV2, "CpuProfiler", "EventBatchV2");
+	Context.InterfaceBuilder.RouteEvent(RouteId_CpuProfiler_EndCaptureV2, "CpuProfiler", "EndCaptureV2");
 }
 
 void FCpuScopeStreamProcessor::OnAnalysisEnd()
@@ -233,6 +243,10 @@ bool FCpuScopeStreamProcessor::OnEvent(uint16 RouteId, EStyle Style, const FOnEv
 	case RouteId_CpuProfiler_EndCapture:
 		OnBatch(Context);
 		break;
+	case RouteId_CpuProfiler_EventBatchV2:
+	case RouteId_CpuProfiler_EndCaptureV2:
+		OnBatchV2(Context);
+		break;
 	};
 
 	return true;
@@ -277,6 +291,61 @@ void FCpuScopeStreamProcessor::OnBatch(const FOnEventContext& Context)
 	}
 }
 
+void FCpuScopeStreamProcessor::OnBatchV2(const FOnEventContext& Context)
+{
+	const FEventData& EventData = Context.EventData;
+	const FEventTime& EventTime = Context.EventTime;
+
+	uint32 ThreadId = Context.ThreadInfo.GetId();
+
+	TArrayView<const uint8> DataView = TraceServices::FTraceAnalyzerUtils::LegacyAttachmentArray("Data", Context);
+	const uint8* Cursor = DataView.GetData();
+	const uint8* End = Cursor + DataView.Num();
+	uint64 LastCycle = 0;
+	while (Cursor < End)
+	{
+		uint64 Value = Decode7bit(Cursor);
+		uint64 Cycle = LastCycle + (Value >> 2);
+		LastCycle = Cycle;
+		double TimeStamp = EventTime.AsSeconds(Cycle);
+
+		if (Value & 2ull)
+		{		
+			if (Value & 1ull)			
+			{
+				uint64 CoroutineId = Decode7bit(Cursor);
+				uint32 TimerScopeDepth = Decode7bit(Cursor);
+				OnCpuScopeEnter(FCpuScopeAnalyzer::CoroutineSpecId, ThreadId, TimeStamp);
+				for (uint32 i = 0; i < TimerScopeDepth; ++i)
+				{
+					OnCpuScopeEnter(FCpuScopeAnalyzer::CoroutineUnknownSpecId, ThreadId, TimeStamp);
+				}
+			}
+			else
+			{			
+				uint32 TimerScopeDepth = Decode7bit(Cursor);
+				for (uint32 i = 0; i < TimerScopeDepth; ++i)
+				{
+					OnCpuScopeExit(ThreadId, TimeStamp);
+				}
+				OnCpuScopeExit(ThreadId, TimeStamp);
+			}
+		}
+		else
+		{
+			if (Value & 1)
+			{
+				uint64 ScopeId = Decode7bit(Cursor);
+				OnCpuScopeEnter(static_cast<uint32>(ScopeId - 1), ThreadId, TimeStamp);
+			}
+			else
+			{
+				OnCpuScopeExit(ThreadId, TimeStamp);
+			}
+		}
+	}
+}
+
 void FCpuScopeStreamProcessor::OnCpuScopeSpec(uint32 ScopeId, const FString* ScopeName)
 {
 	if (ScopeId >= uint32(ScopeNames.Num()))
@@ -828,10 +897,7 @@ protected:
 	virtual void OnCpuScopeAnalysisEnd() override;
 
 private:
-	// Scopes is an array only because indexes are doled out on the process-side
-	// As such, there could be different scopes with the same name
-	// We merge these together later to ensure name is a pkey in the csv
-	TArray<FSummarizeScope> Scopes;
+	TMap<uint32, FSummarizeScope> Scopes;
 
 	// Function invoked to publish the list of scopes.
 	TFunction<void(const TArray<FSummarizeScope>&)> PublishFn;
@@ -840,14 +906,17 @@ private:
 FSummarizeCpuAnalyzer::FSummarizeCpuAnalyzer(TFunction<void(const TArray<FSummarizeScope>&)> InPublishFn)
 	: PublishFn(MoveTemp(InPublishFn))
 {
+	OnCpuScopeDiscovered(FCpuScopeAnalyzer::CoroutineSpecId);
+	OnCpuScopeDiscovered(FCpuScopeAnalyzer::CoroutineUnknownSpecId);
+	OnCpuScopeName(FCpuScopeAnalyzer::CoroutineSpecId, TEXT("Coroutine"));
+	OnCpuScopeName(FCpuScopeAnalyzer::CoroutineUnknownSpecId, TEXT("<unknown>"));
 }
 
 void FSummarizeCpuAnalyzer::OnCpuScopeDiscovered(uint32 ScopeId)
 {
-	if (ScopeId >= uint32(Scopes.Num()))
+	if (!Scopes.Find(ScopeId))
 	{
-		uint32 Num = (ScopeId + 128) & ~127;
-		Scopes.SetNum(Num);
+		Scopes.Add(ScopeId, FSummarizeScope());		
 	}
 }
 
@@ -863,11 +932,18 @@ void FSummarizeCpuAnalyzer::OnCpuScopeExit(const FScope& Scope, const FString* S
 
 void FSummarizeCpuAnalyzer::OnCpuScopeAnalysisEnd()
 {
+	TArray<FSummarizeScope> LocalScopes;
 	// Eliminates scopes that don't have a name. (On scope discovery, the array is expended to creates blank scopes that may never be filled).
-	Scopes.RemoveAll([](const FSummarizeScope& Scope) { return Scope.Name.IsEmpty(); });
-
+	for (TMap<uint32, FSummarizeScope>::TIterator Iter = Scopes.CreateIterator(); Iter; ++Iter)
+	{
+		if(!Iter.Value().Name.IsEmpty())
+		{
+			LocalScopes.Add(Iter.Value());
+		}
+	}
+	
 	// Publish the scopes.
-	PublishFn(Scopes);
+	PublishFn(LocalScopes);
 }
 
 
diff --git a/Engine/Source/Runtime/Core/Private/ProfilingDebugging/CpuProfilerTrace.cpp b/Engine/Source/Runtime/Core/Private/ProfilingDebugging/CpuProfilerTrace.cpp
--- a/Engine/Source/Runtime/Core/Private/ProfilingDebugging/CpuProfilerTrace.cpp
+++ b/Engine/Source/Runtime/Core/Private/ProfilingDebugging/CpuProfilerTrace.cpp
@@ -24,11 +24,11 @@ UE_TRACE_EVENT_BEGIN(CpuProfiler, EventSpec, NoSync|Important)
 	#endif
 UE_TRACE_EVENT_END()
 
-UE_TRACE_EVENT_BEGIN(CpuProfiler, EventBatch, NoSync)
+UE_TRACE_EVENT_BEGIN(CpuProfiler, EventBatchV2, NoSync)
 	UE_TRACE_EVENT_FIELD(uint8[], Data)
 UE_TRACE_EVENT_END()
 
-UE_TRACE_EVENT_BEGIN(CpuProfiler, EndCapture)
+UE_TRACE_EVENT_BEGIN(CpuProfiler, EndCaptureV2)
 	UE_TRACE_EVENT_FIELD(uint8[], Data)
 UE_TRACE_EVENT_END()
 
@@ -99,12 +99,19 @@ struct FCpuProfilerTraceInternal
 	FORCENOINLINE static void FlushThreadBuffer(FThreadBuffer* ThreadBuffer);
 	FORCENOINLINE static void EndCapture(FThreadBuffer* ThreadBuffer);
 
-	static thread_local uint32 ThreadDepth;
+	struct FSuspendScopes
+	{
+		uint32* TimerScopeDepth;
+		uint32 SavedThreadDepth;
+	};
+	static thread_local TArray<FSuspendScopes, TInlineAllocator<3>> NestedTimerScopeDepths;
 	static thread_local FThreadBuffer* ThreadBuffer;
+	static thread_local uint32 ThreadDepth;
 };
 
-thread_local uint32 FCpuProfilerTraceInternal::ThreadDepth = 0;
+thread_local TArray<FCpuProfilerTraceInternal::FSuspendScopes, TInlineAllocator<3>> FCpuProfilerTraceInternal::NestedTimerScopeDepths;
 thread_local FCpuProfilerTraceInternal::FThreadBuffer* FCpuProfilerTraceInternal::ThreadBuffer = nullptr;
+thread_local uint32 FCpuProfilerTraceInternal::ThreadDepth = 0;
 
 FCpuProfilerTraceInternal::FThreadBuffer* FCpuProfilerTraceInternal::CreateThreadBuffer()
 {
@@ -115,16 +122,16 @@ FCpuProfilerTraceInternal::FThreadBuffer* FCpuProfilerTraceInternal::CreateThrea
 
 void FCpuProfilerTraceInternal::FlushThreadBuffer(FThreadBuffer* InThreadBuffer)
 {
-	UE_TRACE_LOG(CpuProfiler, EventBatch, true)
-		<< EventBatch.Data(InThreadBuffer->Buffer, InThreadBuffer->BufferSize);
+	UE_TRACE_LOG(CpuProfiler, EventBatchV2, true)
+		<< EventBatchV2.Data(InThreadBuffer->Buffer, InThreadBuffer->BufferSize);
 	InThreadBuffer->BufferSize = 0;
 	InThreadBuffer->LastCycle = 0;
 }
 
 void FCpuProfilerTraceInternal::EndCapture(FThreadBuffer* InThreadBuffer)
 {
-	UE_TRACE_LOG(CpuProfiler, EndCapture, true)
-		<< EndCapture.Data(InThreadBuffer->Buffer, InThreadBuffer->BufferSize);
+	UE_TRACE_LOG(CpuProfiler, EndCaptureV2, true)
+		<< EndCaptureV2.Data(InThreadBuffer->Buffer, InThreadBuffer->BufferSize);
 	InThreadBuffer->BufferSize = 0;
 	InThreadBuffer->LastCycle = 0;
 }
@@ -142,7 +149,7 @@ void FCpuProfilerTraceInternal::EndCapture(FThreadBuffer* InThreadBuffer)
 	uint64 CycleDiff = Cycle - ThreadBuffer->LastCycle; \
 	ThreadBuffer->LastCycle = Cycle; \
 	uint8* BufferPtr = ThreadBuffer->Buffer + ThreadBuffer->BufferSize; \
-	FTraceUtils::Encode7bit((CycleDiff << 1) | 1ull, BufferPtr); \
+	FTraceUtils::Encode7bit((CycleDiff << 2) | 1ull, BufferPtr); \
 	FTraceUtils::Encode7bit(SpecId, BufferPtr); \
 	ThreadBuffer->BufferSize = (uint16)(BufferPtr - ThreadBuffer->Buffer); \
 	if (ThreadBuffer->BufferSize >= FCpuProfilerTraceInternal::FullBufferThreshold) \
@@ -210,6 +217,54 @@ void FCpuProfilerTrace::OutputBeginDynamicEvent(const FName& Name, const ANSICHA
 	CPUPROFILERTRACE_OUTPUTBEGINEVENT_EPILOGUE();
 }
 
+void FCpuProfilerTrace::OutputResumeEvent(uint64 SpecId, uint32& TimerScopeDepth)
+{
+	FCpuProfilerTraceInternal::NestedTimerScopeDepths.Push({&TimerScopeDepth, FCpuProfilerTraceInternal::ThreadDepth});
+	FCpuProfilerTraceInternal::ThreadDepth = FCpuProfilerTraceInternal::ThreadDepth + TimerScopeDepth;
+
+	FCpuProfilerTraceInternal::FThreadBuffer* ThreadBuffer = FCpuProfilerTraceInternal::ThreadBuffer;
+	if (!ThreadBuffer)
+	{
+		ThreadBuffer = FCpuProfilerTraceInternal::CreateThreadBuffer();
+	}
+	uint64 Cycle = FPlatformTime::Cycles64();
+	uint64 CycleDiff = Cycle - ThreadBuffer->LastCycle;
+	ThreadBuffer->LastCycle = Cycle;
+	uint8* BufferPtr = ThreadBuffer->Buffer + ThreadBuffer->BufferSize;
+	FTraceUtils::Encode7bit(CycleDiff << 2 | 3ull, BufferPtr);
+	FTraceUtils::Encode7bit(SpecId, BufferPtr);
+	FTraceUtils::Encode7bit(TimerScopeDepth, BufferPtr);
+	ThreadBuffer->BufferSize = (uint16)(BufferPtr - ThreadBuffer->Buffer);
+	if (ThreadBuffer->BufferSize >= FCpuProfilerTraceInternal::FullBufferThreshold)
+	{
+		FCpuProfilerTraceInternal::FlushThreadBuffer(ThreadBuffer);
+	}
+}
+
+void FCpuProfilerTrace::OutputSuspendEvent()
+{
+	auto [TimerScopeDepth, SavedThreadDepth] = FCpuProfilerTraceInternal::NestedTimerScopeDepths.Pop();
+	*TimerScopeDepth = FCpuProfilerTraceInternal::ThreadDepth - SavedThreadDepth;
+	FCpuProfilerTraceInternal::ThreadDepth = SavedThreadDepth;
+
+	FCpuProfilerTraceInternal::FThreadBuffer* ThreadBuffer = FCpuProfilerTraceInternal::ThreadBuffer;
+	if (!ThreadBuffer)
+	{
+		ThreadBuffer = FCpuProfilerTraceInternal::CreateThreadBuffer();
+	}
+	uint64 Cycle = FPlatformTime::Cycles64();
+	uint64 CycleDiff = Cycle - ThreadBuffer->LastCycle;
+	ThreadBuffer->LastCycle = Cycle;
+	uint8* BufferPtr = ThreadBuffer->Buffer + ThreadBuffer->BufferSize;
+	FTraceUtils::Encode7bit(CycleDiff << 2 | 2ull, BufferPtr);
+	FTraceUtils::Encode7bit(uint64(*TimerScopeDepth), BufferPtr);
+	ThreadBuffer->BufferSize = (uint16)(BufferPtr - ThreadBuffer->Buffer);
+	if ((FCpuProfilerTraceInternal::ThreadDepth == 0) | (ThreadBuffer->BufferSize >= FCpuProfilerTraceInternal::FullBufferThreshold))
+	{
+		FCpuProfilerTraceInternal::FlushThreadBuffer(ThreadBuffer);
+	}
+}
+
 #undef CPUPROFILERTRACE_OUTPUTBEGINEVENT_PROLOGUE
 #undef CPUPROFILERTRACE_OUTPUTBEGINEVENT_EPILOGUE
 
@@ -225,7 +280,7 @@ void FCpuProfilerTrace::OutputEndEvent()
 	uint64 CycleDiff = Cycle - ThreadBuffer->LastCycle;
 	ThreadBuffer->LastCycle = Cycle;
 	uint8* BufferPtr = ThreadBuffer->Buffer + ThreadBuffer->BufferSize;
-	FTraceUtils::Encode7bit(CycleDiff << 1, BufferPtr);
+	FTraceUtils::Encode7bit(CycleDiff << 2, BufferPtr);
 	ThreadBuffer->BufferSize = (uint16)(BufferPtr - ThreadBuffer->Buffer);
 	if ((FCpuProfilerTraceInternal::ThreadDepth == 0) | (ThreadBuffer->BufferSize >= FCpuProfilerTraceInternal::FullBufferThreshold))
 	{
diff --git a/Engine/Source/Runtime/Core/Public/ProfilingDebugging/CpuProfilerTrace.h b/Engine/Source/Runtime/Core/Public/ProfilingDebugging/CpuProfilerTrace.h
--- a/Engine/Source/Runtime/Core/Public/ProfilingDebugging/CpuProfilerTrace.h
+++ b/Engine/Source/Runtime/Core/Public/ProfilingDebugging/CpuProfilerTrace.h
@@ -87,6 +87,16 @@ struct FCpuProfilerTrace
 	 * Output end event marker for static or dynamic event for the currently open scope.
 	 */
 	CORE_API static void OutputEndEvent();
+	/*
+	* Output resume marker for a given spec. Must always be matched with an suspend event.
+	* @param SpecId unique Resume Event definition id.
+	* @param TimerScopeDepth updates the depth of the current OutputBeginEvent depth.
+	*/
+	CORE_API static void OutputResumeEvent(uint64 SpecId, uint32& TimerScopeDepth);
+	/*
+	* Output suspend event marker for the currently open resume event.
+	*/
+	CORE_API static void OutputSuspendEvent();
 
 	class FEventScope
 	{
