New BitCast<>() function which works like C++20's std::bit_cast<>().
Deprecation of FPlatformMath::IsNegative*() functions.
New FPlatformMath::IsNegativeOrNegativeZero().
Fix up of existing usage to either use < 0 or IsNegativeOrNegativeZero where appropriate.
Fixes for aliasing problems in various FMath functions, including IsNegative*().

diff --git a/Engine/Plugins/Editor/GeometryMode/Source/GeometryMode/Private/GeometryModifiers.cpp b/Engine/Plugins/Editor/GeometryMode/Source/GeometryMode/Private/GeometryModifiers.cpp
--- a/Engine/Plugins/Editor/GeometryMode/Source/GeometryMode/Private/GeometryModifiers.cpp
+++ b/Engine/Plugins/Editor/GeometryMode/Source/GeometryMode/Private/GeometryModifiers.cpp
@@ -2819,7 +2819,7 @@ bool UGeomModifier_Split::OnApply()
 
 						// Make sure the line formed by the edge actually crosses the plane before checking for the intersection point.
 
-						if( FMath::IsNegativeFloat( CuttingPlane.PlaneDot( *v0 ) ) != FMath::IsNegativeFloat( CuttingPlane.PlaneDot( *v1 ) ) )
+						if( ( CuttingPlane.PlaneDot( *v0 ) < 0.0f ) != ( CuttingPlane.PlaneDot( *v1 ) < 0.0f ) )
 						{
 							FVector3f Intersection = FMath::LinePlaneIntersection( *v0, *v1, CuttingPlane );
 
diff --git a/Engine/Source/Developer/AutomationDriver/Private/DriverSequence.cpp b/Engine/Source/Developer/AutomationDriver/Private/DriverSequence.cpp
--- a/Engine/Source/Developer/AutomationDriver/Private/DriverSequence.cpp
+++ b/Engine/Source/Developer/AutomationDriver/Private/DriverSequence.cpp
@@ -976,11 +976,11 @@ private:
 				return Result;
 			}
 
-			if (FMath::IsNegativeFloat(Delta) && ScrollableElement->IsScrolledToEnd())
+			if (Delta < 0.0f && ScrollableElement->IsScrolledToEnd())
 			{
 				return FStep::Failed();
 			}
-			else if (!FMath::IsNegativeFloat(Delta) && ScrollableElement->IsScrolledToBeginning())
+			else if (!(Delta < 0.0f) && ScrollableElement->IsScrolledToBeginning())
 			{
 				return FStep::Failed();
 			}
diff --git a/Engine/Source/Developer/MeshMergeUtilities/Private/MeshMergeHelpers.cpp b/Engine/Source/Developer/MeshMergeUtilities/Private/MeshMergeHelpers.cpp
--- a/Engine/Source/Developer/MeshMergeUtilities/Private/MeshMergeHelpers.cpp
+++ b/Engine/Source/Developer/MeshMergeUtilities/Private/MeshMergeHelpers.cpp
@@ -559,7 +559,7 @@ bool FMeshMergeHelpers::CheckWrappingUVs(const TArray<FVector2D>& UVs)
 	FVector2D Max(-FLT_MAX, -FLT_MAX);
 	for (const FVector2D& Coordinate : UVs)
 	{
-		if ((FMath::IsNegativeFloat(Coordinate.X) || FMath::IsNegativeFloat(Coordinate.Y)) || (Coordinate.X > (1.0f + KINDA_SMALL_NUMBER) || Coordinate.Y > (1.0f + KINDA_SMALL_NUMBER)))
+		if ((Coordinate.X < 0.0f || Coordinate.Y < 0.0f) || (Coordinate.X > (1.0f + KINDA_SMALL_NUMBER) || Coordinate.Y > (1.0f + KINDA_SMALL_NUMBER)))
 		{
 			bResult = true;
 			break;
@@ -583,7 +583,7 @@ bool FMeshMergeHelpers::CheckWrappingUVs(const FMeshDescription& MeshDescription
 	for (const FVertexInstanceID VertexInstanceID : MeshDescription.VertexInstances().GetElementIDs())
 	{
 		const FVector2D& Coordinate = FVector2D(VertexInstanceUVs.Get(VertexInstanceID, UVChannelIndex));
-		if ((FMath::IsNegativeFloat(Coordinate.X) || FMath::IsNegativeFloat(Coordinate.Y)) || (Coordinate.X > (1.0f + KINDA_SMALL_NUMBER) || Coordinate.Y > (1.0f + KINDA_SMALL_NUMBER)))
+		if ((Coordinate.X < 0.0f || Coordinate.Y < 0.0f) || (Coordinate.X > (1.0f + KINDA_SMALL_NUMBER) || Coordinate.Y > (1.0f + KINDA_SMALL_NUMBER)))
 		{
 			bResult = true;
 			break;
diff --git a/Engine/Source/Developer/SkeletalMeshUtilitiesCommon/Private/LODUtilities.cpp b/Engine/Source/Developer/SkeletalMeshUtilitiesCommon/Private/LODUtilities.cpp
--- a/Engine/Source/Developer/SkeletalMeshUtilitiesCommon/Private/LODUtilities.cpp
+++ b/Engine/Source/Developer/SkeletalMeshUtilitiesCommon/Private/LODUtilities.cpp
@@ -513,7 +513,7 @@ void FLODUtilities::RemoveLODs(FSkeletalMeshUpdateContext& UpdateContext, const
 /** Given three direction vectors, indicates if A and B are on the same 'side' of Vec. */
 bool VectorsOnSameSide(const FVector2D& Vec, const FVector2D& A, const FVector2D& B)
 {
-	return !FMath::IsNegativeFloat(((B.Y - A.Y)*(Vec.X - A.X)) + ((A.X - B.X)*(Vec.Y - A.Y)));
+	return !(((B.Y - A.Y)*(Vec.X - A.X)) + ((A.X - B.X)*(Vec.Y - A.Y)) < 0.0f);
 }
 
 float PointToSegmentDistanceSquare(const FVector2D& A, const FVector2D& B, const FVector2D& P)
@@ -560,7 +560,7 @@ bool VectorsOnSameSide(const FVector3f& Vec, const FVector3f& A, const FVector3f
 	const FVector CrossA = FVector(Vec ^ A);
 	const FVector CrossB = FVector(Vec ^ B);
 	float DotWithEpsilon = SameSideDotProductEpsilon + (CrossA | CrossB);
-	return !FMath::IsNegativeFloat(DotWithEpsilon);
+	return !(DotWithEpsilon < 0.0f);
 }
 
 /** Util to see if P lies within triangle created by A, B and C. */
diff --git a/Engine/Source/Developer/TraceInsights/Private/Insights/Common/TimeUtils.cpp b/Engine/Source/Developer/TraceInsights/Private/Insights/Common/TimeUtils.cpp
--- a/Engine/Source/Developer/TraceInsights/Private/Insights/Common/TimeUtils.cpp
+++ b/Engine/Source/Developer/TraceInsights/Private/Insights/Common/TimeUtils.cpp
@@ -66,7 +66,7 @@ FString FormatTimeAuto(const double InDuration, const int32 NumDigits)
 
 	TStringBuilder<64> StrBuilder;
 
-	if (FMath::IsNegative(Duration))
+	if (Duration < 0.0)
 	{
 		Duration = -Duration;
 		StrBuilder.AppendChar(TEXT('-'));
@@ -232,7 +232,7 @@ FString FormatTimeMs(const double InDuration, const int32 NumDigits, bool bAddTi
 	double Duration = InDuration;
 	TStringBuilder<64> StrBuilder;
 
-	if (FMath::IsNegative(Duration))
+	if (Duration < 0.0)
 	{
 		Duration = -Duration;
 		StrBuilder.AppendChar(TEXT('-'));
@@ -286,7 +286,7 @@ FString FormatTime(const double InTime, const double Precision)
 	double Time = InTime;
 	TStringBuilder<64> StrBuilder;
 
-	if (FMath::IsNegative(Time))
+	if (Time < 0.0)
 	{
 		Time = -Time;
 		StrBuilder.AppendChar(TEXT('-'));
@@ -469,7 +469,7 @@ void SplitTime(const double InTime, FTimeSplit& OutTimeSplit)
 		return;
 	}
 
-	if (FMath::IsNegative(Time))
+	if (Time < 0.0)
 	{
 		Time = -Time;
 		OutTimeSplit.bIsNegative = true;
diff --git a/Engine/Source/Developer/TraceServices/Public/Model/AsyncEnumerateTask.h b/Engine/Source/Developer/TraceServices/Public/Model/AsyncEnumerateTask.h
--- a/Engine/Source/Developer/TraceServices/Public/Model/AsyncEnumerateTask.h
+++ b/Engine/Source/Developer/TraceServices/Public/Model/AsyncEnumerateTask.h
@@ -83,7 +83,7 @@ public:
 		// Iterate the pages owned by this task.
 		while (ScopeEntry && CurrentPageIndex < Data.EndPageIndex)
 		{
-			if (FMath::IsNegative(ScopeEntry->Time))
+			if (ScopeEntry->Time < 0.0)
 			{
 				check(CurrentStackDepth < SettingsType::MaxDepth);
 				FEnumerationStackEntry& StackEntry = EventStack[CurrentStackDepth++];
diff --git a/Engine/Source/Developer/TraceServices/Public/Model/MonotonicTimeline.h b/Engine/Source/Developer/TraceServices/Public/Model/MonotonicTimeline.h
--- a/Engine/Source/Developer/TraceServices/Public/Model/MonotonicTimeline.h
+++ b/Engine/Source/Developer/TraceServices/Public/Model/MonotonicTimeline.h
@@ -247,7 +247,7 @@ public:
 		const EventType* Event = EventsIterator.GetCurrentItem();
 		while (ScopeEntry && FMath::Abs(ScopeEntry->Time) < IntervalStart)
 		{
-			if (FMath::IsNegative(ScopeEntry->Time))
+			if (ScopeEntry->Time < 0.0)
 			{
 				check(CurrentStackDepth < SettingsType::MaxDepth);
 				FEnumerationStackEntry& StackEntry = EventStack[CurrentStackDepth++];
@@ -276,7 +276,7 @@ public:
 		}
 		while (ScopeEntry && FMath::Abs(ScopeEntry->Time) <= IntervalEnd)
 		{
-			if (FMath::IsNegative(ScopeEntry->Time))
+			if (ScopeEntry->Time < 0.0)
 			{
 				check(CurrentStackDepth < SettingsType::MaxDepth);
 				FEnumerationStackEntry& StackEntry = EventStack[CurrentStackDepth++];
@@ -309,7 +309,7 @@ public:
 				bSearchEndTimeUsingPages = true;
 				break;
 			}
-			if (FMath::IsNegative(ScopeEntry->Time))
+			if (ScopeEntry->Time < 0.0)
 			{
 				++ExitDepth;
 			}
@@ -452,7 +452,7 @@ public:
 		const EventType* Event = EventsIterator.GetCurrentItem();
 		while (ScopeEntry && FMath::Abs(ScopeEntry->Time) > IntervalEnd)
 		{
-			if (FMath::IsNegative(ScopeEntry->Time))
+			if (ScopeEntry->Time < 0.0)
 			{
 				check(CurrentStackDepth > 0);
 				--CurrentStackDepth;
@@ -480,7 +480,7 @@ public:
 		// Enumerate backwards between IntervalEnd and IntervalStart.
 		while (ScopeEntry && FMath::Abs(ScopeEntry->Time) >= IntervalStart)
 		{
-			if (FMath::IsNegative(ScopeEntry->Time))
+			if (ScopeEntry->Time < 0.0)
 			{
 				check(CurrentStackDepth > 0);
 				FEnumerationStackEntry& StackEntry = EventStack[--CurrentStackDepth];
@@ -516,7 +516,7 @@ public:
 				bSearchStartTimeUsingPages = true;
 				break;
 			}
-			if (FMath::IsNegative(ScopeEntry->Time))
+			if (ScopeEntry->Time < 0.0)
 			{
 				if (ExitDepth == 0)
 				{
@@ -826,7 +826,7 @@ public:
 		//Iterate from the start of the page to the start time of our interval
 		while (ScopeEntry && FMath::Abs(ScopeEntry->Time) <= IntervalStart)
 		{
-			if (FMath::IsNegative(ScopeEntry->Time))
+			if (ScopeEntry->Time < 0.0)
 			{
 				check(CurrentStackDepth < SettingsType::MaxDepth);
 				FEventInfoStackEntry& StackEntry = EventStack[CurrentStackDepth++];
@@ -967,7 +967,7 @@ public:
 		while (ScopeEntry &&
 			   Depth < CurrentStackDepth)
 		{
-			if (FMath::IsNegative(ScopeEntry->Time))
+			if (ScopeEntry->Time < 0.0)
 			{
 				check(CurrentStackDepth < SettingsType::MaxDepth);
 				FEventInfoStackEntry& StackEntry = EventStack[CurrentStackDepth++];
@@ -1034,7 +1034,7 @@ public:
 		const FEventScopeEntry* ScopeEntry = ScopeEntryIterator.GetCurrentItem();
 		while (ScopeEntry && FMath::Abs(ScopeEntry->Time) < Time)
 		{
-			if (FMath::IsNegative(ScopeEntry->Time))
+			if (ScopeEntry->Time < 0.0)
 			{
 				CurrentStackDepth++;
 			}
@@ -1270,7 +1270,7 @@ private:
 		//......]...]....]...TargetExactTime
 		while (ScopeEntry && FMath::Abs(ScopeEntry->Time) <= InParams.TargetExactTime)
 		{
-			if (FMath::IsNegative(ScopeEntry->Time))
+			if (ScopeEntry->Time < 0.0)
 			{
 				check(IterationState.StackDepth < SettingsType::MaxDepth);
 				FEventInfoStackEntry& StackEntry = IterationState.EventStack[IterationState.StackDepth++];
@@ -1322,7 +1322,7 @@ private:
 		bool bHasStartEvent = false;
 		while (ScopeEntry && FMath::Abs(ScopeEntry->Time) <= InParams.TargetEndTime)
 		{
-			if (FMath::IsNegative(ScopeEntry->Time))
+			if (ScopeEntry->Time < 0.0)
 			{
 				check(IterationState.StackDepth < SettingsType::MaxDepth);
 				FEventInfoStackEntry& StackEntry = IterationState.EventStack[IterationState.StackDepth++];
diff --git a/Engine/Source/Editor/UnrealEd/Private/VertexSnapping.cpp b/Engine/Source/Editor/UnrealEd/Private/VertexSnapping.cpp
--- a/Engine/Source/Editor/UnrealEd/Private/VertexSnapping.cpp
+++ b/Engine/Source/Editor/UnrealEd/Private/VertexSnapping.cpp
@@ -432,7 +432,7 @@ bool FVertexSnappingImpl::GetClosestVertexOnComponent( const FSnapActor& SnapAct
 				// Vertex cannot be outside the view
 				&& !bOutside 
 				// In screen space the distance of the vertex must not be too far from the camera.  In any other axis the vertex cannot be beind the actor
-				&& ( ( CurrentAxis == EAxisList::Screen && DistanceFromCamera <= VertexSnappingConstants::MaxSquaredDistanceFromCamera ) || ( CurrentAxis != EAxisList::Screen && !FMath::IsNegativeFloat( Distance ) ) )
+				&& ( ( CurrentAxis == EAxisList::Screen && DistanceFromCamera <= VertexSnappingConstants::MaxSquaredDistanceFromCamera ) || ( CurrentAxis != EAxisList::Screen && !( Distance < 0.0f ) ) )
 				// The vertex must be closer than the current closest vertex
 				&& Distance < ClosestDistance )
 			{
@@ -506,7 +506,7 @@ FSnappingVertex FVertexSnappingImpl::GetClosestVertex( const TArray<FSnapActor>&
 				// we must have made some movement
 				&& !FMath::IsNearlyZero(Distance) 
 				// If not in screen space the vertex cannot be behind the point being snapped
-				&& ( CurrentAxis == EAxisList::Screen || !FMath::IsNegativeFloat( Distance ) )
+				&& ( CurrentAxis == EAxisList::Screen || !( Distance < 0.0f ) )
 				// The vertex must be closer than the current closest vertex
 				&& Distance < ClosestDistance )
 			{
diff --git a/Engine/Source/Runtime/Core/Private/Internationalization/Culture.cpp b/Engine/Source/Runtime/Core/Private/Internationalization/Culture.cpp
--- a/Engine/Source/Runtime/Core/Private/Internationalization/Culture.cpp
+++ b/Engine/Source/Runtime/Core/Private/Internationalization/Culture.cpp
@@ -300,7 +300,7 @@ ETextPluralForm FCulture::GetPluralForm(int32 Val, const ETextPluralType PluralT
 
 ETextPluralForm FCulture::GetPluralForm(double Val, const ETextPluralType PluralType) const
 {
-	if (FMath::IsNegativeDouble(Val))
+	if (Val < 0.0)
 	{
 		// Must be positive
 		Val *= -1.0;
diff --git a/Engine/Source/Runtime/Core/Private/Internationalization/FastDecimalFormat.cpp b/Engine/Source/Runtime/Core/Private/Internationalization/FastDecimalFormat.cpp
--- a/Engine/Source/Runtime/Core/Private/Internationalization/FastDecimalFormat.cpp
+++ b/Engine/Source/Runtime/Core/Private/Internationalization/FastDecimalFormat.cpp
@@ -514,7 +514,7 @@ void FractionalToString(const double InVal, const FDecimalNumberFormattingRules&
 		return;
 	}
 
-	const bool bIsNegative = FMath::IsNegativeDouble(InVal);
+	const bool bIsNegative = FMath::IsNegativeOrNegativeZero(InVal);
 
 	double IntegralPart = 0.0;
 	double FractionalPart = 0.0;
diff --git a/Engine/Source/Runtime/Core/Private/Internationalization/ICUCulture.cpp b/Engine/Source/Runtime/Core/Private/Internationalization/ICUCulture.cpp
--- a/Engine/Source/Runtime/Core/Private/Internationalization/ICUCulture.cpp
+++ b/Engine/Source/Runtime/Core/Private/Internationalization/ICUCulture.cpp
@@ -1216,7 +1216,7 @@ ETextPluralForm FICUCultureImplementation::GetPluralForm(int32 Val, const ETextP
 
 ETextPluralForm FICUCultureImplementation::GetPluralForm(double Val, const ETextPluralType PluralType) const
 {
-	checkf(!FMath::IsNegativeDouble(Val), TEXT("GetPluralFormImpl requires a positive value"));
+	checkf(!FMath::IsNegativeOrNegativeZero(Val), TEXT("GetPluralFormImpl requires a positive value"));
 
 	const icu::PluralRules* ICUPluralRules = (PluralType == ETextPluralType::Cardinal) ? ICUCardinalPluralRules : ICUOrdinalPluralRules;
 	const icu::UnicodeString ICUPluralFormTag = ICUPluralRules->select(Val);
diff --git a/Engine/Source/Runtime/Core/Private/Internationalization/LegacyCulture.cpp b/Engine/Source/Runtime/Core/Private/Internationalization/LegacyCulture.cpp
--- a/Engine/Source/Runtime/Core/Private/Internationalization/LegacyCulture.cpp
+++ b/Engine/Source/Runtime/Core/Private/Internationalization/LegacyCulture.cpp
@@ -221,7 +221,7 @@ ETextPluralForm FLegacyCultureImplementation::GetPluralForm(int32 Val, const ETe
 
 ETextPluralForm FLegacyCultureImplementation::GetPluralForm(double Val, const ETextPluralType PluralType) const
 {
-	checkf(!FMath::IsNegativeDouble(Val), TEXT("GetPluralFormImpl requires a positive value"));
+	checkf(!FMath::IsNegativeOrNegativeZero(Val), TEXT("GetPluralFormImpl requires a positive value"));
 	return GetDefaultPluralForm((int64)Val, PluralType);
 }
 
diff --git a/Engine/Source/Runtime/Core/Private/Math/UnrealMath.cpp b/Engine/Source/Runtime/Core/Private/Math/UnrealMath.cpp
--- a/Engine/Source/Runtime/Core/Private/Math/UnrealMath.cpp
+++ b/Engine/Source/Runtime/Core/Private/Math/UnrealMath.cpp
@@ -1762,7 +1762,7 @@ static bool ComputeProjectedSphereShaft(
 			FVector::FReal Pxa = -Pza * Nza / Nxa;
 			FVector4 P = ProjMatrix.TransformFVector4(FVector4(FVec4Real(Axis.X * Pxa),FVec4Real(Axis.Y * Pxa),FVec4Real(Pza),1));
 			FVector::FReal X = (Dot3(P,Axis) / P.W + 1.0f * AxisSign) / 2.0f * AxisSign;
-			if(FMath::IsNegative(Nxa) ^ FMath::IsNegative(AxisSign))
+			if((Nxa < FVector::FReal{0}) ^ (AxisSign < FVector::FReal{0}))
 			{
 				InOutMaxX = FMath::Min(FMath::CeilToInt32(ViewSizeX * X + ViewX),InOutMaxX);
 			}
@@ -1778,7 +1778,7 @@ static bool ComputeProjectedSphereShaft(
 			FVector::FReal Pxb = -Pzb * Nzb / Nxb;
 			FVector4 P = ProjMatrix.TransformFVector4(FVector4(FVec4Real(Axis.X * Pxb), FVec4Real(Axis.Y * Pxb),FVec4Real(Pzb),1));
 			FVector::FReal X = (Dot3(P,Axis) / P.W + 1.0f * AxisSign) / 2.0f * AxisSign;
-			if(FMath::IsNegative(Nxb) ^ FMath::IsNegative(AxisSign))
+			if((Nxb < FVector::FReal{0}) ^ (AxisSign < FVector::FReal{0}))
 			{
 				InOutMaxX = FMath::Min(FMath::CeilToInt32(ViewSizeX * X + ViewX),InOutMaxX);
 			}
diff --git a/Engine/Source/Runtime/Core/Public/GenericPlatform/GenericPlatformMath.h b/Engine/Source/Runtime/Core/Public/GenericPlatform/GenericPlatformMath.h
--- a/Engine/Source/Runtime/Core/Public/GenericPlatform/GenericPlatformMath.h
+++ b/Engine/Source/Runtime/Core/Public/GenericPlatform/GenericPlatformMath.h
@@ -15,6 +15,7 @@
 #include "Templates/IsFloatingPoint.h"
 #include "Templates/UnrealTypeTraits.h"
 #include "Templates/ResolveTypeAmbiguity.h"
+#include "Templates/TypeCompatibleBytes.h"
 
 #if PLATFORM_HAS_FENV_H 
 #include <fenv.h>
@@ -25,7 +26,6 @@
  */
 struct FGenericPlatformMath
 {
-
 	// load half (F16) to float
 	//https://gist.github.com/rygorous/2156668
 	static FORCEINLINE float LoadHalf(const uint16* Ptr)
@@ -578,37 +578,55 @@ struct FGenericPlatformMath
 	/** Return true if value is NaN (not a number). */
 	static FORCEINLINE bool IsNaN( float A ) 
 	{
-		return ((*(uint32*)&A) & 0x7FFFFFFFU) > 0x7F800000U;
+		return (BitCast<uint32>(A) & 0x7FFFFFFFU) > 0x7F800000U;
 	}
 	static FORCEINLINE bool IsNaN(double A)
 	{
-		return ((*(uint64*)&A) & 0x7FFFFFFFFFFFFFFFULL) > 0x7FF0000000000000ULL;
+		return (BitCast<uint64>(A) & 0x7FFFFFFFFFFFFFFFULL) > 0x7FF0000000000000ULL;
 	}
 
 	/** Return true if value is finite (not NaN and not Infinity). */
 	static FORCEINLINE bool IsFinite( float A )
 	{
-		return ((*(uint32*)&A) & 0x7F800000U) != 0x7F800000U;
+		return (BitCast<uint32>(A) & 0x7F800000U) != 0x7F800000U;
 	}
 	static FORCEINLINE bool IsFinite(double A)
 	{
-		return ((*(uint64*)&A) & 0x7FF0000000000000ULL) != 0x7FF0000000000000ULL;
+		return (BitCast<uint64>(A) & 0x7FF0000000000000ULL) != 0x7FF0000000000000ULL;
+	}
+
+	static FORCEINLINE bool IsNegativeOrNegativeZero(float A)
+	{
+		return BitCast<uint32>(A) >= (uint32)0x80000000; // Detects sign bit.
 	}
 
-	//UE_DEPRECATED(5.0, "Use IsNegative here.")	// LWC_TODO: Enable IsNegativeFloat/Double deprecations and fix engine calls.
-	static FORCEINLINE bool IsNegativeFloat(float A) { return IsNegative(A); }
+	static FORCEINLINE bool IsNegativeOrNegativeZero(double A)
+	{
+		return BitCast<uint64>(A) >= (uint64)0x8000000000000000; // Detects sign bit.
+	}
 
-	//UE_DEPRECATED(5.0, "Use IsNegative here.")
-	static FORCEINLINE bool IsNegativeDouble(double A) { return IsNegative(A); }
+	UE_DEPRECATED(5.1, "IsNegativeFloat has been deprecated in favor of IsNegativeOrNegativeZero or simply A < 0.")
+	static FORCEINLINE bool IsNegativeFloat(float A)
+	{
+		return IsNegativeOrNegativeZero(A);
+	}
+
+	UE_DEPRECATED(5.1, "IsNegativeDouble has been deprecated in favor of IsNegativeOrNegativeZero or simply A < 0.")
+	static FORCEINLINE bool IsNegativeDouble(double A)
+	{
+		return IsNegativeOrNegativeZero(A);
+	}
 
+	UE_DEPRECATED(5.1, "IsNegative has been deprecated in favor of IsNegativeOrNegativeZero or simply A < 0.")
 	static FORCEINLINE bool IsNegative(float A)
 	{
-		return ( A < 0.f );
+		return IsNegativeOrNegativeZero(A);
 	}
 
+	UE_DEPRECATED(5.1, "IsNegative has been deprecated in favor of IsNegativeOrNegativeZero or simply A < 0.")
 	static FORCEINLINE bool IsNegative(double A)
 	{
-		return ( A < 0.0 );
+		return IsNegativeOrNegativeZero(A);
 	}
 
 	/** Returns a random integer between 0 and RAND_MAX, inclusive */
diff --git a/Engine/Source/Runtime/Core/Public/Templates/TypeCompatibleBytes.h b/Engine/Source/Runtime/Core/Public/Templates/TypeCompatibleBytes.h
--- a/Engine/Source/Runtime/Core/Public/Templates/TypeCompatibleBytes.h
+++ b/Engine/Source/Runtime/Core/Public/Templates/TypeCompatibleBytes.h
@@ -4,6 +4,8 @@
 
 
 #include "CoreTypes.h"
+#include <string.h>
+#include <type_traits>
 
 /**
  * Used to declare an untyped array of data with compile-time alignment.
@@ -66,3 +68,20 @@ struct TTypeCompatibleBytes :
 	ElementType*		GetTypedPtr()		{ return (ElementType*)this;  }
 	const ElementType*	GetTypedPtr() const	{ return (const ElementType*)this; }
 };
+
+template <
+	typename ToType,
+	typename FromType,
+	std::enable_if_t<sizeof(ToType) == sizeof(FromType) && std::is_trivially_copyable_v<ToType> && std::is_trivially_copyable_v<FromType>>* = nullptr
+>
+inline ToType BitCast(const FromType& From)
+{
+// Ensure we can use this builtin - seems to be present on Clang 9, GCC 11 and MSVC 19.26
+#if (defined(__clang__) && __clang_major__ >= 9) || (defined(__GNUC__) && __GNUC__ >= 11) || (defined(_MSC_VER) && _MSC_VER >= 1926)
+	return __builtin_bit_cast(ToType, From);
+#else
+	TTypeCompatibleBytes<ToType> Result;
+	memcpy(&Result, &From, sizeof(ToType));
+	return *Result.GetTypedPtr();
+#endif
+}
diff --git a/Engine/Source/Runtime/Engine/Private/Animation/BlendSpaceHelpers.cpp b/Engine/Source/Runtime/Engine/Private/Animation/BlendSpaceHelpers.cpp
--- a/Engine/Source/Runtime/Engine/Private/Animation/BlendSpaceHelpers.cpp
+++ b/Engine/Source/Runtime/Engine/Private/Animation/BlendSpaceHelpers.cpp
@@ -254,7 +254,7 @@ FDelaunayTriangleGenerator::ECircumCircleState FDelaunayTriangleGenerator::GetCi
 	const double Det = M00*M11*M22+M01*M12*M20+M02*M10*M21 - (M02*M11*M20+M01*M10*M22+M00*M12*M21);
 	
 	// When the vertices are sorted in a counterclockwise order, the determinant is positive if and only if Testpoint lies inside the circumcircle of T.
-	if (FMath::IsNegative(Det))
+	if (Det < 0.0)
 	{
 		return ECCS_Outside;
 	}
diff --git a/Engine/Source/Runtime/Engine/Private/GeomTools.cpp b/Engine/Source/Runtime/Engine/Private/GeomTools.cpp
--- a/Engine/Source/Runtime/Engine/Private/GeomTools.cpp
+++ b/Engine/Source/Runtime/Engine/Private/GeomTools.cpp
@@ -126,7 +126,7 @@ void FGeomTools::ClipMeshWithPlane( TArray<FClipSMTriangle>& OutTris, TArray<FUt
 			int32 ThisVert = EdgeIdx;
 
 			// If start vert is inside, add it.
-			if(FMath::IsNegativeFloat(PlaneDist[ThisVert]))
+			if(PlaneDist[ThisVert] < 0.0f)
 			{
 				FinalVerts.Add( SrcTri->Vertices[ThisVert] );
 			}
@@ -134,7 +134,7 @@ void FGeomTools::ClipMeshWithPlane( TArray<FClipSMTriangle>& OutTris, TArray<FUt
 			// If start and next vert are on opposite sides, add intersection
 			int32 NextVert = (EdgeIdx+1)%3;
 
-			if(FMath::IsNegativeFloat(PlaneDist[EdgeIdx]) != FMath::IsNegativeFloat(PlaneDist[NextVert]))
+			if((PlaneDist[EdgeIdx] < 0.0f) != (PlaneDist[NextVert] < 0.0f))
 			{
 				// Find distance along edge that plane is
 				float Alpha = -PlaneDist[ThisVert] / (PlaneDist[NextVert] - PlaneDist[ThisVert]);
@@ -332,7 +332,7 @@ bool FGeomTools::VectorsOnSameSide(const FVector3f& Vec, const FVector3f& A, con
 	const FVector3f CrossA = Vec ^ A;
 	const FVector3f CrossB = Vec ^ B;
 	float DotWithEpsilon = SameSideDotProductEpsilon + ( CrossA | CrossB );
-	return !FMath::IsNegativeFloat(DotWithEpsilon);
+	return !(DotWithEpsilon < 0.0f);
 }
 
 /** Util to see if P lies within triangle created by A, B and C. */
@@ -467,7 +467,7 @@ bool FGeomTools::TriangulatePoly(TArray<FClipSMTriangle>& OutTris, const FClipSM
 				const FVector3f ABEdge = PolyVerts[BIndex].Pos - PolyVerts[AIndex].Pos;
 				const FVector3f ACEdge = PolyVerts[CIndex].Pos - PolyVerts[AIndex].Pos;
 				const float TriangleDeterminant = (ABEdge ^ ACEdge) | (FVector3f)InPoly.FaceNormal;
-				if(FMath::IsNegativeFloat(TriangleDeterminant))
+				if(TriangleDeterminant < 0.0f)
 				{
 					continue;
 				}
@@ -1468,7 +1468,7 @@ bool FGeomTools2D::TriangulatePoly(TArray<FVector2D>& OutTris, const TArray<FVec
 				// Check that this vertex is convex (cross product must be positive)
 				const FVector2D ABEdge = PolyVerts[BIndex] - PolyVerts[AIndex];
 				const FVector2D ACEdge = PolyVerts[CIndex] - PolyVerts[AIndex];
-				if (FMath::IsNegativeFloat(ABEdge ^ ACEdge))
+				if ((ABEdge ^ ACEdge) < 0.0f)
 				{
 					continue;
 				}
diff --git a/Engine/Source/Runtime/Engine/Private/HModel.cpp b/Engine/Source/Runtime/Engine/Private/HModel.cpp
--- a/Engine/Source/Runtime/Engine/Private/HModel.cpp
+++ b/Engine/Source/Runtime/Engine/Private/HModel.cpp
@@ -128,7 +128,7 @@ bool HModel::ResolveSurface(const FSceneView* View,int32 X,int32 Y,uint32& OutSu
 						FVector2D(Vertices[0].X,Vertices[0].Y),
 						FVector2D(Vertices[1].X,Vertices[1].Y),
 						FVector2D(Vertices[2].X,Vertices[2].Y),
-						(Surf.PolyFlags & PF_TwoSided) ? false : FMath::IsNegativeFloat(EdgeA.X * EdgeB.Y - EdgeA.Y * EdgeB.X) // Check if a surface is twosided when it is selected in the editor viewport
+						(Surf.PolyFlags & PF_TwoSided) ? false : (EdgeA.X * EdgeB.Y - EdgeA.Y * EdgeB.X < 0.0f) // Check if a surface is twosided when it is selected in the editor viewport
 						);
 				}
 			}
diff --git a/Engine/Source/Runtime/Engine/Private/SceneView.cpp b/Engine/Source/Runtime/Engine/Private/SceneView.cpp
--- a/Engine/Source/Runtime/Engine/Private/SceneView.cpp
+++ b/Engine/Source/Runtime/Engine/Private/SceneView.cpp
@@ -762,7 +762,7 @@ FSceneView::FSceneView(const FSceneViewInitOptions& InitOptions)
 	SetupViewFrustum(*this);
 
 	// Determine whether the view should reverse the cull mode due to a negative determinant.  Only do this for a valid scene
-	bReverseCulling = (Family && Family->Scene) ? FMath::IsNegativeFloat(ViewMatrices.GetViewMatrix().Determinant()) : false;
+	bReverseCulling = (Family && Family->Scene) ? ViewMatrices.GetViewMatrix().Determinant() < 0.0f : false;
 
 	// OpenGL Gamma space output in GLSL flips Y when rendering directly to the back buffer (so not needed on PC, as we never render directly into the back buffer)
 	auto ShaderPlatform = GShaderPlatformForFeatureLevel[FeatureLevel];
diff --git a/Engine/Source/Runtime/GeometryCore/Private/CompGeom/PolygonTriangulation.cpp b/Engine/Source/Runtime/GeometryCore/Private/CompGeom/PolygonTriangulation.cpp
--- a/Engine/Source/Runtime/GeometryCore/Private/CompGeom/PolygonTriangulation.cpp
+++ b/Engine/Source/Runtime/GeometryCore/Private/CompGeom/PolygonTriangulation.cpp
@@ -201,7 +201,7 @@ static bool VectorsOnSameSide(const FVector& Vec, const FVector& A, const FVecto
 	const FVector CrossA = FVector::CrossProduct(Vec, A);
 	const FVector CrossB = FVector::CrossProduct(Vec, B);
 	float DotWithEpsilon = SameSideDotProductEpsilon + FVector::DotProduct(CrossA, CrossB);
-	return !FMath::IsNegativeFloat(DotWithEpsilon);
+	return !(DotWithEpsilon < 0.0f);
 }
 
 
diff --git a/Engine/Source/Runtime/MeshDescription/Private/MeshDescription.cpp b/Engine/Source/Runtime/MeshDescription/Private/MeshDescription.cpp
--- a/Engine/Source/Runtime/MeshDescription/Private/MeshDescription.cpp
+++ b/Engine/Source/Runtime/MeshDescription/Private/MeshDescription.cpp
@@ -1256,7 +1256,7 @@ static bool VectorsOnSameSide(const FVector& Vec, const FVector& A, const FVecto
 	const FVector CrossA = FVector::CrossProduct(Vec, A);
 	const FVector CrossB = FVector::CrossProduct(Vec, B);
 	float DotWithEpsilon = SameSideDotProductEpsilon + FVector::DotProduct(CrossA, CrossB);
-	return !FMath::IsNegativeFloat(DotWithEpsilon);
+	return !(DotWithEpsilon < 0.0f);
 }
 
 
