[Insights] Added initial implementation for MetadataProvider.

diff --git a/Engine/Source/Developer/TraceServices/Private/Analyzers/AllocationsAnalysis.cpp b/Engine/Source/Developer/TraceServices/Private/Analyzers/AllocationsAnalysis.cpp
--- a/Engine/Source/Developer/TraceServices/Private/Analyzers/AllocationsAnalysis.cpp
+++ b/Engine/Source/Developer/TraceServices/Private/Analyzers/AllocationsAnalysis.cpp
@@ -2,23 +2,23 @@
 
 #include "AllocationsAnalysis.h"
 #include "Model/AllocationsProvider.h"
+#include "Model/MetadataProvider.h"
 #include "ProfilingDebugging/MemoryTrace.h"
 #include "TraceServices/Model/AnalysisSession.h"
+#include "TraceServices/Model/Callstack.h"
 
 namespace TraceServices
 {
 
 #define INSIGHTS_MEM_TRACE_LEGACY_FORMAT 1 // backward compatibility with legacy memory trace format (5.0-EA)
+#define INSIGHTS_MEM_TRACE_METADATA_TEST 0
 
 ////////////////////////////////////////////////////////////////////////////////////////////////////
 
-FAllocationsAnalyzer::FAllocationsAnalyzer(IAnalysisSession& InSession, FAllocationsProvider& InAllocationsProvider)
+FAllocationsAnalyzer::FAllocationsAnalyzer(IAnalysisSession& InSession, FAllocationsProvider& InAllocationsProvider, FMetadataProvider& InMetadataProvider)
 	: Session(InSession)
 	, AllocationsProvider(InAllocationsProvider)
-	, BaseCycle(0)
-	, MarkerPeriod(0)
-	, LastMarkerCycle(0)
-	, LastMarkerSeconds(0.0)
+	, MetadataProvider(InMetadataProvider)
 {
 }
 
@@ -52,6 +52,13 @@ void FAllocationsAnalyzer::OnAnalysisBegin(const FOnAnalysisContext& Context)
 	Builder.RouteEvent(RouteId_HeapUnmarkAlloc,     "Memory", "HeapUnmarkAlloc");
 
 	Builder.RouteLoggerEvents(RouteId_MemScope, "Memory", true);
+
+#if INSIGHTS_MEM_TRACE_METADATA_TEST
+	{
+		FMetadataProvider::FEditScopeLock _(MetadataProvider);
+		TagIdMetadataType = MetadataProvider.RegisterMetadataType(TEXT("MemTagId"), sizeof(TagIdType));
+	}
+#endif
 }
 
 ////////////////////////////////////////////////////////////////////////////////////////////////////
@@ -319,27 +326,53 @@ bool FAllocationsAnalyzer::OnEvent(uint16 RouteId, EStyle Style, const FOnEventC
 				if (FCStringAnsi::Strlen(Context.EventData.GetTypeInfo().GetName()) == 11) // "MemoryScope"
 				{
 					const TagIdType Tag = Context.EventData.GetValue<TagIdType>("Tag");
-					FAllocationsProvider::FEditScopeLock _(AllocationsProvider);
-					AllocationsProvider.EditPushTag(ThreadId, Tracker, Tag);
+					{
+						FAllocationsProvider::FEditScopeLock _(AllocationsProvider);
+						AllocationsProvider.EditPushTag(ThreadId, Tracker, Tag);
+					}
+#if INSIGHTS_MEM_TRACE_METADATA_TEST
+					{
+						FMetadataProvider::FEditScopeLock _(MetadataProvider);
+						MetadataProvider.PushScopedMetadata(ThreadId, TagIdMetadataType, (void*)&Tag, sizeof(TagIdType));
+					}
+#endif
 				}
 				else // "MemoryScopePtr"
 				{
 					const uint64 Ptr = Context.EventData.GetValue<uint64>("Ptr");
-					FAllocationsProvider::FEditScopeLock _(AllocationsProvider);
-					AllocationsProvider.EditPushTagFromPtr(ThreadId, Tracker, Ptr);
+					{
+						FAllocationsProvider::FEditScopeLock _(AllocationsProvider);
+						AllocationsProvider.EditPushTagFromPtr(ThreadId, Tracker, Ptr);
+					}
+#if INSIGHTS_MEM_TRACE_METADATA_TEST
+					{
+						FMetadataProvider::FEditScopeLock _(MetadataProvider);
+						TagIdType Tag = 0; //TODO: AllocationsProvider.GetTagFromPtr(ThreadId, Tracker, Ptr);
+						MetadataProvider.PushScopedMetadata(ThreadId, TagIdMetadataType, (void*)&Tag, sizeof(TagIdType));
+					}
+#endif
 				}
 			}
 			else // EStyle::LeaveScope
 			{
-				FAllocationsProvider::FEditScopeLock _(AllocationsProvider);
-				if (AllocationsProvider.HasTagFromPtrScope(ThreadId, Tracker)) // Is TagFromPtr scope active?
 				{
-					AllocationsProvider.EditPopTagFromPtr(ThreadId, Tracker);
+					FAllocationsProvider::FEditScopeLock _(AllocationsProvider);
+					if (AllocationsProvider.HasTagFromPtrScope(ThreadId, Tracker)) // Is TagFromPtr scope active?
+					{
+						AllocationsProvider.EditPopTagFromPtr(ThreadId, Tracker);
+					}
+					else
+					{
+						check(FCStringAnsi::Strlen(Context.EventData.GetTypeInfo().GetName()) == 11); // "MemoryScope"
+						AllocationsProvider.EditPopTag(ThreadId, Tracker);
+					}
 				}
-				else
+#if INSIGHTS_MEM_TRACE_METADATA_TEST
 				{
-					AllocationsProvider.EditPopTag(ThreadId, Tracker);
+					FMetadataProvider::FEditScopeLock _(MetadataProvider);
+					MetadataProvider.PopScopedMetadata(ThreadId, TagIdMetadataType);
 				}
+#endif
 			}
 			break;
 		}
@@ -357,6 +390,7 @@ double FAllocationsAnalyzer::GetCurrentTime() const
 
 ////////////////////////////////////////////////////////////////////////////////////////////////////
 
+#undef INSIGHTS_MEM_TRACE_METADATA_TEST
 #undef INSIGHTS_MEM_TRACE_LEGACY_FORMAT
 
 } // namespace TraceServices
diff --git a/Engine/Source/Developer/TraceServices/Private/Analyzers/AllocationsAnalysis.h b/Engine/Source/Developer/TraceServices/Private/Analyzers/AllocationsAnalysis.h
--- a/Engine/Source/Developer/TraceServices/Private/Analyzers/AllocationsAnalysis.h
+++ b/Engine/Source/Developer/TraceServices/Private/Analyzers/AllocationsAnalysis.h
@@ -10,6 +10,7 @@ namespace TraceServices
 class IAnalysisSession;
 
 class FAllocationsProvider;
+class FMetadataProvider;
 
 class FAllocationsAnalyzer : public UE::Trace::IAnalyzer
 {
@@ -36,7 +37,7 @@ private:
 	};
 
 public:
-	FAllocationsAnalyzer(IAnalysisSession& Session, FAllocationsProvider& AllocationsProvider);
+	FAllocationsAnalyzer(IAnalysisSession& Session, FAllocationsProvider& AllocationsProvider, FMetadataProvider& MetadataProvider);
 	~FAllocationsAnalyzer();
 	virtual void OnAnalysisBegin(const FOnAnalysisContext& Context) override;
 	virtual void OnAnalysisEnd() override;
@@ -46,11 +47,13 @@ public:
 private:
 	IAnalysisSession& Session;
 	FAllocationsProvider& AllocationsProvider;
-	uint64 BaseCycle;
-	uint32 MarkerPeriod;
-	uint64 LastMarkerCycle;
-	double LastMarkerSeconds;
-	uint8 SizeShift;
+	FMetadataProvider& MetadataProvider;
+	uint64 BaseCycle = 0;
+	uint64 LastMarkerCycle = 0;
+	double LastMarkerSeconds = 0.0;
+	uint32 MarkerPeriod = 0;
+	uint16 TagIdMetadataType = ~0;
+	uint8 SizeShift = 0;
 };
 
 } // namespace TraceServices
diff --git a/Engine/Source/Developer/TraceServices/Private/Model/AllocationsProvider.cpp b/Engine/Source/Developer/TraceServices/Private/Model/AllocationsProvider.cpp
--- a/Engine/Source/Developer/TraceServices/Private/Model/AllocationsProvider.cpp
+++ b/Engine/Source/Developer/TraceServices/Private/Model/AllocationsProvider.cpp
@@ -2,17 +2,18 @@
 
 #include "AllocationsProvider.h"
 
-#include <limits>
-
 #include "AllocationsQuery.h"
-#include "SbTree.h"
 #include "Common/Utils.h"
 #include "Containers/ArrayView.h"
 #include "Misc/PathViews.h"
+#include "Model/MetadataProvider.h"
 #include "ProfilingDebugging/MemoryTrace.h"
+#include "SbTree.h"
 #include "TraceServices/Containers/Allocators.h"
 #include "TraceServices/Model/Callstack.h"
 
+#include <limits>
+
 ////////////////////////////////////////////////////////////////////////////////////////////////////
 
 #define INSIGHTS_SLOW_CHECK(expr) //check(expr)
@@ -41,6 +42,8 @@
 
 #define INSIGHTS_VALIDATE_ALLOC_EVENTS 0
 
+#define INSIGHTS_DEBUG_METADATA 0
+
 ////////////////////////////////////////////////////////////////////////////////////////////////////
 
 namespace TraceServices
@@ -184,7 +187,7 @@ void FTagTracker::AddTagSpec(TagIdType InTag, TagIdType InParentTag, const TCHAR
 				TagMap[ReferencingId].ParentTag = InTag;
 			}
 		}
-		
+
 		if (!InDisplay || *InDisplay == TEXT('\0'))
 		{
 			UE_LOG(LogTraceServices, Warning, TEXT("[MemAlloc] Tag with id %u has invalid display name (ParentTag=%u)!"), InTag, InParentTag);
@@ -271,7 +274,7 @@ const TCHAR* FTagTracker::GetTagString(TagIdType InTag) const
 
 void FTagTracker::EnumerateTags(TFunctionRef<void (const TCHAR*, const TCHAR*, TagIdType, TagIdType)> Callback) const
 {
-	for(const auto& EntryPair : TagMap)
+	for (const auto& EntryPair : TagMap)
 	{
 		const TagIdType Id = EntryPair.Get<0>();
 		const FTagEntry& Entry = EntryPair.Get<1>();
@@ -1113,8 +1116,9 @@ void FLiveAllocCollection::Enumerate(TFunctionRef<void(const FAllocationItem& Al
 // FAllocationsProvider
 ////////////////////////////////////////////////////////////////////////////////////////////////////
 
-FAllocationsProvider::FAllocationsProvider(IAnalysisSession& InSession)
+FAllocationsProvider::FAllocationsProvider(IAnalysisSession& InSession, FMetadataProvider& InMetadataProvider)
 	: Session(InSession)
+	, MetadataProvider(InMetadataProvider)
 	, TagTracker(InSession)
 	, Timeline(Session.GetLinearAllocator(), 1024)
 	, MinTotalAllocatedMemoryTimeline(Session.GetLinearAllocator(), 1024)
@@ -1273,6 +1277,29 @@ void FAllocationsProvider::EditAlloc(double Time, uint32 CallstackId, uint64 Add
 		AllocationPtr = LiveAllocs[RootHeap]->AddNew(Address);
 		FAllocationItem& Allocation = *AllocationPtr;
 
+		uint32 MetadataId = MetadataProvider.InvalidMetadataId;
+		{
+			FMetadataProvider::FEditScopeLock _(MetadataProvider);
+			MetadataId = MetadataProvider.PinAndGetId(CurrentSystemThreadId);
+#if INSIGHTS_DEBUG_METADATA
+			if (MetadataId != FMetadataProvider::InvalidMetadataId &&
+				!TagTracker.HasTagFromPtrScope(CurrentSystemThreadId, CurrentTracker))
+			{
+				const uint32 MetadataStackSize = MetadataProvider.GetMetadataStackSize(CurrentSystemThreadId, MetadataId);
+				check(MetadataStackSize > 0);
+				uint16 MetaType;
+				const void* MetaData;
+				uint32 MetaDataSize;
+				MetadataProvider.GetMetadata(CurrentSystemThreadId, MetadataId, MetadataStackSize - 1, MetaType, MetaData, MetaDataSize);
+				if (MetaType == 0) // "MemTagId"
+				{
+					TagIdType MetaMemTag = *(TagIdType*)MetaData;
+					ensure(MetaMemTag == Tag);
+				}
+			}
+#endif
+		}
+
 		INSIGHTS_SLOW_CHECK(Allocation.Address == Address);
 		Allocation.SizeAndAlignment = FAllocationItem::PackSizeAndAlignment(InSize, InAlignment);
 		Allocation.StartEventIndex = EventIndex[RootHeap];
@@ -1281,7 +1308,7 @@ void FAllocationsProvider::EditAlloc(double Time, uint32 CallstackId, uint64 Add
 		Allocation.EndTime = std::numeric_limits<double>::infinity();
 		Allocation.CallstackId = CallstackId;
 		Allocation.FreeCallstackId = 0; // no callstack yet
-		Allocation.MetadataId = 0; //TODO: MetadataProvider.PinAndGetId() // pins the metadata stack and returns an id for it
+		Allocation.MetadataId = MetadataId;
 		Allocation.Reserved = 0;
 		Allocation.Tag = Tag;
 		Allocation.RootHeap = RootHeap;
@@ -2034,7 +2061,7 @@ void FAllocationsProvider::EnumerateTags(TFunctionRef<void(const TCHAR*, const T
 
 void FAllocationsProvider::DebugPrint() const
 {
-	for(const FSbTree* Tree : SbTree)
+	for (const FSbTree* Tree : SbTree)
 	{
 		if (Tree)
 		{
@@ -2069,6 +2096,7 @@ const IAllocationsProvider::FQueryStatus FAllocationsProvider::PollQuery(FQueryH
 }
 
 ////////////////////////////////////////////////////////////////////////////////////////////////////
+
 void FAllocationsProvider::EnumerateLiveAllocs(TFunctionRef<void(const FAllocationItem& Alloc)> Callback) const
 {
 	ReadAccessCheck();
@@ -2082,13 +2110,7 @@ void FAllocationsProvider::EnumerateLiveAllocs(TFunctionRef<void(const FAllocati
 }
 
 ////////////////////////////////////////////////////////////////////////////////////////////////////
-FName GetAllocationsProviderName()
-{
-	static FName Name(TEXT("AllocationsProvider"));
-	return Name;
-}
 
-////////////////////////////////////////////////////////////////////////////////////////////////////
 uint32 FAllocationsProvider::GetNumLiveAllocs() const
 {
 	ReadAccessCheck();
@@ -2097,6 +2119,14 @@ uint32 FAllocationsProvider::GetNumLiveAllocs() const
 
 ////////////////////////////////////////////////////////////////////////////////////////////////////
 
+FName GetAllocationsProviderName()
+{
+	static FName Name(TEXT("AllocationsProvider"));
+	return Name;
+}
+
+////////////////////////////////////////////////////////////////////////////////////////////////////
+
 const IAllocationsProvider* ReadAllocationsProvider(const IAnalysisSession& Session)
 {
 	return Session.ReadProvider<IAllocationsProvider>(GetAllocationsProviderName());
@@ -2114,3 +2144,4 @@ const IAllocationsProvider* ReadAllocationsProvider(const IAnalysisSession& Sess
 #undef INSIGHTS_SLA_USE_ADDRESS_MAP
 #undef INSIGHTS_USE_LAST_ALLOC
 #undef INSIGHTS_VALIDATE_ALLOC_EVENTS
+#undef INSIGHTS_DEBUG_METADATA
diff --git a/Engine/Source/Developer/TraceServices/Private/Model/AllocationsProvider.h b/Engine/Source/Developer/TraceServices/Private/Model/AllocationsProvider.h
--- a/Engine/Source/Developer/TraceServices/Private/Model/AllocationsProvider.h
+++ b/Engine/Source/Developer/TraceServices/Private/Model/AllocationsProvider.h
@@ -5,8 +5,7 @@
 #include "Common/PagedArray.h"
 #include "Containers/Map.h"
 #include "HAL/CriticalSection.h"
-#include "ProfilingDebugging/MemoryTrace.h"
-#include "TraceServices/Model/Callstack.h"
+#include "ProfilingDebugging/MemoryTrace.h" // for EMemoryTraceHeapFlags and EMemoryTraceHeapAllocationFlags
 #include "TraceServices/Model/AllocationsProvider.h"
 
 namespace TraceServices
@@ -14,6 +13,7 @@ namespace TraceServices
 
 class IAnalysisSession;
 class ILinearAllocator;
+class FMetadataProvider;
 class FSbTree;
 
 ////////////////////////////////////////////////////////////////////////////////////////////////////
@@ -311,7 +311,7 @@ private:
 	static constexpr double DefaultTimelineSampleGranularity = 0.0001; // 0.1ms
 
 public:
-	explicit FAllocationsProvider(IAnalysisSession& InSession);
+	explicit FAllocationsProvider(IAnalysisSession& InSession, FMetadataProvider& InMetadataProvider);
 	virtual ~FAllocationsProvider();
 
 	virtual void BeginEdit() const override { Lock.BeginWrite(); }
@@ -389,6 +389,7 @@ private:
 
 private:
 	IAnalysisSession& Session;
+	FMetadataProvider& MetadataProvider;
 
 	mutable FAllocationsProviderLock Lock;
 
diff --git a/Engine/Source/Developer/TraceServices/Private/Model/MetadataProvider.cpp b/Engine/Source/Developer/TraceServices/Private/Model/MetadataProvider.cpp
new file mode 100644
--- /dev/null
+++ b/Engine/Source/Developer/TraceServices/Private/Model/MetadataProvider.cpp
@@ -0,0 +1,472 @@
+// Copyright Epic Games, Inc. All Rights Reserved.
+
+#include "MetadataProvider.h"
+
+#include "Common/Utils.h"
+
+////////////////////////////////////////////////////////////////////////////////////////////////////
+
+namespace TraceServices
+{
+
+////////////////////////////////////////////////////////////////////////////////////////////////////
+// FMetadataProviderLock
+////////////////////////////////////////////////////////////////////////////////////////////////////
+
+thread_local FMetadataProviderLock* GThreadCurrentMetadataProviderLock;
+thread_local int32 GThreadCurrentReadMetadataProviderLockCount;
+thread_local int32 GThreadCurrentWriteMetadataProviderLockCount;
+
+void FMetadataProviderLock::ReadAccessCheck() const
+{
+	checkf(GThreadCurrentMetadataProviderLock == this && (GThreadCurrentReadMetadataProviderLockCount > 0 || GThreadCurrentWriteMetadataProviderLockCount > 0),
+		TEXT("Trying to READ from metadata provider outside of a READ scope"));
+}
+
+void FMetadataProviderLock::WriteAccessCheck() const
+{
+	checkf(GThreadCurrentMetadataProviderLock == this && GThreadCurrentWriteMetadataProviderLockCount > 0,
+		TEXT("Trying to WRITE to metadata provider outside of an EDIT/WRITE scope"));
+}
+
+void FMetadataProviderLock::BeginRead()
+{
+	check(!GThreadCurrentMetadataProviderLock || GThreadCurrentMetadataProviderLock == this);
+	checkf(GThreadCurrentWriteMetadataProviderLockCount == 0, TEXT("Trying to lock metadata provider for READ while holding EDIT/WRITE access"));
+	if (GThreadCurrentReadMetadataProviderLockCount++ == 0)
+	{
+		GThreadCurrentMetadataProviderLock = this;
+		RWLock.ReadLock();
+	}
+}
+
+void FMetadataProviderLock::EndRead()
+{
+	check(GThreadCurrentReadMetadataProviderLockCount > 0);
+	if (--GThreadCurrentReadMetadataProviderLockCount == 0)
+	{
+		RWLock.ReadUnlock();
+		GThreadCurrentMetadataProviderLock = nullptr;
+	}
+}
+
+void FMetadataProviderLock::BeginWrite()
+{
+	check(!GThreadCurrentMetadataProviderLock || GThreadCurrentMetadataProviderLock == this);
+	checkf(GThreadCurrentReadMetadataProviderLockCount == 0, TEXT("Trying to lock metadata provider for EDIT/WRITE while holding READ access"));
+	if (GThreadCurrentWriteMetadataProviderLockCount++ == 0)
+	{
+		GThreadCurrentMetadataProviderLock = this;
+		RWLock.WriteLock();
+	}
+}
+
+void FMetadataProviderLock::EndWrite()
+{
+	check(GThreadCurrentWriteMetadataProviderLockCount > 0);
+	if (--GThreadCurrentWriteMetadataProviderLockCount == 0)
+	{
+		RWLock.WriteUnlock();
+		GThreadCurrentMetadataProviderLock = nullptr;
+	}
+}
+
+////////////////////////////////////////////////////////////////////////////////////////////////////
+// FMetadataProvider
+////////////////////////////////////////////////////////////////////////////////////////////////////
+
+FMetadataProvider::FMetadataProvider(IAnalysisSession& InSession)
+	: Session(InSession)
+	, MetadataStore(Session.GetLinearAllocator(), 1024)
+{
+}
+
+////////////////////////////////////////////////////////////////////////////////////////////////////
+
+FMetadataProvider::~FMetadataProvider()
+{
+	for (const auto& KV : Threads)
+	{
+		for (const FMetadataStackEntry& StackEntry : KV.Value->CurrentStack)
+		{
+			if (StackEntry.StoreIndex == InvalidMetadataStoreIndex &&
+				StackEntry.StoreEntry.Size > MaxInlinedMetadataSize)
+			{
+				FMemory::Free(StackEntry.StoreEntry.Ptr);
+				check(AllocationCount > 0);
+				--AllocationCount;
+				check(TotalAllocatedMemory >= StackEntry.StoreEntry.Size);
+				TotalAllocatedMemory -= StackEntry.StoreEntry.Size;
+			}
+		}
+		delete KV.Value;
+	}
+
+	if (MetadataStore.Num() > 0)
+	{
+		for (const FMetadataStoreEntry& StoreEntry : MetadataStore)
+		{
+			if (StoreEntry.Size > MaxInlinedMetadataSize)
+			{
+				FMemory::Free(StoreEntry.Ptr);
+				check(AllocationCount > 0);
+				--AllocationCount;
+				check(TotalAllocatedMemory >= StoreEntry.Size);
+				TotalAllocatedMemory -= StoreEntry.Size;
+			}
+		}
+	}
+
+	check(AllocationCount == 0);
+	check(TotalAllocatedMemory == 0);
+}
+
+////////////////////////////////////////////////////////////////////////////////////////////////////
+
+uint16 FMetadataProvider::RegisterMetadataType(const TCHAR* Name, uint32 Size)
+{
+	Lock.WriteAccessCheck();
+
+	check(RegisteredTypes.Num() <= MaxMetadataTypeId);
+	const uint16 Type = RegisteredTypes.Num();
+
+	RegisteredTypes.Add({ Name, Size });
+	RegisteredTypesMap.Add(Name, Type);
+
+	return Type;
+}
+
+////////////////////////////////////////////////////////////////////////////////////////////////////
+
+uint16 FMetadataProvider::GetRegisteredMetadataType(const TCHAR* Name) const
+{
+	Lock.ReadAccessCheck();
+
+	const uint16* TypePtr = RegisteredTypesMap.Find(Name);
+	return TypePtr ? *TypePtr : InvalidMetadataType;
+}
+
+////////////////////////////////////////////////////////////////////////////////////////////////////
+
+const TCHAR* FMetadataProvider::GetRegisteredMetadataName(uint16 Type) const
+{
+	Lock.ReadAccessCheck();
+
+	const int32 Index = (int32)Type;
+	return Index < RegisteredTypes.Num() ? RegisteredTypes[Index].Name : nullptr;
+}
+
+////////////////////////////////////////////////////////////////////////////////////////////////////
+
+FMetadataProvider::FMetadataThread::FMetadataThread(uint32 InThreadId, ILinearAllocator& InAllocator)
+	: ThreadId(InThreadId)
+	, Metadata(InAllocator, 1024)
+{
+}
+
+////////////////////////////////////////////////////////////////////////////////////////////////////
+
+FMetadataProvider::FMetadataThread& FMetadataProvider::GetOrAddThread(uint32 ThreadId)
+{
+	FMetadataThread** MetadataThreadPtr = Threads.Find(ThreadId);
+	if (MetadataThreadPtr)
+	{
+		return **MetadataThreadPtr;
+	}
+	FMetadataThread* MetadataThread = new FMetadataThread(ThreadId, Session.GetLinearAllocator());
+	Threads.Add(ThreadId, MetadataThread);
+	return *MetadataThread;
+}
+
+////////////////////////////////////////////////////////////////////////////////////////////////////
+
+FMetadataProvider::FMetadataThread* FMetadataProvider::GetThread(uint32 ThreadId)
+{
+	FMetadataThread** MetadataThreadPtr = Threads.Find(ThreadId);
+	return (MetadataThreadPtr) ? *MetadataThreadPtr : nullptr;
+}
+
+////////////////////////////////////////////////////////////////////////////////////////////////////
+
+const FMetadataProvider::FMetadataThread* FMetadataProvider::GetThread(uint32 ThreadId) const
+{
+	FMetadataThread* const* MetadataThreadPtr = Threads.Find(ThreadId);
+	return (MetadataThreadPtr) ? *MetadataThreadPtr : nullptr;
+}
+
+////////////////////////////////////////////////////////////////////////////////////////////////////
+
+void FMetadataProvider::PushScopedMetadata(uint32 ThreadId, uint16 Type, void* Data, uint32 Size)
+{
+	Lock.WriteAccessCheck();
+
+	check(Type < RegisteredTypes.Num());
+	check(Data != nullptr && Size > 0);
+
+	if (Size > MaxMetadataSize)
+	{
+		UE_LOG(LogTraceServices, Error, TEXT("[Meta] Cannot push metadata (thread %u, type %u, size %u). Data size is too large."), ThreadId, Type, Size);
+		return;
+	}
+
+	FMetadataThread& MetadataThread = GetOrAddThread(ThreadId);
+
+	if (MetadataThread.CurrentStack.Num() == MaxMetadataStackSize)
+	{
+		UE_LOG(LogTraceServices, Error, TEXT("[Meta] Cannot push metadata (thread %u, type %u). Stack size is too large."), ThreadId, Type);
+		return;
+	}
+
+	FMetadataStackEntry StackEntry;
+	if (Size > MaxInlinedMetadataSize)
+	{
+		StackEntry.StoreEntry.Ptr = FMemory::Malloc(Size);
+		++AllocationEventCount;
+		++AllocationCount;
+		TotalAllocatedMemory += Size;
+		FMemory::Memcpy(StackEntry.StoreEntry.Ptr, Data, Size);
+	}
+	else
+	{
+		FMemory::Memcpy(StackEntry.StoreEntry.Value, Data, Size);
+	}
+	StackEntry.StoreEntry.Size = Size;
+	StackEntry.StoreEntry.Type = Type;
+	StackEntry.StoreIndex = InvalidMetadataStoreIndex;
+	StackEntry.PinnedId = InvalidMetadataId;
+
+	MetadataThread.CurrentStack.Push(StackEntry);
+	++EventCount;
+}
+
+////////////////////////////////////////////////////////////////////////////////////////////////////
+
+void FMetadataProvider::PopScopedMetadata(uint32 ThreadId, uint16 Type)
+{
+	Lock.WriteAccessCheck();
+
+	FMetadataThread& MetadataThread = GetOrAddThread(ThreadId);
+
+	if (MetadataThread.CurrentStack.Num() == 0)
+	{
+		UE_LOG(LogTraceServices, Error, TEXT("[Meta] Cannot pop metadata (thread %u, type %u). Stack is empty."), ThreadId, uint32(Type));
+		return;
+	}
+
+	const FMetadataStackEntry& Top = MetadataThread.CurrentStack.Top();
+
+	if (Top.StoreEntry.Type != Type)
+	{
+		UE_LOG(LogTraceServices, Error, TEXT("[Meta] Cannot pop metadata (thread %u, type %u). Type mismatch."), ThreadId, uint32(Type));
+		return;
+	}
+
+	if (Top.StoreIndex == InvalidMetadataStoreIndex && Top.StoreEntry.Size > MaxInlinedMetadataSize)
+	{
+		FMemory::Free(Top.StoreEntry.Ptr);
+		check(AllocationCount > 0);
+		--AllocationCount;
+		check(TotalAllocatedMemory >= Top.StoreEntry.Size);
+		TotalAllocatedMemory -= Top.StoreEntry.Size;
+	}
+	MetadataThread.CurrentStack.Pop();
+	++EventCount;
+}
+
+////////////////////////////////////////////////////////////////////////////////////////////////////
+
+uint32 FMetadataProvider::PinAndGetId(uint32 ThreadId)
+{
+	Lock.WriteAccessCheck();
+
+	FMetadataThread* MetadataThread = GetThread(ThreadId);
+
+	if (!MetadataThread || MetadataThread->CurrentStack.Num() == 0)
+	{
+		return InvalidMetadataId;
+	}
+
+	TArray<FMetadataStackEntry>& CurrentStack = MetadataThread->CurrentStack;
+	TPagedArray<FMetadataEntry>& Metadata = MetadataThread->Metadata;
+
+	const uint32 LastPinnedId = CurrentStack.Top().PinnedId;
+	if (LastPinnedId != InvalidMetadataId)
+	{
+		return LastPinnedId;
+	}
+
+	const int32 StackSize = CurrentStack.Num();
+
+	// Store all entries of the current stack.
+	for (int32 StackIndex = StackSize - 1; StackIndex >= 0; --StackIndex)
+	{
+		FMetadataStackEntry& StackEntry = CurrentStack[StackIndex];
+		if (StackEntry.PinnedId != InvalidMetadataId)
+		{
+			break;
+		}
+		if (StackEntry.StoreIndex == InvalidMetadataStoreIndex)
+		{
+			StackEntry.StoreIndex = MetadataStore.Num();
+			MetadataStore.EmplaceBack(StackEntry.StoreEntry);
+		}
+	}
+
+	int32 FirstUnpinnedStackIndex = StackSize - 1;
+	for (int32 StackIndex = StackSize - 2; StackIndex >= 0; --StackIndex)
+	{
+		const FMetadataStackEntry& StackEntry = CurrentStack[StackIndex];
+		if (StackEntry.PinnedId != InvalidMetadataId)
+		{
+			break;
+		}
+		FirstUnpinnedStackIndex = StackIndex;
+	}
+	if (FirstUnpinnedStackIndex > 0)
+	{
+		const FMetadataStackEntry& LastPinnedStackEntry = CurrentStack[FirstUnpinnedStackIndex - 1];
+		if (LastPinnedStackEntry.PinnedId == Metadata.Num() - 1)
+		{
+			// Reuse the partial pinned metadata stack.
+			// Only pin additional metadata entries.
+			for (int32 StackIndex = FirstUnpinnedStackIndex; StackIndex < StackSize; ++StackIndex)
+			{
+				FMetadataStackEntry& StackEntry = CurrentStack[StackIndex];
+				StackEntry.PinnedId = Metadata.Num();
+				FMetadataEntry& Entry = Metadata.PushBack();
+				Entry.StoreIndex = StackEntry.StoreIndex;
+				Entry.Type = StackEntry.StoreEntry.Type;
+				Entry.StackSize = StackIndex + 1;
+				Entry.Unused = 0;
+			}
+
+			return CurrentStack.Top().PinnedId;
+		}
+	}
+
+	// Pin all metadata entries in the current stack.
+	for (int32 StackIndex = StackSize - 1; StackIndex >= 0; --StackIndex)
+	{
+		FMetadataStackEntry& StackEntry = CurrentStack[StackIndex];
+		StackEntry.PinnedId = Metadata.Num();
+		FMetadataEntry& Entry = Metadata.PushBack();
+		Entry.StoreIndex = StackEntry.StoreIndex;
+		Entry.Type = StackEntry.StoreEntry.Type;
+		Entry.StackSize = StackIndex + 1;
+		Entry.Unused = 0;
+	}
+
+	return CurrentStack.Top().PinnedId;
+}
+
+////////////////////////////////////////////////////////////////////////////////////////////////////
+
+uint32 FMetadataProvider::GetMetadataStackSize(uint32 InThreadId, uint32 InMetadataId) const
+{
+	Lock.ReadAccessCheck();
+
+	const FMetadataThread* MetadataThread = GetThread(InThreadId);
+
+	if (!MetadataThread || InMetadataId >= MetadataThread->Metadata.Num())
+	{
+		UE_LOG(LogTraceServices, Warning, TEXT("[Meta] Invalid metadata id (thread %u, id %u)."), InThreadId, InMetadataId);
+		return 0;
+	}
+
+	auto Iterator = MetadataThread->Metadata.GetIteratorFromItem(InMetadataId);
+	const FMetadataEntry* Entry = Iterator.GetCurrentItem();
+	return Entry->StackSize;
+}
+
+////////////////////////////////////////////////////////////////////////////////////////////////////
+
+bool FMetadataProvider::GetMetadata(uint32 InThreadId, uint32 InMetadataId, uint32 InStackDepth, uint16& OutType, const void*& OutData, uint32& OutSize) const
+{
+	Lock.ReadAccessCheck();
+
+	const FMetadataThread* MetadataThread = GetThread(InThreadId);
+
+	if (!MetadataThread || InMetadataId >= MetadataThread->Metadata.Num())
+	{
+		UE_LOG(LogTraceServices, Warning, TEXT("[Meta] Invalid metadata id (thread %u, id %u)."), InThreadId, InMetadataId);
+		OutType = InvalidMetadataType;
+		OutData = nullptr;
+		OutSize = 0;
+		return false;
+	}
+
+	auto Iterator = MetadataThread->Metadata.GetIteratorFromItem(InMetadataId);
+	const FMetadataEntry* Entry = Iterator.GetCurrentItem();
+	if (InStackDepth >= Entry->StackSize)
+	{
+		UE_LOG(LogTraceServices, Warning, TEXT("[Meta] Invalid metadata stack detph (thread %u, id %u, depth %u)."), InThreadId, InMetadataId, InStackDepth);
+		OutType = InvalidMetadataType;
+		OutData = nullptr;
+		OutSize = 0;
+		return false;
+	}
+	while (Entry->StackSize != InStackDepth + 1)
+	{
+		Entry = Iterator.PrevItem();
+		check(Entry != nullptr);
+	}
+
+	check(Entry->StoreIndex < MetadataStore.Num());
+	const FMetadataStoreEntry& StoreEntry = *MetadataStore.GetIteratorFromItem(Entry->StoreIndex);
+	OutType = StoreEntry.Type;
+	OutData = (StoreEntry.Size > MaxInlinedMetadataSize) ? StoreEntry.Ptr : StoreEntry.Value;
+	OutSize = StoreEntry.Size;
+	return true;
+}
+
+////////////////////////////////////////////////////////////////////////////////////////////////////
+
+void FMetadataProvider::EnumerateMetadata(uint32 InThreadId, uint32 InMetadataId, TFunctionRef<bool(uint32 StackDepth, uint16 Type, const void* Data, uint32 Size)> Callback) const
+{
+	Lock.ReadAccessCheck();
+
+	const FMetadataThread* MetadataThread = GetThread(InThreadId);
+
+	if (!MetadataThread || InMetadataId >= MetadataThread->Metadata.Num())
+	{
+		UE_LOG(LogTraceServices, Warning, TEXT("[Meta] Invalid metadata id (thread %u, id %u)."), InThreadId, InMetadataId);
+		return;
+	}
+
+	auto Iterator = MetadataThread->Metadata.GetIteratorFromItem(InMetadataId);
+	const FMetadataEntry* Entry = Iterator.GetCurrentItem();
+	while (true)
+	{
+		check(Entry->StoreIndex < MetadataStore.Num());
+		const FMetadataStoreEntry& StoreEntry = *MetadataStore.GetIteratorFromItem(Entry->StoreIndex);
+		const void* Data = (StoreEntry.Size > MaxInlinedMetadataSize) ? StoreEntry.Ptr : StoreEntry.Value;
+		if (!Callback(Entry->StackSize - 1, StoreEntry.Type, Data, StoreEntry.Size))
+		{
+			break;
+		}
+		if (Entry->StackSize == 1) // last one
+		{
+			break;
+		}
+		Entry = Iterator.PrevItem();
+	}
+}
+
+////////////////////////////////////////////////////////////////////////////////////////////////////
+
+FName GetMetadataProviderName()
+{
+	static FName Name(TEXT("MetadataProvider"));
+	return Name;
+}
+
+////////////////////////////////////////////////////////////////////////////////////////////////////
+
+const IMetadataProvider* ReadMetadataProvider(const IAnalysisSession& Session)
+{
+	return Session.ReadProvider<IMetadataProvider>(GetMetadataProviderName());
+}
+
+////////////////////////////////////////////////////////////////////////////////////////////////////
+
+} // namespace TraceServices
diff --git a/Engine/Source/Developer/TraceServices/Private/Model/MetadataProvider.h b/Engine/Source/Developer/TraceServices/Private/Model/MetadataProvider.h
new file mode 100644
--- /dev/null
+++ b/Engine/Source/Developer/TraceServices/Private/Model/MetadataProvider.h
@@ -0,0 +1,152 @@
+// Copyright Epic Games, Inc. All Rights Reserved.
+
+#pragma once
+
+#include "Common/PagedArray.h"
+#include "Containers/Map.h"
+#include "HAL/CriticalSection.h"
+#include "TraceServices/Model/MetadataProvider.h"
+
+namespace TraceServices
+{
+
+class IAnalysisSession;
+class ILinearAllocator;
+
+////////////////////////////////////////////////////////////////////////////////////////////////////
+
+class FMetadataProviderLock
+{
+public:
+	void ReadAccessCheck() const;
+	void WriteAccessCheck() const;
+
+	void BeginRead();
+	void EndRead();
+
+	void BeginWrite();
+	void EndWrite();
+
+private:
+	FRWLock RWLock;
+};
+
+////////////////////////////////////////////////////////////////////////////////////////////////////
+
+class FMetadataProvider : public IMetadataProvider
+{
+public:
+	static constexpr uint32 MaxMetadataSize = 0xFFFF;
+	static constexpr uint32 InvalidMetadataId = 0xFFFFFFFF;
+	static constexpr uint16 InvalidMetadataType = 0xFFFF;
+
+private:
+	static constexpr uint32 MaxInlinedMetadataSize = 12;
+	static constexpr uint32 MaxMetadataTypeId = 0xFFFF;
+	static constexpr int32 MaxMetadataStackSize = 0xFF;
+	static constexpr uint32 InvalidMetadataStoreIndex = 0xFFFFFFFF;
+
+	struct FMetadataType
+	{
+		const TCHAR* Name;
+		uint32 Size; // 0 for variable size data
+	};
+
+	struct FMetadataStoreEntry
+	{
+		union
+		{
+			void* Ptr;
+			uint8 Value[8];
+		};
+		uint8 ValueEx[4];
+		uint16 Size;
+		uint16 Type;
+	};
+	static_assert(sizeof(FMetadataStoreEntry) == 16, "sizeof(FMetadataStoreEntry)");
+
+	struct FMetadataStackEntry
+	{
+		FMetadataStoreEntry StoreEntry;
+		uint32 StoreIndex;
+		uint32 PinnedId;
+	};
+	static_assert(sizeof(FMetadataStackEntry) == 24, "sizeof(FMetadataStackEntry)");
+
+	struct FMetadataEntry
+	{
+		uint32 StoreIndex;
+		uint16 Type;
+		uint8 StackSize;
+		uint8 Unused;
+	};
+	static_assert(sizeof(FMetadataEntry) == 8, "sizeof(FMetadataEntry)");
+
+	struct FMetadataThread
+	{
+		FMetadataThread(uint32 InThreadId, ILinearAllocator&);
+
+		uint32 ThreadId;
+		TArray<FMetadataStackEntry> CurrentStack;
+		TPagedArray<FMetadataEntry> Metadata; // a metadata id is an index in this array
+	};
+
+public:
+	explicit FMetadataProvider(IAnalysisSession& InSession);
+	virtual ~FMetadataProvider();
+
+	virtual void BeginEdit() const override { Lock.BeginWrite(); }
+	virtual void EndEdit() const override { Lock.EndWrite(); }
+	void EditAccessCheck() const { return Lock.WriteAccessCheck(); }
+
+	virtual void BeginRead() const override { Lock.BeginRead(); }
+	virtual void EndRead() const override { Lock.EndRead(); }
+	void ReadAccessCheck() const { return Lock.ReadAccessCheck(); }
+
+	//////////////////////////////////////////////////
+	// Edit operations
+
+	uint16 RegisterMetadataType(const TCHAR* Name, uint32 FixedSize = 0);
+
+	void PushScopedMetadata(uint32 ThreadId, uint16 Type, void* Data, uint32 Size);
+	void PopScopedMetadata(uint32 ThreadId, uint16 Type);
+
+	// Pins the metadata stack and returns an id for it.
+	uint32 PinAndGetId(uint32 ThreadId);
+
+	//////////////////////////////////////////////////
+	// Read operations
+
+	uint16 GetRegisteredMetadataType(const TCHAR* Name) const;
+	const TCHAR* GetRegisteredMetadataName(uint16 Type) const;
+
+	uint32 GetMetadataStackSize(uint32 InThreadId, uint32 InMetadataId) const;
+	bool GetMetadata(uint32 InThreadId, uint32 InMetadataId, uint32 InStackDepth, uint16& OutType, const void*& OutData, uint32& OutSize) const;
+	void EnumerateMetadata(uint32 InThreadId, uint32 InMetadataId, TFunctionRef<bool(uint32 StackDepth, uint16 Type, const void* Data, uint32 Size)> Callback) const;
+
+private:
+	FMetadataThread& GetOrAddThread(uint32 ThreadId);
+	FMetadataThread* GetThread(uint32 ThreadId);
+	const FMetadataThread* GetThread(uint32 ThreadId) const;
+
+private:
+	IAnalysisSession& Session;
+
+	mutable FMetadataProviderLock Lock;
+
+	TArray<FMetadataType> RegisteredTypes;
+	TMap<const TCHAR*, uint16> RegisteredTypesMap;
+
+	TPagedArray<FMetadataStoreEntry> MetadataStore; // stores individual metadata values
+
+	TMap<uint32, FMetadataThread*> Threads;
+
+	uint64 EventCount = 0; // debug
+	uint64 AllocationEventCount = 0; // debug
+	uint64 AllocationCount = 0; // debug
+	uint64 TotalAllocatedMemory = 0; // debug
+};
+
+////////////////////////////////////////////////////////////////////////////////////////////////////
+
+} // namespace TraceServices
diff --git a/Engine/Source/Developer/TraceServices/Private/Modules/MemoryModule.cpp b/Engine/Source/Developer/TraceServices/Private/Modules/MemoryModule.cpp
--- a/Engine/Source/Developer/TraceServices/Private/Modules/MemoryModule.cpp
+++ b/Engine/Source/Developer/TraceServices/Private/Modules/MemoryModule.cpp
@@ -7,6 +7,7 @@
 #include "Analyzers/ModuleAnalysis.h"
 #include "Model/AllocationsProvider.h"
 #include "Model/CallstacksProvider.h"
+#include "Model/MetadataProvider.h"
 #include "TraceServices/Model/AnalysisSession.h"
 
 namespace TraceServices
@@ -35,10 +36,14 @@ void FMemoryModule::OnAnalysisBegin(IAnalysisSession& Session)
 	Session.AddProvider(GetCallstacksProviderName(), CallstacksProvider);
 	Session.AddAnalyzer(new FCallstacksAnalyzer(Session, CallstacksProvider));
 
+	// Metadata
+	FMetadataProvider* MetadataProvider = new FMetadataProvider(Session);
+	Session.AddProvider(GetMetadataProviderName(), MetadataProvider);
+
 	// Allocations
-	FAllocationsProvider* AllocationsProvider = new FAllocationsProvider(Session);
+	FAllocationsProvider* AllocationsProvider = new FAllocationsProvider(Session, *MetadataProvider);
 	Session.AddProvider(GetAllocationsProviderName(), AllocationsProvider);
-	Session.AddAnalyzer(new FAllocationsAnalyzer(Session, *AllocationsProvider));
+	Session.AddAnalyzer(new FAllocationsAnalyzer(Session, *AllocationsProvider, *MetadataProvider));
 }
 
 void FMemoryModule::GetLoggers(TArray<const TCHAR*>& OutLoggers)
diff --git a/Engine/Source/Developer/TraceServices/Public/TraceServices/Model/MetadataProvider.h b/Engine/Source/Developer/TraceServices/Public/TraceServices/Model/MetadataProvider.h
new file mode 100644
--- /dev/null
+++ b/Engine/Source/Developer/TraceServices/Public/TraceServices/Model/MetadataProvider.h
@@ -0,0 +1,66 @@
+// Copyright Epic Games, Inc. All Rights Reserved.
+
+#pragma once
+
+#include "AnalysisSession.h"
+#include "CoreTypes.h"
+#include "Templates/UniquePtr.h"
+#include "ProfilingDebugging/MemoryTrace.h"
+
+#include <new>
+
+namespace TraceServices
+{
+
+////////////////////////////////////////////////////////////////////////////////////////////////////
+
+class IMetadataProvider : public IProvider
+{
+public:
+	struct TRACESERVICES_API FEditScopeLock
+	{
+		FEditScopeLock(const IMetadataProvider& InMetadataProvider)
+			: MetadataProvider(InMetadataProvider)
+		{
+			MetadataProvider.BeginEdit();
+		}
+
+		~FEditScopeLock()
+		{
+			MetadataProvider.EndEdit();
+		}
+
+		const IMetadataProvider& MetadataProvider;
+	};
+
+	struct TRACESERVICES_API FReadScopeLock
+	{
+		FReadScopeLock(const IMetadataProvider& InMetadataProvider)
+			: MetadataProvider(InMetadataProvider)
+		{
+			MetadataProvider.BeginRead();
+		}
+
+		~FReadScopeLock()
+		{
+			MetadataProvider.EndRead();
+		}
+
+		const IMetadataProvider& MetadataProvider;
+	};
+
+public:
+	virtual ~IMetadataProvider() = default;
+
+	virtual void BeginEdit() const = 0;
+	virtual void EndEdit() const = 0;
+	virtual void BeginRead() const = 0;
+	virtual void EndRead() const = 0;
+
+	//...
+};
+
+TRACESERVICES_API FName GetMetadataProviderName();
+TRACESERVICES_API const IMetadataProvider* ReadMetadataProvider(const IAnalysisSession& Session);
+
+} // namespace TraceServices
