[Insights]
  - CpuProfiler: Fixed "dispatch of timing events with metadata" when merged with regular events.
  - CpuProfiler: Added "Insights/FCpuProfilerAnalyzer" LLM tag to reduce Untagged mem allocations.
  - CpuProfiler: Added new debug code/macros around each timing event for easier debugging. This is preparation for further changes related to coroutines.

diff --git a/Engine/Source/Developer/TraceServices/Private/Analyzers/CpuProfilerTraceAnalysis.cpp b/Engine/Source/Developer/TraceServices/Private/Analyzers/CpuProfilerTraceAnalysis.cpp
--- a/Engine/Source/Developer/TraceServices/Private/Analyzers/CpuProfilerTraceAnalysis.cpp
+++ b/Engine/Source/Developer/TraceServices/Private/Analyzers/CpuProfilerTraceAnalysis.cpp
@@ -1,22 +1,34 @@
 // Copyright Epic Games, Inc. All Rights Reserved.
+
 #include "CpuProfilerTraceAnalysis.h"
+
 #include "AnalysisServicePrivate.h"
 #include "CborWriter.h"
-#include "Serialization/MemoryWriter.h"
 #include "Common/Utils.h"
+#include "HAL/LowLevelMemTracker.h"
 #include "Model/ThreadsPrivate.h"
+#include "Serialization/MemoryWriter.h"
+
+#define CPUPROFILER_DEBUG_LOGF(Format, ...) //{ if (ThreadState.ThreadId == 2) FPlatformMisc::LowLevelOutputDebugStringf(Format, __VA_ARGS__); }
+#define CPUPROFILER_DEBUG_BEGIN_EVENT(Time, Event) { ++TotalScopeCount; }
+#define CPUPROFILER_DEBUG_END_EVENT(Time)
+
+#define CPUPROFILER_SAFE_DISPATCH_PENDING_EVENTS 1
 
 namespace TraceServices
 {
 
+////////////////////////////////////////////////////////////////////////////////////////////////////
+
 FCpuProfilerAnalyzer::FCpuProfilerAnalyzer(IAnalysisSession& InSession, FTimingProfilerProvider& InTimingProfilerProvider, FThreadProvider& InThreadProvider)
 	: Session(InSession)
 	, TimingProfilerProvider(InTimingProfilerProvider)
 	, ThreadProvider(InThreadProvider)
 {
-
 }
 
+////////////////////////////////////////////////////////////////////////////////////////////////////
+
 FCpuProfilerAnalyzer::~FCpuProfilerAnalyzer()
 {
 	for (auto& KV : ThreadStatesMap)
@@ -26,6 +38,8 @@ FCpuProfilerAnalyzer::~FCpuProfilerAnalyzer()
 	}
 }
 
+////////////////////////////////////////////////////////////////////////////////////////////////////
+
 void FCpuProfilerAnalyzer::OnAnalysisBegin(const FOnAnalysisContext& Context)
 {
 	auto& Builder = Context.InterfaceBuilder;
@@ -39,8 +53,51 @@ void FCpuProfilerAnalyzer::OnAnalysisBegin(const FOnAnalysisContext& Context)
 	Builder.RouteEvent(RouteId_EndCaptureV2, "CpuProfiler", "EndCaptureV2");
 }
 
+////////////////////////////////////////////////////////////////////////////////////////////////////
+
+void FCpuProfilerAnalyzer::OnAnalysisEnd(/*const FOnAnalysisEndContext& Context*/)
+{
+#if 0
+	LLM_SCOPE_BYNAME(TEXT("Insights/FCpuProfilerAnalyzer"));
+
+	//TODO: Context.EventTime
+	for (auto& KV : ThreadStatesMap)
+	{
+		FThreadState& ThreadState = *KV.Value;
+
+		//TODO: EndThread(Context.EventTime, ThreadState);
+
+		ensure(ThreadState.PendingEvents.Num() == 0);
+
+		if (ThreadState.LastCycle != 0 && ThreadState.LastCycle != ~0)
+		{
+			//double Timestamp = Context.EventTime.AsSeconds(ThreadState.LastCycle);
+			//Session.UpdateDurationSeconds(Timestamp);
+			double Timestamp;
+			{
+				FAnalysisSessionEditScope _(Session);
+				Timestamp = Session.GetDurationSeconds();
+			}
+			while (ThreadState.ScopeStack.Num())
+			{
+				ThreadState.ScopeStack.Pop();
+				CPUPROFILER_DEBUG_LOGF(TEXT("[%u] ~E=%llu (%.9f)\n"), ThreadState.ThreadId, ThreadState.LastCycle, Timestamp);
+				ThreadState.Timeline->AppendEndEvent(Timestamp);
+				CPUPROFILER_DEBUG_END_EVENT(Timestamp);
+			}
+		}
+
+		ensure(ThreadState.ScopeStack.Num() == 0);
+	}
+#endif
+}
+
+////////////////////////////////////////////////////////////////////////////////////////////////////
+
 bool FCpuProfilerAnalyzer::OnEvent(uint16 RouteId, EStyle Style, const FOnEventContext& Context)
 {
+	LLM_SCOPE_BYNAME(TEXT("Insights/FCpuProfilerAnalyzer"));
+
 	FAnalysisSessionEditScope _(Session);
 
 	const auto& EventData = Context.EventData;
@@ -98,8 +155,15 @@ bool FCpuProfilerAnalyzer::OnEvent(uint16 RouteId, EStyle Style, const FOnEventC
 
 	case RouteId_EndThread:
 	{
-		uint32 ThreadId = FTraceAnalyzerUtils::GetThreadIdField(Context);
+		const uint32 ThreadId = FTraceAnalyzerUtils::GetThreadIdField(Context);
 		FThreadState& ThreadState = GetThreadState(ThreadId);
+
+		if (ThreadState.LastCycle == ~0ull)
+		{
+			// Ignore timing events received after EndThread.
+			break;
+		}
+
 		ensure(ThreadState.LastCycle != 0 || ThreadState.ScopeStack.Num() == 0);
 		if (ThreadState.LastCycle != 0)
 		{
@@ -108,7 +172,9 @@ bool FCpuProfilerAnalyzer::OnEvent(uint16 RouteId, EStyle Style, const FOnEventC
 			while (ThreadState.ScopeStack.Num())
 			{
 				ThreadState.ScopeStack.Pop();
+				CPUPROFILER_DEBUG_LOGF(TEXT("[%u] ^E=%llu (%.9f)\n"), ThreadState.ThreadId, ThreadState.LastCycle, Timestamp);
 				ThreadState.Timeline->AppendEndEvent(Timestamp);
+				CPUPROFILER_DEBUG_END_EVENT(Timestamp);
 			}
 		}
 		ThreadState.LastCycle = ~0ull;
@@ -130,7 +196,6 @@ bool FCpuProfilerAnalyzer::OnEvent(uint16 RouteId, EStyle Style, const FOnEventC
 		}
 
 		TArrayView<const uint8> DataView = FTraceAnalyzerUtils::LegacyAttachmentArray("Data", Context);
-		TotalEventSize += DataView.Num();
 		const uint32 BufferSize = DataView.Num();
 		const uint8* BufferPtr = DataView.GetData();
 
@@ -153,10 +218,14 @@ bool FCpuProfilerAnalyzer::OnEvent(uint16 RouteId, EStyle Style, const FOnEventC
 				for (int32 i = ThreadState.ScopeStack.Num(); i--;)
 				{
 					ThreadState.ScopeStack.Pop();
+					CPUPROFILER_DEBUG_LOGF(TEXT("[%u] -E=%llu (%.9f)\n"), ThreadState.ThreadId, LastCycle, LastTimestamp);
 					ThreadState.Timeline->AppendEndEvent(LastTimestamp);
+					CPUPROFILER_DEBUG_END_EVENT(LastTimestamp);
 				}
 			}
 		}
+
+		TotalEventSize += BufferSize;
 		BytesPerScope = double(TotalEventSize) / double(TotalScopeCount);
 		break;
 	}
@@ -176,10 +245,14 @@ bool FCpuProfilerAnalyzer::OnEvent(uint16 RouteId, EStyle Style, const FOnEventC
 	return true;
 }
 
+////////////////////////////////////////////////////////////////////////////////////////////////////
+
 uint64 FCpuProfilerAnalyzer::ProcessBuffer(const FEventTime& EventTime, FThreadState& ThreadState, const uint8* BufferPtr, uint32 BufferSize)
 {
 	uint64 LastCycle = ThreadState.LastCycle;
 
+	CPUPROFILER_DEBUG_LOGF(TEXT("[%u] ProcessBuffer %llu (%.9f)\n"), ThreadState.ThreadId, LastCycle, EventTime.AsSeconds(LastCycle));
+
 	int32 RemainingPending = ThreadState.PendingEvents.Num();
 	const FPendingEvent* PendingCursor = ThreadState.PendingEvents.GetData();
 
@@ -221,24 +294,45 @@ uint64 FCpuProfilerAnalyzer::ProcessBuffer(const FEventTime& EventTime, FThreadS
 				break;
 			}
 
+#if CPUPROFILER_SAFE_DISPATCH_PENDING_EVENTS
+			if (ensure(PendingCycle >= LastCycle))
+			{
+				// Update LastCycle in order to verify time (of following pending events) increases monotonically.
+				LastCycle = PendingCycle;
+			}
+			else
+			{
+				// Time needs to increase monotonically.
+				// We are not allowing events to "go back in time".
+				PendingCycle = LastCycle;
+			}
+#else
 			if (PendingCycle < LastCycle)
 			{
 				PendingCycle = LastCycle;
 			}
+#endif
 
 			double PendingTime = EventTime.AsSeconds(PendingCycle);
+
 			if (bEnter)
 			{
+				CPUPROFILER_DEBUG_LOGF(TEXT("[%u] <B=%llu (%.9f)\n"), ThreadState.ThreadId, PendingCycle, PendingTime);
 				FTimingProfilerEvent Event;
 				Event.TimerIndex = PendingCursor->TimerId;
 				ThreadState.Timeline->AppendBeginEvent(PendingTime, Event);
+				CPUPROFILER_DEBUG_BEGIN_EVENT(PendingTime, Event);
 			}
 			else
 			{
+				CPUPROFILER_DEBUG_LOGF(TEXT("[%u] <E=%llu (%.9f)\n"), ThreadState.ThreadId, PendingCycle, PendingTime);
 				ThreadState.Timeline->AppendEndEvent(PendingTime);
+				CPUPROFILER_DEBUG_END_EVENT(PendingTime);
 			}
 		}
 
+		double ActualTime = EventTime.AsSeconds(ActualCycle);
+
 		if (DecodedCycle & 1ull)
 		{
 			uint32 SpecId = FTraceAnalyzerUtils::Decode7bit(BufferPtr);
@@ -248,11 +342,11 @@ uint64 FCpuProfilerAnalyzer::ProcessBuffer(const FEventTime& EventTime, FThreadS
 			ScopeState.StartCycle = ActualCycle;
 			ScopeState.EventTypeId = TimerId;
 
+			CPUPROFILER_DEBUG_LOGF(TEXT("[%u]  B=%llu (%.9f)\n"), ThreadState.ThreadId, ActualCycle, ActualTime);
 			FTimingProfilerEvent Event;
 			Event.TimerIndex = TimerId;
-			double ActualTime = EventTime.AsSeconds(ActualCycle);
 			ThreadState.Timeline->AppendBeginEvent(ActualTime, Event);
-			++TotalScopeCount;
+			CPUPROFILER_DEBUG_BEGIN_EVENT(ActualTime, Event);
 		}
 		else
 		{
@@ -261,43 +355,73 @@ uint64 FCpuProfilerAnalyzer::ProcessBuffer(const FEventTime& EventTime, FThreadS
 			if (ThreadState.ScopeStack.Num() > 0)
 			{
 				ThreadState.ScopeStack.Pop();
-				double ActualTime = EventTime.AsSeconds(ActualCycle);
+				CPUPROFILER_DEBUG_LOGF(TEXT("[%u]  E=%llu (%.9f)\n"), ThreadState.ThreadId, ActualCycle, ActualTime);
 				ThreadState.Timeline->AppendEndEvent(ActualTime);
+				CPUPROFILER_DEBUG_END_EVENT(ActualTime);
 			}
 		}
 
+		check(ActualCycle > 0);
 		LastCycle = ActualCycle;
 	}
 	check(BufferPtr == BufferEnd);
 
-	// Dispatch remaining pending events.
+	// Dispatch remaining pending events...
 	for (; RemainingPending > 0; RemainingPending--, PendingCursor++)
 	{
+		bool bEnter = true;
 		uint64 PendingCycle = PendingCursor->Cycle;
 		if (int64(PendingCycle) < 0)
 		{
 			PendingCycle = ~PendingCycle;
-			double PendingTime = EventTime.AsSeconds(PendingCycle);
-			ThreadState.Timeline->AppendEndEvent(PendingTime);
+			bEnter = false;
+		}
+
+#if CPUPROFILER_SAFE_DISPATCH_PENDING_EVENTS
+		if (ensure(PendingCycle >= LastCycle))
+		{
+			// Update LastCycle in order to verify time (of following pending events) increases monotonically.
+			LastCycle = PendingCycle;
 		}
 		else
 		{
+			// Time needs to increase monotonically.
+			// We are not allowing events to "go back in time".
+			PendingCycle = LastCycle;
+		}
+#endif
+
+		double PendingTime = EventTime.AsSeconds(PendingCycle);
+
+		if (bEnter)
+		{
+			CPUPROFILER_DEBUG_LOGF(TEXT("[%u] >B=%llu (%.9f)\n"), ThreadState.ThreadId, PendingCycle, PendingTime);
 			FTimingProfilerEvent Event;
 			Event.TimerIndex = PendingCursor->TimerId;
-			double PendingTime = EventTime.AsSeconds(PendingCycle);
 			ThreadState.Timeline->AppendBeginEvent(PendingTime, Event);
+			CPUPROFILER_DEBUG_BEGIN_EVENT(PendingTime, Event);
+		}
+		else
+		{
+			CPUPROFILER_DEBUG_LOGF(TEXT("[%u] >E=%llu (%.9f)\n"), ThreadState.ThreadId, PendingCycle, PendingTime);
+			ThreadState.Timeline->AppendEndEvent(PendingTime);
+			CPUPROFILER_DEBUG_END_EVENT(PendingTime);
 		}
 	}
-
 	ThreadState.PendingEvents.Reset();
+
 	ThreadState.LastCycle = LastCycle;
 	return LastCycle;
 }
 
+////////////////////////////////////////////////////////////////////////////////////////////////////
+
 uint64 FCpuProfilerAnalyzer::ProcessBufferV2(const FEventTime& EventTime, FThreadState& ThreadState, const uint8* BufferPtr, uint32 BufferSize)
 {
 	uint64 LastCycle = ThreadState.LastCycle;
 
+	CPUPROFILER_DEBUG_LOGF(TEXT("[%u] ProcessBuffer %llu (%.9f)\n"), ThreadState.ThreadId, LastCycle, EventTime.AsSeconds(LastCycle));
+
 	int32 RemainingPending = ThreadState.PendingEvents.Num();
 	const FPendingEvent* PendingCursor = ThreadState.PendingEvents.GetData();
 
@@ -339,25 +463,44 @@ uint64 FCpuProfilerAnalyzer::ProcessBufferV2(const FEventTime& EventTime, FThrea
 				break;
 			}
 
+#if CPUPROFILER_SAFE_DISPATCH_PENDING_EVENTS
+			if (ensure(PendingCycle >= LastCycle))
+			{
+				// Update LastCycle in order to verify time (of following pending events) increases monotonically.
+				LastCycle = PendingCycle;
+			}
+			else
+			{
+				// Time needs to increase monotonically.
+				// We are not allowing events to "go back in time".
+				PendingCycle = LastCycle;
+			}
+#else
 			if (PendingCycle < LastCycle)
 			{
 				PendingCycle = LastCycle;
 			}
+#endif
 
 			double PendingTime = EventTime.AsSeconds(PendingCycle);
+
 			if (bEnter)
 			{
+				CPUPROFILER_DEBUG_LOGF(TEXT("[%u] <B=%llu (%.9f)\n"), ThreadState.ThreadId, PendingCycle, PendingTime);
 				FTimingProfilerEvent Event;
 				Event.TimerIndex = PendingCursor->TimerId;
 				ThreadState.Timeline->AppendBeginEvent(PendingTime, Event);
+				CPUPROFILER_DEBUG_BEGIN_EVENT(PendingTime, Event);
 			}
 			else
 			{
+				CPUPROFILER_DEBUG_LOGF(TEXT("[%u] <E=%llu (%.9f)\n"), ThreadState.ThreadId, PendingCycle, PendingTime);
 				ThreadState.Timeline->AppendEndEvent(PendingTime);
+				CPUPROFILER_DEBUG_END_EVENT(PendingTime);
 			}
 		}
 
-		const double ActualTime = EventTime.AsSeconds(ActualCycle);
+		double ActualTime = EventTime.AsSeconds(ActualCycle);
 
 		if (DecodedCycle & 2ull)
 		{
@@ -393,9 +536,11 @@ uint64 FCpuProfilerAnalyzer::ProcessBufferV2(const FEventTime& EventTime, FThrea
 					ScopeState.StartCycle = ActualCycle;
 					ScopeState.EventTypeId = MetadataTimerId;
 
+					CPUPROFILER_DEBUG_LOGF(TEXT("[%u] *B=%llu (%.9f)\n"), ThreadState.ThreadId, ActualCycle, ActualTime);
 					FTimingProfilerEvent Event;
 					Event.TimerIndex = MetadataTimerId;
 					ThreadState.Timeline->AppendBeginEvent(ActualTime, Event);
+					CPUPROFILER_DEBUG_BEGIN_EVENT(ActualTime, Event);
 				}
 
 				// Begins the cpu scoped timers (suspended in previous coroutine execution).
@@ -405,15 +550,18 @@ uint64 FCpuProfilerAnalyzer::ProcessBufferV2(const FEventTime& EventTime, FThrea
 						CoroutineUnknownTimerId = DefineNewTimerChecked(CoroutineUnknownSpecId, TEXT("<unknown>"));
 					}
 
+					//TODO: Restore the saved stack of cpu scoped timers for this CoroutineId.
 					for (uint32 i = 0; i < TimerScopeDepth; ++i)
 					{
 						FEventScopeState& ScopeState = ThreadState.ScopeStack.AddDefaulted_GetRef();
 						ScopeState.StartCycle = ActualCycle;
 						ScopeState.EventTypeId = CoroutineUnknownTimerId;
 
+						CPUPROFILER_DEBUG_LOGF(TEXT("[%u] +B=%llu (%.9f)\n"), ThreadState.ThreadId, ActualCycle, ActualTime);
 						FTimingProfilerEvent Event;
 						Event.TimerIndex = CoroutineUnknownTimerId;
 						ThreadState.Timeline->AppendBeginEvent(ActualTime, Event);
+						CPUPROFILER_DEBUG_BEGIN_EVENT(ActualTime, Event);
 					}
 				}
 			}
@@ -423,6 +571,8 @@ uint64 FCpuProfilerAnalyzer::ProcessBufferV2(const FEventTime& EventTime, FThrea
 
 				if (TimerScopeDepth != 0)
 				{
+					//TODO: Save current stack of cpu scoped timers (using id from metadata of CoroTask timer?)
+
 					// Ends (suspends) the cpu scoped timers.
 					for (uint32 i = 0; i < TimerScopeDepth; ++i)
 					{
@@ -431,7 +581,9 @@ uint64 FCpuProfilerAnalyzer::ProcessBufferV2(const FEventTime& EventTime, FThrea
 						if (ThreadState.ScopeStack.Num() > 0)
 						{
 							ThreadState.ScopeStack.Pop();
+							CPUPROFILER_DEBUG_LOGF(TEXT("[%u] +E=%llu (%.9f)\n"), ThreadState.ThreadId, ActualCycle, ActualTime);
 							ThreadState.Timeline->AppendEndEvent(ActualTime);
+							CPUPROFILER_DEBUG_END_EVENT(ActualTime);
 						}
 					}
 
@@ -452,7 +604,9 @@ uint64 FCpuProfilerAnalyzer::ProcessBufferV2(const FEventTime& EventTime, FThrea
 					if (ThreadState.ScopeStack.Num() > 0)
 					{
 						ThreadState.ScopeStack.Pop();
+						CPUPROFILER_DEBUG_LOGF(TEXT("[%u] *E=%llu (%.9f)\n"), ThreadState.ThreadId, ActualCycle, ActualTime);
 						ThreadState.Timeline->AppendEndEvent(ActualTime);
+						CPUPROFILER_DEBUG_END_EVENT(ActualTime);
 					}
 				}
 			}
@@ -468,10 +622,11 @@ uint64 FCpuProfilerAnalyzer::ProcessBufferV2(const FEventTime& EventTime, FThrea
 				ScopeState.StartCycle = ActualCycle;
 				ScopeState.EventTypeId = TimerId;
 
+				CPUPROFILER_DEBUG_LOGF(TEXT("[%u]  B=%llu (%.9f)\n"), ThreadState.ThreadId, ActualCycle, ActualTime);
 				FTimingProfilerEvent Event;
 				Event.TimerIndex = TimerId;
 				ThreadState.Timeline->AppendBeginEvent(ActualTime, Event);
-				++TotalScopeCount;
+				CPUPROFILER_DEBUG_BEGIN_EVENT(ActualTime, Event);
 			}
 			else
 			{
@@ -480,39 +635,68 @@ uint64 FCpuProfilerAnalyzer::ProcessBufferV2(const FEventTime& EventTime, FThrea
 				if (ThreadState.ScopeStack.Num() > 0)
 				{
 					ThreadState.ScopeStack.Pop();
+					CPUPROFILER_DEBUG_LOGF(TEXT("[%u]  E=%llu (%.9f)\n"), ThreadState.ThreadId, ActualCycle, ActualTime);
 					ThreadState.Timeline->AppendEndEvent(ActualTime);
+					CPUPROFILER_DEBUG_END_EVENT(ActualTime);
 				}
 			}
 		}
 
+		check(ActualCycle > 0);
 		LastCycle = ActualCycle;
 	}
 	check(BufferPtr == BufferEnd);
 
-	// Dispatch remaining pending events.
+	// Dispatch remaining pending events...
 	for (; RemainingPending > 0; RemainingPending--, PendingCursor++)
 	{
+		bool bEnter = true;
 		uint64 PendingCycle = PendingCursor->Cycle;
 		if (int64(PendingCycle) < 0)
 		{
 			PendingCycle = ~PendingCycle;
-			double PendingTime = EventTime.AsSeconds(PendingCycle);
-			ThreadState.Timeline->AppendEndEvent(PendingTime);
+			bEnter = false;
+		}
+
+#if CPUPROFILER_SAFE_DISPATCH_PENDING_EVENTS
+		if (ensure(PendingCycle >= LastCycle))
+		{
+			// Update LastCycle in order to verify time (of following pending events) increases monotonically.
+			LastCycle = PendingCycle;
 		}
 		else
 		{
+			// Time needs to increase monotonically.
+			// We are not allowing events to "go back in time".
+			PendingCycle = LastCycle;
+		}
+#endif
+
+		double PendingTime = EventTime.AsSeconds(PendingCycle);
+
+		if (bEnter)
+		{
+			CPUPROFILER_DEBUG_LOGF(TEXT("[%u] >B=%llu (%.9f)\n"), ThreadState.ThreadId, PendingCycle, PendingTime);
 			FTimingProfilerEvent Event;
 			Event.TimerIndex = PendingCursor->TimerId;
-			double PendingTime = EventTime.AsSeconds(PendingCycle);
 			ThreadState.Timeline->AppendBeginEvent(PendingTime, Event);
+			CPUPROFILER_DEBUG_BEGIN_EVENT(PendingTime, Event);
+		}
+		else
+		{
+			CPUPROFILER_DEBUG_LOGF(TEXT("[%u] >E=%llu (%.9f)\n"), ThreadState.ThreadId, PendingCycle, PendingTime);
+			ThreadState.Timeline->AppendEndEvent(PendingTime);
+			CPUPROFILER_DEBUG_END_EVENT(PendingTime);
 		}
 	}
-
 	ThreadState.PendingEvents.Reset();
+
 	ThreadState.LastCycle = LastCycle;
 	return LastCycle;
 }
 
+////////////////////////////////////////////////////////////////////////////////////////////////////
+
 void FCpuProfilerAnalyzer::OnCpuScopeEnter(const FOnEventContext& Context)
 {
 	if (Context.EventTime.GetTimestamp() == 0)
@@ -548,6 +732,8 @@ void FCpuProfilerAnalyzer::OnCpuScopeEnter(const FOnEventContext& Context)
 	ThreadState.PendingEvents.Add({Cycle, TimerId});
 }
 
+////////////////////////////////////////////////////////////////////////////////////////////////////
+
 void FCpuProfilerAnalyzer::OnCpuScopeLeave(const FOnEventContext& Context)
 {
 	if (Context.EventTime.GetTimestamp() == 0)
@@ -562,6 +748,8 @@ void FCpuProfilerAnalyzer::OnCpuScopeLeave(const FOnEventContext& Context)
 	ThreadState.PendingEvents.Add({~Cycle});
 }
 
+////////////////////////////////////////////////////////////////////////////////////////////////////
+
 uint32 FCpuProfilerAnalyzer::DefineTimer(uint32 SpecId, const TCHAR* Name, const TCHAR* File, uint32 Line, bool bMergeByName)
 {
 	// The cpu scoped events (timers) will be merged by name.
@@ -625,6 +813,8 @@ uint32 FCpuProfilerAnalyzer::DefineTimer(uint32 SpecId, const TCHAR* Name, const
 	}
 }
 
+////////////////////////////////////////////////////////////////////////////////////////////////////
+
 uint32 FCpuProfilerAnalyzer::DefineNewTimerChecked(uint32 SpecId, const TCHAR* TimerName, const TCHAR* File, uint32 Line)
 {
 	TimerName = Session.StoreString(TimerName);
@@ -633,6 +823,8 @@ uint32 FCpuProfilerAnalyzer::DefineNewTimerChecked(uint32 SpecId, const TCHAR* T
 	return NewTimerId;
 }
 
+////////////////////////////////////////////////////////////////////////////////////////////////////
+
 uint32 FCpuProfilerAnalyzer::GetTimerId(uint32 SpecId)
 {
 	if (uint32* FindIt = SpecIdToTimerIdMap.Find(SpecId))
@@ -648,12 +840,15 @@ uint32 FCpuProfilerAnalyzer::GetTimerId(uint32 SpecId)
 	}
 }
 
+////////////////////////////////////////////////////////////////////////////////////////////////////
+
 FCpuProfilerAnalyzer::FThreadState& FCpuProfilerAnalyzer::GetThreadState(uint32 ThreadId)
 {
 	FThreadState* ThreadState = ThreadStatesMap.FindRef(ThreadId);
 	if (!ThreadState)
 	{
 		ThreadState = new FThreadState();
+		ThreadState->ThreadId = ThreadId;
 		ThreadState->Timeline = &TimingProfilerProvider.EditCpuThreadTimeline(ThreadId);
 		ThreadStatesMap.Add(ThreadId, ThreadState);
 
@@ -664,4 +859,12 @@ FCpuProfilerAnalyzer::FThreadState& FCpuProfilerAnalyzer::GetThreadState(uint32
 	return *ThreadState;
 }
 
+////////////////////////////////////////////////////////////////////////////////////////////////////
+
 } // namespace TraceServices
+
+#undef CPUPROFILER_SAFE_DISPATCH_PENDING_EVENTS
+
+#undef CPUPROFILER_DEBUG_LOGF
+#undef CPUPROFILER_DEBUG_BEGIN_EVENT
+#undef CPUPROFILER_DEBUG_END_EVENT
diff --git a/Engine/Source/Developer/TraceServices/Private/Analyzers/CpuProfilerTraceAnalysis.h b/Engine/Source/Developer/TraceServices/Private/Analyzers/CpuProfilerTraceAnalysis.h
--- a/Engine/Source/Developer/TraceServices/Private/Analyzers/CpuProfilerTraceAnalysis.h
+++ b/Engine/Source/Developer/TraceServices/Private/Analyzers/CpuProfilerTraceAnalysis.h
@@ -20,6 +20,7 @@ public:
 	FCpuProfilerAnalyzer(IAnalysisSession& Session, FTimingProfilerProvider& TimingProfilerProvider, FThreadProvider& InThreadProvider);
 	~FCpuProfilerAnalyzer();
 	virtual void OnAnalysisBegin(const FOnAnalysisContext& Context) override;
+	virtual void OnAnalysisEnd(/*const FOnAnalysisEndContext& Context*/) override;
 	virtual bool OnEvent(uint16 RouteId, EStyle Style, const FOnEventContext& Context) override;
 
 private:
@@ -37,9 +38,10 @@ private:
 
 	struct FThreadState
 	{
+		uint32 ThreadId = 0;
 		TArray<FEventScopeState> ScopeStack;
 		TArray<FPendingEvent> PendingEvents;
-		FTimingProfilerProvider::TimelineInternal* Timeline;
+		FTimingProfilerProvider::TimelineInternal* Timeline = nullptr;
 		uint64 LastCycle = 0;
 	};
 
