[Insights]
  - Counters/Stats: Added support for stats counters with "ShouldClearEveryFrame" flag to actually reset to 0 at beginning of each Game frame (during analysis).
  - Counters/Stats: Added separate counters for "1/frame" variants of stats counters with "ShouldClearEveryFrame" flag. Normal counters shows all intermediate points (variation within a single frame), while the new ones (that have same name but with "(1/frame)" suffix) only shows a single value per each Game thread (i.e. last value in respective frame). These new counters are more useful when investigating variation between multiple frames.
  - Counters/Stats: Added "Insights/FStatsAnalyzer" LLM tag to reduce Untagged mem allocations.

diff --git a/Engine/Source/Developer/TraceServices/Private/Analyzers/StatsTraceAnalysis.cpp b/Engine/Source/Developer/TraceServices/Private/Analyzers/StatsTraceAnalysis.cpp
--- a/Engine/Source/Developer/TraceServices/Private/Analyzers/StatsTraceAnalysis.cpp
+++ b/Engine/Source/Developer/TraceServices/Private/Analyzers/StatsTraceAnalysis.cpp
@@ -1,9 +1,16 @@
 // Copyright Epic Games, Inc. All Rights Reserved.
+
 #include "StatsTraceAnalysis.h"
+
 #include "AnalysisServicePrivate.h"
 #include "Common/Utils.h"
+#include "HAL/LowLevelMemTracker.h"
 #include "TraceServices/Model/Counters.h"
 
+#include <limits>
+
+#define STATS_ANALYZER_DEBUG_LOG(StatId, Format, ...) //{ if (StatId == 389078 /*STAT_MeshDrawCalls*/) UE_LOG(LogTraceServices, Log, Format, ##__VA_ARGS__); }
+
 namespace TraceServices
 {
 
@@ -11,7 +18,6 @@ FStatsAnalyzer::FStatsAnalyzer(IAnalysisSession& InSession, ICounterProvider& In
 	: Session(InSession)
 	, CounterProvider(InCounterProvider)
 {
-
 }
 
 void FStatsAnalyzer::OnAnalysisBegin(const FOnAnalysisContext& Context)
@@ -20,10 +26,19 @@ void FStatsAnalyzer::OnAnalysisBegin(const FOnAnalysisContext& Context)
 
 	Builder.RouteEvent(RouteId_Spec, "Stats", "Spec");
 	Builder.RouteEvent(RouteId_EventBatch, "Stats", "EventBatch");
+	Builder.RouteEvent(RouteId_BeginFrame, "Misc", "BeginFrame");
+	Builder.RouteEvent(RouteId_EndFrame, "Misc", "EndFrame");
+}
+
+void FStatsAnalyzer::OnAnalysisEnd()
+{
+	CreateFrameCounters();
 }
 
 bool FStatsAnalyzer::OnEvent(uint16 RouteId, EStyle Style, const FOnEventContext& Context)
 {
+	LLM_SCOPE_BYNAME(TEXT("Insights/FStatsAnalyzer"));
+
 	FAnalysisSessionEditScope _(Session);
 
 	const auto& EventData = Context.EventData;
@@ -53,28 +68,31 @@ bool FStatsAnalyzer::OnEvent(uint16 RouteId, EStyle Style, const FOnEventContext
 				Description = reinterpret_cast<const TCHAR*>(EventData.GetAttachment() + Name.Len() + 1);
 			}
 
-
-			ECounterDisplayHint DisplayHint = CounterDisplayHint_None;
-			if (EventData.GetValue<bool>("IsMemory"))
-			{
-				DisplayHint = CounterDisplayHint_Memory;
-			}
-
 			if (Name.IsEmpty())
 			{
-				UE_LOG(LogTraceServices, Warning, TEXT("Got an invalid/empty counter name from StatsTrace events."))
+				UE_LOG(LogTraceServices, Warning, TEXT("Invalid counter name from Stats counter %u."), StatId);
+				Name = FString::Printf(TEXT("<UnknownStatsCounter%u>"), StatId);
 			}
-				
+
 			Counter->SetName(Session.StoreString(*Name));
 			if (!Group.IsEmpty())
 			{
 				Counter->SetGroup(Session.StoreString(*Group));
 			}
 			Counter->SetDescription(Session.StoreString(*Description));
+
 			Counter->SetIsFloatingPoint(EventData.GetValue<bool>("IsFloatingPoint"));
+			Counter->SetIsResetEveryFrame(EventData.GetValue<bool>("ShouldClearEveryFrame"));
+
+			ECounterDisplayHint DisplayHint = CounterDisplayHint_None;
+			if (EventData.GetValue<bool>("IsMemory"))
+			{
+				DisplayHint = CounterDisplayHint_Memory;
+			}
 			Counter->SetDisplayHint(DisplayHint);
 			break;
 		}
+
 		case RouteId_EventBatch:
 		{
 			uint32 ThreadId = FTraceAnalyzerUtils::GetThreadIdField(Context);
@@ -103,33 +121,39 @@ bool FStatsAnalyzer::OnEvent(uint16 RouteId, EStyle Style, const FOnEventContext
 					Counter = CounterProvider.CreateCounter();
 					CountersMap.Add(StatId, Counter);
 				}
+
 				uint8 Op = DecodedIdAndOp & 0x7;
 				uint64 CycleDiff = FTraceAnalyzerUtils::Decode7bit(BufferPtr);
 				uint64 Cycle = ThreadState->LastCycle + CycleDiff;
 				double Time = Context.EventTime.AsSeconds(Cycle);
 				ThreadState->LastCycle = Cycle;
+
 				switch (Op)
 				{
 				case Increment:
 				{
 					Counter->AddValue(Time, int64(1));
+					STATS_ANALYZER_DEBUG_LOG(StatId, TEXT("%f INC() %u"), Time, ThreadId);
 					break;
 				}
 				case Decrement:
 				{
 					Counter->AddValue(Time, int64(-1));
+					STATS_ANALYZER_DEBUG_LOG(StatId, TEXT("%f DEC() %u"), Time, ThreadId);
 					break;
 				}
 				case AddInteger:
 				{
 					int64 Amount = FTraceAnalyzerUtils::DecodeZigZag(BufferPtr);
 					Counter->AddValue(Time, Amount);
+					STATS_ANALYZER_DEBUG_LOG(StatId, TEXT("%f ADD(%lli) %u"), Time, Amount, ThreadId);
 					break;
 				}
 				case SetInteger:
 				{
 					int64 Value = FTraceAnalyzerUtils::DecodeZigZag(BufferPtr);
 					Counter->SetValue(Time, Value);
+					STATS_ANALYZER_DEBUG_LOG(StatId, TEXT("%f SET(%lli) %u"), Time, Value, ThreadId);
 					break;
 				}
 				case AddFloat:
@@ -138,6 +162,7 @@ bool FStatsAnalyzer::OnEvent(uint16 RouteId, EStyle Style, const FOnEventContext
 					memcpy(&Amount, BufferPtr, sizeof(double));
 					BufferPtr += sizeof(double);
 					Counter->AddValue(Time, Amount);
+					STATS_ANALYZER_DEBUG_LOG(StatId, TEXT("%f ADD(%f) %u"), Time, Amount, ThreadId);
 					break;
 				}
 				case SetFloat:
@@ -146,6 +171,7 @@ bool FStatsAnalyzer::OnEvent(uint16 RouteId, EStyle Style, const FOnEventContext
 					memcpy(&Value, BufferPtr, sizeof(double));
 					BufferPtr += sizeof(double);
 					Counter->SetValue(Time, Value);
+					STATS_ANALYZER_DEBUG_LOG(StatId, TEXT("%f SET(%f) %u"), Time, Value, ThreadId);
 					break;
 				}
 				}
@@ -153,6 +179,34 @@ bool FStatsAnalyzer::OnEvent(uint16 RouteId, EStyle Style, const FOnEventContext
 			check(BufferPtr == BufferEnd);
 			break;
 		}
+
+		case RouteId_BeginFrame:
+		//case RouteId_EndFrame:
+		{
+			uint8 FrameType = EventData.GetValue<uint8>("FrameType");
+			check(FrameType < TraceFrameType_Count);
+			if (ETraceFrameType(FrameType) == ETraceFrameType::TraceFrameType_Game)
+			{
+				uint64 Cycle = EventData.GetValue<uint64>("Cycle");
+				double Time = Context.EventTime.AsSeconds(Cycle);
+				for (auto& KV : CountersMap)
+				{
+					if (KV.Value->IsResetEveryFrame())
+					{
+						STATS_ANALYZER_DEBUG_LOG(KV.Key, TEXT("%f RESET"), Time);
+						if (KV.Value->IsFloatingPoint())
+						{
+							KV.Value->SetValue(Time, 0.0);
+						}
+						else
+						{
+							KV.Value->SetValue(Time, (int64)0);
+						}
+					}
+				}
+			}
+			break;
+		}
 	}
 
 	return true;
@@ -173,4 +227,90 @@ TSharedRef<FStatsAnalyzer::FThreadState> FStatsAnalyzer::GetThreadState(uint32 T
 	}
 }
 
+void FStatsAnalyzer::CreateFrameCounters()
+{
+	LLM_SCOPE_BYNAME(TEXT("Insights/FStatsAnalyzer"));
+
+	FAnalysisSessionEditScope _(Session);
+
+	TMap<uint32, IEditableCounter*> ResetEveryFrameCountersMap;
+
+	for (auto& KV : CountersMap)
+	{
+		uint32 StatId = KV.Key;
+		IEditableCounter* Counter = KV.Value;
+
+		if (Counter->IsResetEveryFrame())
+		{
+			ResetEveryFrameCountersMap.Add(StatId, Counter);
+		}
+	}
+
+	for (auto& KV : ResetEveryFrameCountersMap)
+	{
+		IEditableCounter* Counter = KV.Value;
+		IEditableCounter* FrameCounter = CounterProvider.CreateCounter();
+
+		FString FrameCounterName = FString(Counter->GetName()) + TEXT(" (1/frame)");
+		FrameCounter->SetName(Session.StoreString(*FrameCounterName));
+		if (Counter->GetGroup())
+		{
+			FrameCounter->SetGroup(Counter->GetGroup());
+		}
+		if (Counter->GetDescription())
+		{
+			FrameCounter->SetDescription(Counter->GetDescription());
+		}
+
+		FrameCounter->SetIsFloatingPoint(Counter->IsFloatingPoint());
+		FrameCounter->SetIsResetEveryFrame(false);
+		FrameCounter->SetDisplayHint(Counter->GetDisplayHint());
+
+		constexpr double InfiniteTime = std::numeric_limits<double>::infinity();
+
+		if (Counter->IsFloatingPoint())
+		{
+			bool bFirst = true;
+			double FrameTime = 0.0;
+			double FrameValue = 0.0;
+			Counter->EnumerateFloatValues(-InfiniteTime, InfiniteTime, false, [FrameCounter, &bFirst, &FrameTime, &FrameValue](double Time, double Value)
+				{
+					if (bFirst && Value != 0.0)
+					{
+						bFirst = false;
+						FrameCounter->SetValue(Time, 0.0);
+					}
+					if (Value == 0.0 && FrameValue != 0.0)
+					{
+						FrameCounter->SetValue(FrameTime, FrameValue);
+					}
+					FrameTime = Time;
+					FrameValue = Value;
+				});
+		}
+		else
+		{
+			bool bFirst = true;
+			double FrameTime = 0.0;
+			int64 FrameValue = 0;
+			Counter->EnumerateValues(-InfiniteTime, InfiniteTime, false, [FrameCounter, &bFirst, &FrameTime, &FrameValue](double Time, int64 Value)
+				{
+					if (bFirst && Value != 0)
+					{
+						bFirst = false;
+						FrameCounter->SetValue(Time, (int64)0);
+					}
+					if (Value == 0 && FrameValue != 0)
+					{
+						FrameCounter->SetValue(FrameTime, FrameValue);
+					}
+					FrameTime = Time;
+					FrameValue = Value;
+				});
+		}
+	}
+}
+
 } // namespace TraceServices
+
+#undef STATS_ANALYZER_DEBUG_LOG
diff --git a/Engine/Source/Developer/TraceServices/Private/Analyzers/StatsTraceAnalysis.h b/Engine/Source/Developer/TraceServices/Private/Analyzers/StatsTraceAnalysis.h
--- a/Engine/Source/Developer/TraceServices/Private/Analyzers/StatsTraceAnalysis.h
+++ b/Engine/Source/Developer/TraceServices/Private/Analyzers/StatsTraceAnalysis.h
@@ -18,13 +18,19 @@ class FStatsAnalyzer
 public:
 	FStatsAnalyzer(IAnalysisSession& Session, ICounterProvider& CounterProvider);
 	virtual void OnAnalysisBegin(const FOnAnalysisContext& Context) override;
+	virtual void OnAnalysisEnd() override;
 	virtual bool OnEvent(uint16 RouteId, EStyle Style, const FOnEventContext& Context) override;
 
+private:
+	void CreateFrameCounters();
+
 private:
 	enum : uint16
 	{
 		RouteId_Spec,
 		RouteId_EventBatch,
+		RouteId_BeginFrame,
+		RouteId_EndFrame
 	};
 
 	struct FThreadState
diff --git a/Engine/Source/Developer/TraceServices/Private/Model/CountersPrivate.h b/Engine/Source/Developer/TraceServices/Private/Model/CountersPrivate.h
--- a/Engine/Source/Developer/TraceServices/Private/Model/CountersPrivate.h
+++ b/Engine/Source/Developer/TraceServices/Private/Model/CountersPrivate.h
@@ -202,9 +202,10 @@ public:
 	virtual void SetDescription(const TCHAR* InDescription) override { Description = InDescription; }
 	virtual bool IsFloatingPoint() const override { return bIsFloatingPoint; }
 	virtual void SetIsFloatingPoint(bool bInIsFloatingPoint) override;
+	virtual bool IsResetEveryFrame() const override { return bIsResetEveryFrame; }
+	virtual void SetIsResetEveryFrame(bool bInIsResetEveryFrame) override { bIsResetEveryFrame = bInIsResetEveryFrame; }
 	virtual ECounterDisplayHint GetDisplayHint() const { return DisplayHint; }
 	virtual void SetDisplayHint(ECounterDisplayHint InDisplayHint) override { DisplayHint = InDisplayHint; }
-	virtual void SetIsResetEveryFrame(bool bInIsResetEveryFrame) override { bIsResetEveryFrame = bInIsResetEveryFrame; }
 	virtual void EnumerateValues(double IntervalStart, double IntervalEnd, bool bIncludeExternalBounds, TFunctionRef<void(double, int64)> Callback) const override;
 	virtual void EnumerateFloatValues(double IntervalStart, double IntervalEnd, bool bIncludeExternalBounds, TFunctionRef<void(double, double)> Callback) const override;
 	virtual void AddValue(double Time, int64 Value) override;
diff --git a/Engine/Source/Developer/TraceServices/Public/TraceServices/Model/Counters.h b/Engine/Source/Developer/TraceServices/Public/TraceServices/Model/Counters.h
--- a/Engine/Source/Developer/TraceServices/Public/TraceServices/Model/Counters.h
+++ b/Engine/Source/Developer/TraceServices/Public/TraceServices/Model/Counters.h
@@ -23,6 +23,7 @@ public:
 	virtual const TCHAR* GetGroup() const = 0;
 	virtual const TCHAR* GetDescription() const = 0;
 	virtual bool IsFloatingPoint() const = 0;
+	virtual bool IsResetEveryFrame() const { return false; }
 	virtual ECounterDisplayHint GetDisplayHint() const = 0;
 	virtual void EnumerateValues(double IntervalStart, double IntervalEnd, bool bIncludeExternalBounds, TFunctionRef<void(double, int64)> Callback) const = 0;
 	virtual void EnumerateFloatValues(double IntervalStart, double IntervalEnd, bool bIncludeExternalBounds, TFunctionRef<void(double, double)> Callback) const = 0;
@@ -38,8 +39,8 @@ public:
 	virtual void SetGroup(const TCHAR* Group) = 0;
 	virtual void SetDescription(const TCHAR* Description) = 0;
 	virtual void SetIsFloatingPoint(bool bIsFloatingPoint) = 0;
-	virtual void SetDisplayHint(ECounterDisplayHint DisplayHint) = 0;
 	virtual void SetIsResetEveryFrame(bool bInIsResetEveryFrame) = 0;
+	virtual void SetDisplayHint(ECounterDisplayHint DisplayHint) = 0;
 	virtual void AddValue(double Time, int64 Value) = 0;
 	virtual void AddValue(double Time, double Value) = 0;
 	virtual void SetValue(double Time, int64 Value) = 0;
