[Insights] Stats Counters: Added better detection for noname/unknwon stats counters.

diff --git a/Engine/Source/Developer/TraceServices/Private/Analyzers/StatsTraceAnalysis.cpp b/Engine/Source/Developer/TraceServices/Private/Analyzers/StatsTraceAnalysis.cpp
--- a/Engine/Source/Developer/TraceServices/Private/Analyzers/StatsTraceAnalysis.cpp
+++ b/Engine/Source/Developer/TraceServices/Private/Analyzers/StatsTraceAnalysis.cpp
@@ -70,11 +70,11 @@ bool FStatsAnalyzer::OnEvent(uint16 RouteId, EStyle Style, const FOnEventContext
 
 			if (Name.IsEmpty())
 			{
-				UE_LOG(LogTraceServices, Warning, TEXT("Invalid counter name from Stats counter %u."), StatId);
-				Name = FString::Printf(TEXT("<UnknownStatsCounter%u>"), StatId);
+				UE_LOG(LogTraceServices, Warning, TEXT("Invalid counter name for Stats counter %u."), StatId);
+				Name = FString::Printf(TEXT("<noname stats counter %u>"), StatId);
 			}
-
 			Counter->SetName(Session.StoreString(*Name));
+
 			if (!Group.IsEmpty())
 			{
 				Counter->SetGroup(Session.StoreString(*Group));
@@ -119,6 +119,8 @@ bool FStatsAnalyzer::OnEvent(uint16 RouteId, EStyle Style, const FOnEventContext
 				if (!Counter)
 				{
 					Counter = CounterProvider.CreateCounter();
+					FString Name = FString::Printf(TEXT("<unknown stats counter %u>"), StatId);
+					Counter->SetName(Session.StoreString(*Name));
 					CountersMap.Add(StatId, Counter);
 				}
 
