[Insights] Added LLM_SCOPEs to each analyzer in TraceServices and in few Insights plugins (in an effort to reduce the Untagged allocations).

diff --git a/Engine/Plugins/Animation/GameplayInsights/Source/GameplayInsights/Private/AnimationAnalyzer.cpp b/Engine/Plugins/Animation/GameplayInsights/Source/GameplayInsights/Private/AnimationAnalyzer.cpp
--- a/Engine/Plugins/Animation/GameplayInsights/Source/GameplayInsights/Private/AnimationAnalyzer.cpp
+++ b/Engine/Plugins/Animation/GameplayInsights/Source/GameplayInsights/Private/AnimationAnalyzer.cpp
@@ -1,10 +1,12 @@
 // Copyright Epic Games, Inc. All Rights Reserved.
 
 #include "AnimationAnalyzer.h"
-#include "TraceServices/Model/AnalysisSession.h"
-#include "TraceServices/Utils.h"
+
 #include "AnimationProvider.h"
 #include "Containers/ArrayView.h"
+#include "HAL/LowLevelMemTracker.h"
+#include "TraceServices/Model/AnalysisSession.h"
+#include "TraceServices/Utils.h"
 
 FAnimationAnalyzer::FAnimationAnalyzer(TraceServices::IAnalysisSession& InSession, FAnimationProvider& InAnimationProvider)
 	: Session(InSession)
@@ -50,6 +52,8 @@ bool FAnimationAnalyzer::OnEvent(uint16 RouteId, EStyle Style, const FOnEventCon
 {
 	using namespace TraceServices;
 
+	LLM_SCOPE_BYNAME(TEXT("Insights/FAnimationAnalyzer"));
+
 	FAnalysisSessionEditScope _(Session);
 
 	const auto& EventData = Context.EventData;
diff --git a/Engine/Plugins/Animation/GameplayInsights/Source/GameplayInsights/Private/GameplayAnalyzer.cpp b/Engine/Plugins/Animation/GameplayInsights/Source/GameplayInsights/Private/GameplayAnalyzer.cpp
--- a/Engine/Plugins/Animation/GameplayInsights/Source/GameplayInsights/Private/GameplayAnalyzer.cpp
+++ b/Engine/Plugins/Animation/GameplayInsights/Source/GameplayInsights/Private/GameplayAnalyzer.cpp
@@ -1,9 +1,11 @@
 // Copyright Epic Games, Inc. All Rights Reserved.
 
 #include "GameplayAnalyzer.h"
-#include "TraceServices/Utils.h"
-#include "TraceServices/Model/AnalysisSession.h"
+
 #include "GameplayProvider.h"
+#include "HAL/LowLevelMemTracker.h"
+#include "TraceServices/Model/AnalysisSession.h"
+#include "TraceServices/Utils.h"
 
 FGameplayAnalyzer::FGameplayAnalyzer(TraceServices::IAnalysisSession& InSession, FGameplayProvider& InGameplayProvider)
 	: Session(InSession)
@@ -34,6 +36,8 @@ void FGameplayAnalyzer::OnAnalysisBegin(const FOnAnalysisContext& Context)
 
 bool FGameplayAnalyzer::OnEvent(uint16 RouteId, EStyle Style, const FOnEventContext& Context)
 {
+	LLM_SCOPE_BYNAME(TEXT("Insights/FGameplayAnalyzer"));
+
 	TraceServices::FAnalysisSessionEditScope _(Session);
 
 	const auto& EventData = Context.EventData;
diff --git a/Engine/Plugins/Animation/GameplayInsights/Source/RewindDebuggerVLog/Private/VisualLoggerAnalyzer.cpp b/Engine/Plugins/Animation/GameplayInsights/Source/RewindDebuggerVLog/Private/VisualLoggerAnalyzer.cpp
--- a/Engine/Plugins/Animation/GameplayInsights/Source/RewindDebuggerVLog/Private/VisualLoggerAnalyzer.cpp
+++ b/Engine/Plugins/Animation/GameplayInsights/Source/RewindDebuggerVLog/Private/VisualLoggerAnalyzer.cpp
@@ -1,9 +1,11 @@
 // Copyright Epic Games, Inc. All Rights Reserved.
 
 #include "VisualLoggerAnalyzer.h"
+
+#include "HAL/LowLevelMemTracker.h"
+#include "Serialization/MemoryReader.h"
 #include "TraceServices/Model/AnalysisSession.h"
 #include "VisualLoggerProvider.h"
-#include "Serialization/MemoryReader.h"
 
 FVisualLoggerAnalyzer::FVisualLoggerAnalyzer(TraceServices::IAnalysisSession& InSession, FVisualLoggerProvider& InProvider)
 	: Session(InSession)
@@ -20,6 +22,8 @@ void FVisualLoggerAnalyzer::OnAnalysisBegin(const FOnAnalysisContext& Context)
 
 bool FVisualLoggerAnalyzer::OnEvent(uint16 RouteId, EStyle Style, const FOnEventContext& Context)
 {
+	LLM_SCOPE_BYNAME(TEXT("Insights/FVisualLoggerAnalyzer"));
+
 	TraceServices::FAnalysisSessionEditScope _(Session);
 
 	const auto& EventData = Context.EventData;
diff --git a/Engine/Plugins/Experimental/Animation/PoseSearch/Source/Editor/Private/Trace/PoseSearchTraceAnalyzer.cpp b/Engine/Plugins/Experimental/Animation/PoseSearch/Source/Editor/Private/Trace/PoseSearchTraceAnalyzer.cpp
--- a/Engine/Plugins/Experimental/Animation/PoseSearch/Source/Editor/Private/Trace/PoseSearchTraceAnalyzer.cpp
+++ b/Engine/Plugins/Experimental/Animation/PoseSearch/Source/Editor/Private/Trace/PoseSearchTraceAnalyzer.cpp
@@ -1,6 +1,8 @@
 ﻿// Copyright Epic Games, Inc. All Rights Reserved.
 
 #include "PoseSearchTraceAnalyzer.h"
+
+#include "HAL/LowLevelMemTracker.h"
 #include "PoseSearchTraceProvider.h"
 #include "Runtime/Private/Trace/PoseSearchTraceLogger.h"
 #include "TraceServices/Model/AnalysisSession.h"
@@ -26,6 +28,8 @@ void FTraceAnalyzer::OnAnalysisBegin(const FOnAnalysisContext& Context)
 
 bool FTraceAnalyzer::OnEvent(uint16 RouteId, EStyle Style, const FOnEventContext& Context)
 {
+	LLM_SCOPE_BYNAME(TEXT("Insights/PoseSearch::FTraceAnalyzer"));
+
 	TraceServices::FAnalysisSessionEditScope Scope(Session);
 	const FEventData& EventData = Context.EventData;
 
diff --git a/Engine/Plugins/RenderGraphInsights/Source/RenderGraphInsights/Private/RenderGraphAnalyzer.cpp b/Engine/Plugins/RenderGraphInsights/Source/RenderGraphInsights/Private/RenderGraphAnalyzer.cpp
--- a/Engine/Plugins/RenderGraphInsights/Source/RenderGraphInsights/Private/RenderGraphAnalyzer.cpp
+++ b/Engine/Plugins/RenderGraphInsights/Source/RenderGraphInsights/Private/RenderGraphAnalyzer.cpp
@@ -2,6 +2,7 @@
 
 #include "RenderGraphAnalyzer.h"
 
+#include "HAL/LowLevelMemTracker.h"
 #include "RenderGraphProvider.h"
 #include "Trace/Analyzer.h"
 #include "TraceServices/Model/AnalysisSession.h"
@@ -32,6 +33,8 @@ void FRenderGraphAnalyzer::OnAnalysisBegin(const FOnAnalysisContext& Context)
 
 bool FRenderGraphAnalyzer::OnEvent(uint16 RouteId, EStyle Style, const FOnEventContext& Context)
 {
+	LLM_SCOPE_BYNAME(TEXT("Insights/FRenderGraphAnalyzer"));
+
 	TraceServices::FAnalysisSessionEditScope _(Session);
 	switch (RouteId)
 	{
diff --git a/Engine/Plugins/Runtime/NetworkPredictionInsights/Source/NetworkPredictionInsights/Private/NetworkPredictionAnalyzer.cpp b/Engine/Plugins/Runtime/NetworkPredictionInsights/Source/NetworkPredictionInsights/Private/NetworkPredictionAnalyzer.cpp
--- a/Engine/Plugins/Runtime/NetworkPredictionInsights/Source/NetworkPredictionInsights/Private/NetworkPredictionAnalyzer.cpp
+++ b/Engine/Plugins/Runtime/NetworkPredictionInsights/Source/NetworkPredictionInsights/Private/NetworkPredictionAnalyzer.cpp
@@ -1,8 +1,10 @@
 // Copyright Epic Games, Inc. All Rights Reserved.
 
 #include "NetworkPredictionAnalyzer.h"
-#include "NetworkPredictionProvider.h"
+
+#include "HAL/LowLevelMemTracker.h"
 #include "Containers/StringView.h"
+#include "NetworkPredictionProvider.h"
 #include "TraceServices/Utils.h"
 
 // As we are tracing from multiple threads and we have trace events that are stringed together we need to track some state per thread to not trace bad data.
@@ -67,6 +69,8 @@ void FNetworkPredictionAnalyzer::OnAnalysisEnd()
 
 bool FNetworkPredictionAnalyzer::OnEvent(uint16 RouteId, EStyle Style, const FOnEventContext& Context)
 {
+	LLM_SCOPE_BYNAME(TEXT("Insights/FNetworkPredictionAnalyzer"));
+
 	TraceServices::FAnalysisSessionEditScope _(Session);
 	const auto& EventData = Context.EventData;
 
diff --git a/Engine/Plugins/Slate/SlateInsights/Source/SlateInsights/Private/SlateAnalyzer.cpp b/Engine/Plugins/Slate/SlateInsights/Source/SlateInsights/Private/SlateAnalyzer.cpp
--- a/Engine/Plugins/Slate/SlateInsights/Source/SlateInsights/Private/SlateAnalyzer.cpp
+++ b/Engine/Plugins/Slate/SlateInsights/Source/SlateInsights/Private/SlateAnalyzer.cpp
@@ -4,6 +4,7 @@
 
 #include "FastUpdate/WidgetUpdateFlags.h"
 #include "Framework/Application/SlateApplication.h"
+#include "HAL/LowLevelMemTracker.h"
 #include "SlateProvider.h"
 #include "Trace/Analyzer.h"
 #include "TraceServices/Model/AnalysisSession.h"
@@ -37,6 +38,8 @@ void FSlateAnalyzer::OnAnalysisBegin(const FOnAnalysisContext& Context)
 
 bool FSlateAnalyzer::OnEvent(uint16 RouteId, EStyle Style, const FOnEventContext& Context)
 {
+	LLM_SCOPE_BYNAME(TEXT("Insights/FSlateAnalyzer"));
+
 	TraceServices::FAnalysisSessionEditScope _(Session);
 
 	const UE::Trace::IAnalyzer::FEventData& EventData = Context.EventData;
diff --git a/Engine/Source/Developer/TraceInsights/Private/Insights/InsightsStyle.cpp b/Engine/Source/Developer/TraceInsights/Private/Insights/InsightsStyle.cpp
--- a/Engine/Source/Developer/TraceInsights/Private/Insights/InsightsStyle.cpp
+++ b/Engine/Source/Developer/TraceInsights/Private/Insights/InsightsStyle.cpp
@@ -1,7 +1,9 @@
 // Copyright Epic Games, Inc. All Rights Reserved.
 
 #include "InsightsStyle.h"
+
 #include "Framework/Application/SlateApplication.h"
+#include "HAL/LowLevelMemTracker.h"
 #include "Styling/AppStyle.h"
 #include "Styling/SlateStyleMacros.h"
 #include "Styling/SlateStyleRegistry.h"
diff --git a/Engine/Source/Developer/TraceInsights/Private/Insights/StoreService/DiagnosticsSessionAnalyzer.cpp b/Engine/Source/Developer/TraceInsights/Private/Insights/StoreService/DiagnosticsSessionAnalyzer.cpp
--- a/Engine/Source/Developer/TraceInsights/Private/Insights/StoreService/DiagnosticsSessionAnalyzer.cpp
+++ b/Engine/Source/Developer/TraceInsights/Private/Insights/StoreService/DiagnosticsSessionAnalyzer.cpp
@@ -1,6 +1,8 @@
 // Copyright Epic Games, Inc. All Rights Reserved.
 
 #include "DiagnosticsSessionAnalyzer.h"
+
+#include "HAL/LowLevelMemTracker.h"
 #include "TraceServices/Model/Diagnostics.h"
 
 namespace Insights
@@ -16,6 +18,8 @@ void FDiagnosticsSessionAnalyzer::OnAnalysisBegin(const FOnAnalysisContext& Cont
 
 bool FDiagnosticsSessionAnalyzer::OnEvent(uint16 RouteId, EStyle, const FOnEventContext& Context)
 {
+	LLM_SCOPE_BYNAME(TEXT("Insights/FDiagnosticsSessionAnalyzer"));
+
 	const FEventData& EventData = Context.EventData;
 
 	switch (RouteId)
diff --git a/Engine/Source/Developer/TraceServices/Private/Analyzers/AllocationsAnalysis.cpp b/Engine/Source/Developer/TraceServices/Private/Analyzers/AllocationsAnalysis.cpp
--- a/Engine/Source/Developer/TraceServices/Private/Analyzers/AllocationsAnalysis.cpp
+++ b/Engine/Source/Developer/TraceServices/Private/Analyzers/AllocationsAnalysis.cpp
@@ -1,6 +1,8 @@
 // Copyright Epic Games, Inc. All Rights Reserved.
 
 #include "AllocationsAnalysis.h"
+
+#include "HAL/LowLevelMemTracker.h"
 #include "Model/AllocationsProvider.h"
 #include "Model/MetadataProvider.h"
 #include "ProfilingDebugging/MemoryTrace.h"
@@ -84,6 +86,8 @@ void FAllocationsAnalyzer::OnAnalysisEnd()
 
 bool FAllocationsAnalyzer::OnEvent(uint16 RouteId, EStyle Style, const FOnEventContext& Context)
 {
+	LLM_SCOPE_BYNAME(TEXT("Insights/FAllocationsAnalyzer"));
+
 	const auto& EventData = Context.EventData;
 	HeapId RootHeap = 0;
 
diff --git a/Engine/Source/Developer/TraceServices/Private/Analyzers/CallstacksAnalysis.cpp b/Engine/Source/Developer/TraceServices/Private/Analyzers/CallstacksAnalysis.cpp
--- a/Engine/Source/Developer/TraceServices/Private/Analyzers/CallstacksAnalysis.cpp
+++ b/Engine/Source/Developer/TraceServices/Private/Analyzers/CallstacksAnalysis.cpp
@@ -1,8 +1,10 @@
 // Copyright Epic Games, Inc. All Rights Reserved.
 
 #include "CallstacksAnalysis.h"
-#include "TraceServices/Model/AnalysisSession.h"
+
+#include "HAL/LowLevelMemTracker.h"
 #include "Model/CallstacksProvider.h"
+#include "TraceServices/Model/AnalysisSession.h"
 
 namespace TraceServices
 {
@@ -25,6 +27,8 @@ void FCallstacksAnalyzer::OnAnalysisBegin(const FOnAnalysisContext& Context)
 ////////////////////////////////////////////////////////////////////////////////
 bool FCallstacksAnalyzer::OnEvent(uint16 RouteId, EStyle Style, const FOnEventContext& Context)
 {
+	LLM_SCOPE_BYNAME(TEXT("Insights/FCallstacksAnalyzer"));
+
 	switch (RouteId)
 	{
 		case RouteId_Callstack:
diff --git a/Engine/Source/Developer/TraceServices/Private/Analyzers/CountersTraceAnalysis.cpp b/Engine/Source/Developer/TraceServices/Private/Analyzers/CountersTraceAnalysis.cpp
--- a/Engine/Source/Developer/TraceServices/Private/Analyzers/CountersTraceAnalysis.cpp
+++ b/Engine/Source/Developer/TraceServices/Private/Analyzers/CountersTraceAnalysis.cpp
@@ -1,8 +1,11 @@
 // Copyright Epic Games, Inc. All Rights Reserved.
+
 #include "CountersTraceAnalysis.h"
+
 #include "Common/Utils.h"
-#include "TraceServices/Model/Counters.h"
+#include "HAL/LowLevelMemTracker.h"
 #include "ProfilingDebugging/CountersTrace.h"
+#include "TraceServices/Model/Counters.h"
 
 namespace TraceServices
 {
@@ -24,6 +27,8 @@ void FCountersAnalyzer::OnAnalysisBegin(const FOnAnalysisContext& Context)
 
 bool FCountersAnalyzer::OnEvent(uint16 RouteId, EStyle Style, const FOnEventContext& Context)
 {
+	LLM_SCOPE_BYNAME(TEXT("Insights/FCountersAnalyzer"));
+
 	FAnalysisSessionEditScope _(Session);
 
 	const auto& EventData = Context.EventData;
@@ -44,8 +49,12 @@ bool FCountersAnalyzer::OnEvent(uint16 RouteId, EStyle Style, const FOnEventCont
 			Counter->SetDisplayHint(CounterDisplayHint_Memory);
 		}
 		FString Name = FTraceAnalyzerUtils::LegacyAttachmentString<TCHAR>("Name", Context);
-		const TCHAR* NamePtr = Session.StoreString(*Name);
-		Counter->SetName(NamePtr);
+		if (Name.IsEmpty())
+		{
+			UE_LOG(LogTraceServices, Warning, TEXT("Invalid counter name for counter %u."), uint32(CounterId));
+			Name = FString::Printf(TEXT("<noname counter %u>"), uint32(CounterId));
+		}
+		Counter->SetName(Session.StoreString(*Name));
 		CountersMap.Add(CounterId, Counter);
 		break;
 	}
@@ -54,10 +63,10 @@ bool FCountersAnalyzer::OnEvent(uint16 RouteId, EStyle Style, const FOnEventCont
 		uint16 CounterId = EventData.GetValue<uint16>("CounterId");
 		int64 Value = EventData.GetValue<int64>("Value");
 		double Timestamp = Context.EventTime.AsSeconds(EventData.GetValue<uint64>("Cycle"));
-		IEditableCounter** FindCounter = CountersMap.Find(CounterId);
-		if (FindCounter)
+		IEditableCounter* FindCounter = CountersMap.FindRef(CounterId);
+		if (ensure(FindCounter))
 		{
-			(*FindCounter)->SetValue(Timestamp, Value);
+			FindCounter->SetValue(Timestamp, Value);
 		}
 		break;
 	}
@@ -67,7 +76,7 @@ bool FCountersAnalyzer::OnEvent(uint16 RouteId, EStyle Style, const FOnEventCont
 		float Value = EventData.GetValue<float>("Value");
 		double Timestamp = Context.EventTime.AsSeconds(EventData.GetValue<uint64>("Cycle"));
 		IEditableCounter* FindCounter = CountersMap.FindRef(CounterId);
-		if (FindCounter)
+		if (ensure(FindCounter))
 		{
 			FindCounter->SetValue(Timestamp, Value);
 		}
diff --git a/Engine/Source/Developer/TraceServices/Private/Analyzers/CsvProfilerTraceAnalysis.cpp b/Engine/Source/Developer/TraceServices/Private/Analyzers/CsvProfilerTraceAnalysis.cpp
--- a/Engine/Source/Developer/TraceServices/Private/Analyzers/CsvProfilerTraceAnalysis.cpp
+++ b/Engine/Source/Developer/TraceServices/Private/Analyzers/CsvProfilerTraceAnalysis.cpp
@@ -1,7 +1,10 @@
 // Copyright Epic Games, Inc. All Rights Reserved.
+
 #include "CsvProfilerTraceAnalysis.h"
+
 #include "AnalysisServicePrivate.h"
 #include "Common/Utils.h"
+#include "HAL/LowLevelMemTracker.h"
 #include "TraceServices/Model/Counters.h"
 #include "TraceServices/Model/Frames.h"
 
@@ -65,6 +68,8 @@ void FCsvProfilerAnalyzer::OnAnalysisEnd()
 
 bool FCsvProfilerAnalyzer::OnEvent(uint16 RouteId, EStyle Style, const FOnEventContext& Context)
 {
+	LLM_SCOPE_BYNAME(TEXT("Insights/FCsvProfilerAnalyzer"));
+
 	FAnalysisSessionEditScope _(Session);
 
 	const auto& EventData = Context.EventData;
@@ -241,6 +246,12 @@ const TCHAR* FCsvProfilerAnalyzer::GetStatSeriesName(const FStatSeriesDefinition
 		Name = TEXT("COUNTS/") + Name;
 	}
 
+	if (Name.IsEmpty())
+	{
+		UE_LOG(LogTraceServices, Warning, TEXT("Invalid counter name for CSV column %d."), Definition->ColumnIndex);
+		Name = FString::Printf(TEXT("<noname CSV column %d>"), Definition->ColumnIndex);
+	}
+
 	return Session.StoreString(*Name);
 }
 
diff --git a/Engine/Source/Developer/TraceServices/Private/Analyzers/DiagnosticsAnalysis.cpp b/Engine/Source/Developer/TraceServices/Private/Analyzers/DiagnosticsAnalysis.cpp
--- a/Engine/Source/Developer/TraceServices/Private/Analyzers/DiagnosticsAnalysis.cpp
+++ b/Engine/Source/Developer/TraceServices/Private/Analyzers/DiagnosticsAnalysis.cpp
@@ -1,7 +1,10 @@
 // Copyright Epic Games, Inc. All Rights Reserved.
+
 #include "DiagnosticsAnalysis.h"
+
 #include "AnalysisServicePrivate.h"
 #include "Common/Utils.h"
+#include "HAL/LowLevelMemTracker.h"
 
 namespace TraceServices
 {
@@ -32,6 +35,8 @@ bool FDiagnosticsAnalyzer::OnEvent(uint16 RouteId, EStyle Style, const FOnEventC
 		return false;
 	}
 
+	LLM_SCOPE_BYNAME(TEXT("Insights/FDiagnosticsAnalyzer"));
+
 	FAnalysisSessionEditScope _(Session);
 
 	const auto& EventData = Context.EventData;
diff --git a/Engine/Source/Developer/TraceServices/Private/Analyzers/GpuProfilerTraceAnalysis.cpp b/Engine/Source/Developer/TraceServices/Private/Analyzers/GpuProfilerTraceAnalysis.cpp
--- a/Engine/Source/Developer/TraceServices/Private/Analyzers/GpuProfilerTraceAnalysis.cpp
+++ b/Engine/Source/Developer/TraceServices/Private/Analyzers/GpuProfilerTraceAnalysis.cpp
@@ -1,7 +1,10 @@
 // Copyright Epic Games, Inc. All Rights Reserved.
+
 #include "GpuProfilerTraceAnalysis.h"
+
 #include "AnalysisServicePrivate.h"
 #include "Common/Utils.h"
+#include "HAL/LowLevelMemTracker.h"
 
 namespace TraceServices
 {
@@ -26,6 +29,8 @@ void FGpuProfilerAnalyzer::OnAnalysisBegin(const FOnAnalysisContext& Context)
 
 bool FGpuProfilerAnalyzer::OnEvent(uint16 RouteId, EStyle Style, const FOnEventContext& Context)
 {
+	LLM_SCOPE_BYNAME(TEXT("Insights/FGpuProfilerAnalyzer"));
+
 	FAnalysisSessionEditScope _(Session);
 
 	const auto& EventData = Context.EventData;
diff --git a/Engine/Source/Developer/TraceServices/Private/Analyzers/LoadTimeTraceAnalysis.cpp b/Engine/Source/Developer/TraceServices/Private/Analyzers/LoadTimeTraceAnalysis.cpp
--- a/Engine/Source/Developer/TraceServices/Private/Analyzers/LoadTimeTraceAnalysis.cpp
+++ b/Engine/Source/Developer/TraceServices/Private/Analyzers/LoadTimeTraceAnalysis.cpp
@@ -1,11 +1,14 @@
 // Copyright Epic Games, Inc. All Rights Reserved.
+
 #include "LoadTimeTraceAnalysis.h"
 
-#include "Serialization/LoadTimeTrace.h"
-#include "Model/LoadTimeProfilerPrivate.h"
 #include "Analyzers/MiscTraceAnalysis.h"
 #include "Common/FormatArgs.h"
 #include "Common/Utils.h"
+#include "HAL/LowLevelMemTracker.h"
+#include "Model/LoadTimeProfilerPrivate.h"
+#include "Serialization/LoadTimeTrace.h"
+
 #include <limits>
 
 namespace TraceServices
@@ -95,6 +98,8 @@ void FAsyncLoadingTraceAnalyzer::OnAnalysisBegin(const FOnAnalysisContext& Conte
 
 bool FAsyncLoadingTraceAnalyzer::OnEvent(uint16 RouteId, EStyle Style, const FOnEventContext& Context)
 {
+	LLM_SCOPE_BYNAME(TEXT("Insights/FAsyncLoadingTraceAnalyzer"));
+
 	const auto& EventData = Context.EventData;
 	switch (RouteId)
 	{
diff --git a/Engine/Source/Developer/TraceServices/Private/Analyzers/LogTraceAnalysis.cpp b/Engine/Source/Developer/TraceServices/Private/Analyzers/LogTraceAnalysis.cpp
--- a/Engine/Source/Developer/TraceServices/Private/Analyzers/LogTraceAnalysis.cpp
+++ b/Engine/Source/Developer/TraceServices/Private/Analyzers/LogTraceAnalysis.cpp
@@ -1,7 +1,10 @@
 // Copyright Epic Games, Inc. All Rights Reserved.
+
 #include "LogTraceAnalysis.h"
+
 #include "AnalysisServicePrivate.h"
 #include "Common/Utils.h"
+#include "HAL/LowLevelMemTracker.h"
 #include "Logging/LogTrace.h"
 #include "Model/LogPrivate.h"
 
@@ -26,6 +29,8 @@ void FLogTraceAnalyzer::OnAnalysisBegin(const FOnAnalysisContext& Context)
 
 bool FLogTraceAnalyzer::OnEvent(uint16 RouteId, EStyle Style, const FOnEventContext& Context)
 {
+	LLM_SCOPE_BYNAME(TEXT("Insights/FLogTraceAnalyzer"));
+
 	FAnalysisSessionEditScope _(Session);
 
 	const auto& EventData = Context.EventData;
diff --git a/Engine/Source/Developer/TraceServices/Private/Analyzers/MemoryAnalysis.cpp b/Engine/Source/Developer/TraceServices/Private/Analyzers/MemoryAnalysis.cpp
--- a/Engine/Source/Developer/TraceServices/Private/Analyzers/MemoryAnalysis.cpp
+++ b/Engine/Source/Developer/TraceServices/Private/Analyzers/MemoryAnalysis.cpp
@@ -1,7 +1,10 @@
 // Copyright Epic Games, Inc. All Rights Reserved.
+
 #include "MemoryAnalysis.h"
+
 #include "AnalysisServicePrivate.h"
 #include "Common/Utils.h"
+#include "HAL/LowLevelMemTracker.h"
 
 namespace TraceServices
 {
@@ -33,6 +36,8 @@ bool FMemoryAnalyzer::OnEvent(uint16 RouteId, EStyle Style, const FOnEventContex
 		return false;
 	}
 
+	LLM_SCOPE_BYNAME(TEXT("Insights/FMemoryAnalyzer"));
+
 	FAnalysisSessionEditScope _(Session);
 
 	const auto& EventData = Context.EventData;
diff --git a/Engine/Source/Developer/TraceServices/Private/Analyzers/MiscTraceAnalysis.cpp b/Engine/Source/Developer/TraceServices/Private/Analyzers/MiscTraceAnalysis.cpp
--- a/Engine/Source/Developer/TraceServices/Private/Analyzers/MiscTraceAnalysis.cpp
+++ b/Engine/Source/Developer/TraceServices/Private/Analyzers/MiscTraceAnalysis.cpp
@@ -1,16 +1,20 @@
 // Copyright Epic Games, Inc. All Rights Reserved.
+
 #include "MiscTraceAnalysis.h"
-#include "TraceServices/Model/AnalysisSession.h"
-#include "Model/LogPrivate.h"
-#include "Model/ThreadsPrivate.h"
+
+#include "Common/Utils.h"
+#include "HAL/LowLevelMemTracker.h"
 #include "Model/BookmarksPrivate.h"
+#include "Model/Channel.h"
 #include "Model/FramesPrivate.h"
+#include "Model/LogPrivate.h"
 #include "Model/ScreenshotProviderPrivate.h"
-#include "Model/Channel.h"
-#include "Common/Utils.h"
+#include "Model/ThreadsPrivate.h"
+#include "TraceServices/Model/AnalysisSession.h"
 
 namespace TraceServices
 {
+
 FMiscTraceAnalyzer::FMiscTraceAnalyzer(IAnalysisSession& InSession,
 									   FThreadProvider& InThreadProvider,
 									   FBookmarkProvider& InBookmarkProvider,
@@ -73,6 +77,8 @@ void FMiscTraceAnalyzer::OnThreadInfo(const FThreadInfo& ThreadInfo)
 
 bool FMiscTraceAnalyzer::OnEvent(uint16 RouteId, EStyle Style, const FOnEventContext& Context)
 {
+	LLM_SCOPE_BYNAME(TEXT("Insights/FMiscTraceAnalyzer"));
+
 	FAnalysisSessionEditScope _(Session);
 
 	const auto& EventData = Context.EventData;
diff --git a/Engine/Source/Developer/TraceServices/Private/Analyzers/ModuleAnalysis.cpp b/Engine/Source/Developer/TraceServices/Private/Analyzers/ModuleAnalysis.cpp
--- a/Engine/Source/Developer/TraceServices/Private/Analyzers/ModuleAnalysis.cpp
+++ b/Engine/Source/Developer/TraceServices/Private/Analyzers/ModuleAnalysis.cpp
@@ -1,13 +1,17 @@
 // Copyright Epic Games, Inc. All Rights Reserved.
 
 #include "ModuleAnalysis.h"
-#include "TraceServices/Model/AnalysisSession.h"
+
+#include "HAL/LowLevelMemTracker.h"
 #include "Model/ModuleProvider.h"
+#include "TraceServices/Model/AnalysisSession.h"
 #include "UObject/NameTypes.h"
 
-namespace TraceServices {
+namespace TraceServices
+{
 
 ////////////////////////////////////////////////////////////////////////////////
+
 enum Routes
 {
 	RouteId_ModuleInit,
@@ -16,6 +20,7 @@ enum Routes
 };
 
 ////////////////////////////////////////////////////////////////////////////////
+
 FModuleAnalyzer::FModuleAnalyzer(IAnalysisSession& InSession)
 	: Session(InSession)
 	, Provider(nullptr)
@@ -24,6 +29,7 @@ FModuleAnalyzer::FModuleAnalyzer(IAnalysisSession& InSession)
 }
 
 ////////////////////////////////////////////////////////////////////////////////
+
 void FModuleAnalyzer::OnAnalysisBegin(const FOnAnalysisContext& Context)
 {
 	FInterfaceBuilder& Builder = Context.InterfaceBuilder;
@@ -33,6 +39,7 @@ void FModuleAnalyzer::OnAnalysisBegin(const FOnAnalysisContext& Context)
 }
 
 ////////////////////////////////////////////////////////////////////////////////
+
 void FModuleAnalyzer::OnAnalysisEnd()
 {
 	if (Provider)
@@ -43,9 +50,12 @@ void FModuleAnalyzer::OnAnalysisEnd()
 
 
 ////////////////////////////////////////////////////////////////////////////////
+
 bool FModuleAnalyzer::OnEvent(uint16 RouteId, EStyle Style, const FOnEventContext& Context)
 {
-	switch(RouteId)
+	LLM_SCOPE_BYNAME(TEXT("Insights/FModuleAnalyzer"));
+
+	switch (RouteId)
 	{
 		case RouteId_ModuleInit:
 			{
@@ -91,7 +101,8 @@ bool FModuleAnalyzer::OnEvent(uint16 RouteId, EStyle Style, const FOnEventContex
 }
 
 ////////////////////////////////////////////////////////////////////////////////
-uint64	FModuleAnalyzer::GetBaseAddress(const FEventData& EventData) const
+
+uint64 FModuleAnalyzer::GetBaseAddress(const FEventData& EventData) const
 {
 	if (ModuleBaseShift == 0)
 	{
@@ -107,4 +118,5 @@ uint64	FModuleAnalyzer::GetBaseAddress(const FEventData& EventData) const
 }
 
 ////////////////////////////////////////////////////////////////////////////////
+
 } // namespace TraceServices
diff --git a/Engine/Source/Developer/TraceServices/Private/Analyzers/NetTraceAnalyzer.cpp b/Engine/Source/Developer/TraceServices/Private/Analyzers/NetTraceAnalyzer.cpp
--- a/Engine/Source/Developer/TraceServices/Private/Analyzers/NetTraceAnalyzer.cpp
+++ b/Engine/Source/Developer/TraceServices/Private/Analyzers/NetTraceAnalyzer.cpp
@@ -1,9 +1,12 @@
 // Copyright Epic Games, Inc. All Rights Reserved.
+
 #include "NetTraceAnalyzer.h"
+
 #include "AnalysisServicePrivate.h"
 #include "Common/Utils.h"
-#include "TraceServices/Model/Threads.h"
+#include "HAL/LowLevelMemTracker.h"
 #include "Logging/LogMacros.h"
+#include "TraceServices/Model/Threads.h"
 
 DEFINE_LOG_CATEGORY_STATIC(LogNetTrace, Log, All);
 
@@ -77,6 +80,8 @@ uint32 FNetTraceAnalyzer::GetTracedEventTypeIndex(uint16 NameIndex, uint8 Level)
 
 bool FNetTraceAnalyzer::OnEvent(uint16 RouteId, EStyle Style, const FOnEventContext& Context)
 {
+	LLM_SCOPE_BYNAME(TEXT("Insights/FNetTraceAnalyzer"));
+
 	FAnalysisSessionEditScope _(Session);
 
 	// check that we always get the InitEvent before processing any other events
diff --git a/Engine/Source/Developer/TraceServices/Private/Analyzers/PlatformEventTraceAnalysis.cpp b/Engine/Source/Developer/TraceServices/Private/Analyzers/PlatformEventTraceAnalysis.cpp
--- a/Engine/Source/Developer/TraceServices/Private/Analyzers/PlatformEventTraceAnalysis.cpp
+++ b/Engine/Source/Developer/TraceServices/Private/Analyzers/PlatformEventTraceAnalysis.cpp
@@ -1,8 +1,11 @@
 // Copyright Epic Games, Inc. All Rights Reserved.
+
 #include "PlatformEventTraceAnalysis.h"
-#include "TraceServices/Model/AnalysisSession.h"
+
+#include "HAL/LowLevelMemTracker.h"
 #include "Model/ContextSwitchesPrivate.h"
 #include "Model/StackSamplesPrivate.h"
+#include "TraceServices/Model/AnalysisSession.h"
 
 namespace TraceServices
 {
@@ -27,6 +30,8 @@ void FPlatformEventTraceAnalyzer::OnAnalysisBegin(const FOnAnalysisContext& Cont
 
 bool FPlatformEventTraceAnalyzer::OnEvent(uint16 RouteId, EStyle Style, const FOnEventContext& Context)
 {
+	LLM_SCOPE_BYNAME(TEXT("Insights/FPlatformEventTraceAnalyzer"));
+
 	FAnalysisSessionEditScope _(Session);
 
 	const auto& EventData = Context.EventData;
diff --git a/Engine/Source/Developer/TraceServices/Private/Analyzers/PlatformFileTraceAnalysis.cpp b/Engine/Source/Developer/TraceServices/Private/Analyzers/PlatformFileTraceAnalysis.cpp
--- a/Engine/Source/Developer/TraceServices/Private/Analyzers/PlatformFileTraceAnalysis.cpp
+++ b/Engine/Source/Developer/TraceServices/Private/Analyzers/PlatformFileTraceAnalysis.cpp
@@ -1,7 +1,10 @@
 // Copyright Epic Games, Inc. All Rights Reserved.
+
 #include "PlatformFileTraceAnalysis.h"
+
 #include "AnalysisServicePrivate.h"
 #include "Common/Utils.h"
+#include "HAL/LowLevelMemTracker.h"
 #include "Model/FileActivity.h"
 
 #define DEBUG_PLATFORMFILETRACE 0
@@ -40,6 +43,8 @@ void FPlatformFileTraceAnalyzer::OnAnalysisBegin(const FOnAnalysisContext& Conte
 
 bool FPlatformFileTraceAnalyzer::OnEvent(uint16 RouteId, EStyle Style, const FOnEventContext& Context)
 {
+	LLM_SCOPE_BYNAME(TEXT("Insights/FPlatformFileTraceAnalyzer"));
+
 	FAnalysisSessionEditScope _(Session);
 
 	const auto& EventData = Context.EventData;
@@ -327,7 +332,7 @@ bool FPlatformFileTraceAnalyzer::OnEvent(uint16 RouteId, EStyle Style, const FOn
 	return true;
 }
 
+} // namespace TraceServices
+
 #undef PLATFORMFILETRACE_WARNING
 #undef DEBUG_PLATFORMFILETRACE
-
-} // namespace TraceServices
diff --git a/Engine/Source/Developer/TraceServices/Private/Analyzers/TasksAnalysis.cpp b/Engine/Source/Developer/TraceServices/Private/Analyzers/TasksAnalysis.cpp
--- a/Engine/Source/Developer/TraceServices/Private/Analyzers/TasksAnalysis.cpp
+++ b/Engine/Source/Developer/TraceServices/Private/Analyzers/TasksAnalysis.cpp
@@ -1,46 +1,51 @@
 // Copyright Epic Games, Inc. All Rights Reserved.
 
 #include "TasksAnalysis.h"
+
 #include "AnalysisServicePrivate.h"
-#include "Model/TasksProfilerPrivate.h"
 #include "Containers/UnrealString.h"
+#include "HAL/LowLevelMemTracker.h"
+#include "Model/TasksProfilerPrivate.h"
 
 namespace TraceServices
 {
-	FTasksAnalyzer::FTasksAnalyzer(IAnalysisSession& InSession, FTasksProvider& InTasksProvider)
-		: Session(InSession)
-		, TasksProvider(InTasksProvider)
-	{
-	}
 
-	void FTasksAnalyzer::OnAnalysisBegin(const FOnAnalysisContext& Context)
-	{
-		auto& Builder = Context.InterfaceBuilder;
-
-		Builder.RouteEvent(RouteId_Init, "TaskTrace", "Init");
-		Builder.RouteEvent(RouteId_Created, "TaskTrace", "Created");
-		Builder.RouteEvent(RouteId_Launched, "TaskTrace", "Launched");
-		Builder.RouteEvent(RouteId_Scheduled, "TaskTrace", "Scheduled");
-		Builder.RouteEvent(RouteId_SubsequentAdded, "TaskTrace", "SubsequentAdded");
-		Builder.RouteEvent(RouteId_Started, "TaskTrace", "Started");
-		Builder.RouteEvent(RouteId_NestedAdded, "TaskTrace", "NestedAdded");
-		Builder.RouteEvent(RouteId_Finished, "TaskTrace", "Finished");
-		Builder.RouteEvent(RouteId_Completed, "TaskTrace", "Completed");
-
-		Builder.RouteEvent(RouteId_WaitingStarted, "TaskTrace", "WaitingStarted");
-		Builder.RouteEvent(RouteId_WaitingFinished, "TaskTrace", "WaitingFinished");
-	}
+FTasksAnalyzer::FTasksAnalyzer(IAnalysisSession& InSession, FTasksProvider& InTasksProvider)
+	: Session(InSession)
+	, TasksProvider(InTasksProvider)
+{
+}
 
-	bool FTasksAnalyzer::OnEvent(uint16 RouteId, EStyle Style, const FOnEventContext& Context)
-	{
-		FAnalysisSessionEditScope _(Session);
+void FTasksAnalyzer::OnAnalysisBegin(const FOnAnalysisContext& Context)
+{
+	auto& Builder = Context.InterfaceBuilder;
+
+	Builder.RouteEvent(RouteId_Init, "TaskTrace", "Init");
+	Builder.RouteEvent(RouteId_Created, "TaskTrace", "Created");
+	Builder.RouteEvent(RouteId_Launched, "TaskTrace", "Launched");
+	Builder.RouteEvent(RouteId_Scheduled, "TaskTrace", "Scheduled");
+	Builder.RouteEvent(RouteId_SubsequentAdded, "TaskTrace", "SubsequentAdded");
+	Builder.RouteEvent(RouteId_Started, "TaskTrace", "Started");
+	Builder.RouteEvent(RouteId_NestedAdded, "TaskTrace", "NestedAdded");
+	Builder.RouteEvent(RouteId_Finished, "TaskTrace", "Finished");
+	Builder.RouteEvent(RouteId_Completed, "TaskTrace", "Completed");
+
+	Builder.RouteEvent(RouteId_WaitingStarted, "TaskTrace", "WaitingStarted");
+	Builder.RouteEvent(RouteId_WaitingFinished, "TaskTrace", "WaitingFinished");
+}
+
+bool FTasksAnalyzer::OnEvent(uint16 RouteId, EStyle Style, const FOnEventContext& Context)
+{
+	LLM_SCOPE_BYNAME(TEXT("Insights/FTasksAnalyzer"));
+
+	FAnalysisSessionEditScope _(Session);
 	
-		// the protocol must match TaskTrace.cpp
+	// the protocol must match TaskTrace.cpp
 
-		const auto& EventData = Context.EventData;
-		uint32 ThreadId = Context.ThreadInfo.GetId();
-		switch (RouteId)
-		{
+	const auto& EventData = Context.EventData;
+	uint32 ThreadId = Context.ThreadInfo.GetId();
+	switch (RouteId)
+	{
 		case RouteId_Init:
 		{
 			uint32 Version = EventData.GetValue<uint32>("Version");
@@ -130,8 +135,9 @@ namespace TraceServices
 			TasksProvider.WaitingFinished(Timestamp, ThreadId);
 			break;
 		}
-		}
-
-		return true;
 	}
-}
\ No newline at end of file
+
+	return true;
+}
+
+} // namespace TraceServices
