[Insights]
  - Added "Insights/TraceServices" memory tag for untagged memory allocations from TraceServices module.
  - Added more mem scopes for FMiscTraceAnalyzer and FPlatformEventTraceAnalyzer.

diff --git a/Engine/Source/Developer/TraceServices/Private/Analyzers/MiscTraceAnalysis.cpp b/Engine/Source/Developer/TraceServices/Private/Analyzers/MiscTraceAnalysis.cpp
--- a/Engine/Source/Developer/TraceServices/Private/Analyzers/MiscTraceAnalysis.cpp
+++ b/Engine/Source/Developer/TraceServices/Private/Analyzers/MiscTraceAnalysis.cpp
@@ -60,6 +60,8 @@
 
 void FMiscTraceAnalyzer::OnThreadInfo(const FThreadInfo& ThreadInfo)
 {
+	LLM_SCOPE_BYNAME(TEXT("Insights/FMiscTraceAnalyzer"));
+
 	uint32 ThreadId = ThreadInfo.GetId();
 	FString Name = ThreadInfo.GetName();
 
diff --git a/Engine/Source/Developer/TraceServices/Private/Analyzers/PlatformEventTraceAnalysis.cpp b/Engine/Source/Developer/TraceServices/Private/Analyzers/PlatformEventTraceAnalysis.cpp
--- a/Engine/Source/Developer/TraceServices/Private/Analyzers/PlatformEventTraceAnalysis.cpp
+++ b/Engine/Source/Developer/TraceServices/Private/Analyzers/PlatformEventTraceAnalysis.cpp
@@ -82,6 +82,8 @@
 
 void FPlatformEventTraceAnalyzer::OnThreadInfo(const FThreadInfo& ThreadInfo)
 {
+	LLM_SCOPE_BYNAME(TEXT("Insights/FPlatformEventTraceAnalyzer"));
+
 	FAnalysisSessionEditScope _(Session);
 	ContextSwitchesProvider.AddThreadInfo(ThreadInfo.GetId(), ThreadInfo.GetSystemId());
 }
diff --git a/Engine/Source/Developer/TraceServices/Private/ModuleService.cpp b/Engine/Source/Developer/TraceServices/Private/ModuleService.cpp
--- a/Engine/Source/Developer/TraceServices/Private/ModuleService.cpp
+++ b/Engine/Source/Developer/TraceServices/Private/ModuleService.cpp
@@ -1,10 +1,14 @@
 // Copyright Epic Games, Inc. All Rights Reserved.
 
 #include "TraceServices/ModuleService.h"
+
+#include "Features/IModularFeatures.h"
+#include "HAL/LowLevelMemTracker.h"
 #include "ModuleServicePrivate.h"
 #include "UObject/NameTypes.h"
-#include "Features/IModularFeatures.h"
 
+LLM_DECLARE_TAG(Insights_TraceServices);
+
 namespace TraceServices
 {
 
@@ -21,6 +25,7 @@
 	{
 		return;
 	}
+
 	TArray<IModule*> Modules = IModularFeatures::Get().GetModularFeatureImplementations<IModule>(ModuleFeatureName);
 	for (IModule* Module : Modules)
 	{
@@ -101,6 +106,7 @@
 
 void FModuleService::OnAnalysisBegin(IAnalysisSession& Session)
 {
+	LLM_SCOPE_BYTAG(Insights_TraceServices);
 	FScopeLock Lock(&CriticalSection);
 	Initialize();
 	for (IModule* Module : EnabledModules)
diff --git a/Engine/Source/Developer/TraceServices/Private/TraceServicesModule.cpp b/Engine/Source/Developer/TraceServices/Private/TraceServicesModule.cpp
--- a/Engine/Source/Developer/TraceServices/Private/TraceServicesModule.cpp
+++ b/Engine/Source/Developer/TraceServices/Private/TraceServicesModule.cpp
@@ -1,23 +1,27 @@
 // Copyright Epic Games, Inc. All Rights Reserved.
 
-#include "Modules/ModuleManager.h"
 #include "TraceServices/ITraceServicesModule.h"
+
 #include "AnalysisServicePrivate.h"
-#include "ModuleServicePrivate.h"
 #include "Features/IModularFeatures.h"
+#include "HAL/LowLevelMemTracker.h"
+#include "ModuleServicePrivate.h"
 #include "Modules/CookProfilerModule.h"
-#include "Modules/TimingProfilerModule.h"
-#include "Modules/LoadTimeProfilerModule.h"
-#include "Modules/StatsModule.h"
-#include "Modules/CsvProfilerModule.h"
 #include "Modules/CountersModule.h"
-#include "Modules/NetProfilerModule.h"
-#include "Modules/MemoryModule.h"
+#include "Modules/CsvProfilerModule.h"
 #include "Modules/DiagnosticsModule.h"
+#include "Modules/LoadTimeProfilerModule.h"
+#include "Modules/MemoryModule.h"
+#include "Modules/ModuleManager.h"
+#include "Modules/NetProfilerModule.h"
 #include "Modules/PlatformEventsModule.h"
+#include "Modules/StatsModule.h"
 #include "Modules/TasksModule.h"
+#include "Modules/TimingProfilerModule.h"
 #include "Stats/StatsTrace.h"
 
+LLM_DEFINE_TAG(Insights_TraceServices, NAME_None, TEXT("Insights"));
+
 class FTraceServicesModule
 	: public ITraceServicesModule
 {
@@ -52,6 +56,8 @@
 	if (!AnalysisService.IsValid())
 	{
 		GetModuleService();
+
+		LLM_SCOPE_BYTAG(Insights_TraceServices);
 		AnalysisService = MakeShared<TraceServices::FAnalysisService>(*ModuleService.Get());
 	}
 	return AnalysisService;
@@ -61,6 +67,7 @@
 {
 	if (!ModuleService.IsValid())
 	{
+		LLM_SCOPE_BYTAG(Insights_TraceServices);
 		ModuleService = MakeShared<TraceServices::FModuleService>();
 	}
 	return ModuleService;
@@ -70,19 +77,31 @@
 {
 	checkf(!AnalysisService.IsValid(), TEXT("A AnalysisService already exists."));
 	GetModuleService();
-	AnalysisService = MakeShared<TraceServices::FAnalysisService>(*ModuleService.Get());
+
+	{
+		LLM_SCOPE_BYTAG(Insights_TraceServices);
+		AnalysisService = MakeShared<TraceServices::FAnalysisService>(*ModuleService.Get());
+	}
+
 	return AnalysisService;
 }
 
 TSharedPtr<TraceServices::IModuleService> FTraceServicesModule::CreateModuleService()
 {
 	checkf(!ModuleService.IsValid(), TEXT("A ModuleService already exists."));
-	ModuleService = MakeShared<TraceServices::FModuleService>();
+
+	{
+		LLM_SCOPE_BYTAG(Insights_TraceServices);
+		ModuleService = MakeShared<TraceServices::FModuleService>();
+	}
+
 	return ModuleService;
 }
 
 void FTraceServicesModule::StartupModule()
 {
+	LLM_SCOPE_BYTAG(Insights_TraceServices);
+
 	IModularFeatures::Get().RegisterModularFeature(TraceServices::ModuleFeatureName, &TimingProfilerModule);
 	IModularFeatures::Get().RegisterModularFeature(TraceServices::ModuleFeatureName, &CsvProfilerModule);
 	IModularFeatures::Get().RegisterModularFeature(TraceServices::ModuleFeatureName, &CountersModule);
@@ -98,6 +117,8 @@
 
 void FTraceServicesModule::ShutdownModule()
 {
+	LLM_SCOPE_BYTAG(Insights_TraceServices);
+
 	IModularFeatures::Get().UnregisterModularFeature(TraceServices::ModuleFeatureName, &CookProfilingModule);
 	IModularFeatures::Get().UnregisterModularFeature(TraceServices::ModuleFeatureName, &TasksModule);
 	IModularFeatures::Get().UnregisterModularFeature(TraceServices::ModuleFeatureName, &MemoryModule);
