[Insights] Add definition provider

Adds a new provider which allows users to register global accessible "definitions". Analyzers can publish definitions using native structs, which consumers can access using unique references as defined in the runtime definition events.

diff --git a/Engine/Source/Developer/TraceServices/Private/Model/DefinitionProvider.cpp b/Engine/Source/Developer/TraceServices/Private/Model/DefinitionProvider.cpp
new file mode 100644
--- /dev/null
+++ b/Engine/Source/Developer/TraceServices/Private/Model/DefinitionProvider.cpp
@@ -0,0 +1,60 @@
+﻿// Copyright Epic Games, Inc. All Rights Reserved.
+#include "DefinitionProvider.h"
+#include "HAL/UnrealMemory.h"
+
+namespace TraceServices
+{
+	
+//////////////////////////////////////////////////////////////////////////////////////////
+FDefinitionProvider::FDefinitionProvider(IAnalysisSession* InSession)
+	: PageRemain(0)
+{
+}
+
+//////////////////////////////////////////////////////////////////////////////////////////
+void FDefinitionProvider::AddEntry(uint64 Hash, void* Ptr)
+{
+	FWriteScopeLock _(DefinitionsLock);
+	Definitions.Add(Hash, Ptr);
+}
+
+//////////////////////////////////////////////////////////////////////////////////////////
+const void* FDefinitionProvider::FindEntry(uint64 Hash) const
+{
+	FReadScopeLock _(DefinitionsLock);
+	return Definitions.FindRef(Hash);
+}
+
+//////////////////////////////////////////////////////////////////////////////////////////
+void* FDefinitionProvider::Allocate(uint32 Size, uint32 Alignment)
+{
+	check(Size);
+	check(Size <= PageSize);
+	uint8* Dest = nullptr;
+	if (PageRemain > Size)
+	{
+		Dest = Pages.Last().Get() + (PageSize - PageRemain);
+	}
+	else
+	{
+		auto& NewPage = Pages.Emplace_GetRef((uint8*)FMemory::MallocZeroed(PageSize));
+		
+		Dest = NewPage.Get();
+		PageRemain = PageSize;
+	}
+	PageRemain -= Size;
+	return Dest;
+}
+
+//////////////////////////////////////////////////////////////////////////////////////////
+IDefinitionProvider* GetDefinitionProvider(IAnalysisSession& Session)
+{
+	return Session.EditProvider<IDefinitionProvider>("DefinitionProvider");
+}
+//////////////////////////////////////////////////////////////////////////////////////////
+const IDefinitionProvider* ReadDefinitionProvider(IAnalysisSession& Session)
+{
+	return Session.ReadProvider<IDefinitionProvider>("DefinitionProvider");
+}
+	
+}
diff --git a/Engine/Source/Developer/TraceServices/Private/Model/DefinitionProvider.h b/Engine/Source/Developer/TraceServices/Private/Model/DefinitionProvider.h
new file mode 100644
--- /dev/null
+++ b/Engine/Source/Developer/TraceServices/Private/Model/DefinitionProvider.h
@@ -0,0 +1,32 @@
+﻿// Copyright Epic Games, Inc. All Rights Reserved.
+#pragma once
+
+#include "Common/PagedArray.h"
+#include "Templates/UniquePtr.h"
+#include "TraceServices/Model/Definitions.h"
+#include "Trace/Analyzer.h"
+#include "HAL/CriticalSection.h"
+#include "Misc/ScopeRWLock.h"
+
+namespace TraceServices
+{
+
+class FDefinitionProvider : public IDefinitionProvider
+{
+public:
+	FDefinitionProvider(IAnalysisSession* InSession);
+
+private:
+	virtual void AddEntry(uint64 Hash, void* Ptr) override;
+	virtual const void* FindEntry(uint64 Hash) const override;
+	virtual void* Allocate(uint32 Size, uint32 Alignment) override;
+
+	TArray<TUniquePtr<uint8>> Pages;
+	static constexpr uint32 PageSize = 1024;
+	uint32 PageRemain;
+
+	mutable FRWLock DefinitionsLock;
+	TMap<uint64, void*> Definitions;
+};
+
+} // namespace TraceServices
diff --git a/Engine/Source/Developer/TraceServices/Public/TraceServices/Model/Definitions.h b/Engine/Source/Developer/TraceServices/Public/TraceServices/Model/Definitions.h
new file mode 100644
--- /dev/null
+++ b/Engine/Source/Developer/TraceServices/Public/TraceServices/Model/Definitions.h
@@ -0,0 +1,82 @@
+﻿// Copyright Epic Games, Inc. All Rights Reserved.
+#pragma once
+
+#include "TraceServices/Model/AnalysisSession.h"
+#include "Trace/Analyzer.h"
+
+namespace TraceServices
+{
+
+/**
+ * Allows users to publish "definitions", structs representing a definition event. This data can then be queried when
+ * resolving a reference field in another event.
+ *
+ * \note If your definition type in turn references other definitions it encouraged to resolve those during analysis to
+ * avoid having to resolve complex chains on every lookup.
+ *
+ * Example usage:
+ * This code creates a definition, initializes the instance, resolves a reference and finally registers it.
+ * \code 
+ *		FMyDefinitionClass* Instance = DefinitionProvider->Create<FMyDefinitionClass>();
+ *		Instance->SomeValue = EventData.GetValue<uint32>("SomeValue");
+ *		const FEventRef Reference = EventData.GetReferenceValue("SomeString");
+ *		Instance->SomeString = DefinitionProvider->Get<FStringDefinition>(Reference)->Display;
+ *		...
+ *		
+ *		const UE::Trace::FEventRef Id(EventData.GetDefinitionId(), EventData.GetTypeInfo().GetId());
+ *		DefinitionProvider->Register<FMyDefinitionClass>(Instance, Id);
+ * \endcode
+ */
+class IDefinitionProvider : public IProvider
+{
+public:
+
+	/**
+	 * Allocates memory for the definition. The memory is guaranteed to be valid during the lifetime of the
+	 * analysis session.
+	 * @return Pointer to zeroed memory enough to hold an instance of T.
+	 */
+	template<typename T>
+	T* Create()
+	{
+		const uint32 Size = sizeof(T);
+		const uint32 Alignment = alignof(T);
+		T* Allocated = (T*) Allocate(Size, Alignment);
+		return Allocated;
+	}
+
+	/**
+	 * Makes the instance of a definition visible to queries.
+	 * @param Instance Pointer to the instance
+	 * @param Id Id that should be used to reference the definition.
+	 */
+	template<typename T, typename DefinitionType>
+	void Register(T* Instance, UE::Trace::TEventRef<DefinitionType> Id)
+	{
+		AddEntry(Id.GetHash(), Instance);
+	}
+
+	/**
+	 * Attempts to retrieve a previously registered instance of a definition using a reference.
+	 * @param Reference Id used to uniquely identify the definition.
+	 * @return Pointer to the definition or null if the id was not found.
+	 */
+	template<typename T, typename DefinitionType>
+	const T* Get(UE::Trace::TEventRef<DefinitionType> Reference) const
+	{
+		const void* Value = FindEntry(Reference.GetHash());
+		return (T*) Value;
+	}
+
+protected:
+	virtual void* Allocate(uint32 Size, uint32 Alignment) = 0;
+	virtual void AddEntry(uint64 Hash, void* Ptr) = 0;
+	virtual const void* FindEntry(uint64 Hash) const = 0;
+};
+
+IDefinitionProvider* GetDefinitionProvider(IAnalysisSession& Session);
+const IDefinitionProvider* ReadDefinitionProvider(IAnalysisSession& Session);
+
+} // namespace TraceServices
+
+
