Nomaster fixes

diff --git a/Engine/Source/Developer/TraceServices/Private/Model/DefinitionProvider.cpp b/Engine/Source/Developer/TraceServices/Private/Model/DefinitionProvider.cpp
--- a/Engine/Source/Developer/TraceServices/Private/Model/DefinitionProvider.cpp
+++ b/Engine/Source/Developer/TraceServices/Private/Model/DefinitionProvider.cpp
@@ -1,6 +1,7 @@
 ﻿// Copyright Epic Games, Inc. All Rights Reserved.
 #include "DefinitionProvider.h"
 #include "HAL/UnrealMemory.h"
+#include "UObject/NameTypes.h"
 
 namespace TraceServices
 {
@@ -52,7 +53,7 @@ IDefinitionProvider* GetDefinitionProvider(IAnalysisSession& Session)
 	return Session.EditProvider<IDefinitionProvider>("DefinitionProvider");
 }
 //////////////////////////////////////////////////////////////////////////////////////////
-const IDefinitionProvider* ReadDefinitionProvider(IAnalysisSession& Session)
+const IDefinitionProvider* ReadDefinitionProvider(const IAnalysisSession& Session)
 {
 	return Session.ReadProvider<IDefinitionProvider>("DefinitionProvider");
 }
diff --git a/Engine/Source/Developer/TraceServices/Private/Model/DefinitionProvider.h b/Engine/Source/Developer/TraceServices/Private/Model/DefinitionProvider.h
--- a/Engine/Source/Developer/TraceServices/Private/Model/DefinitionProvider.h
+++ b/Engine/Source/Developer/TraceServices/Private/Model/DefinitionProvider.h
@@ -2,11 +2,12 @@
 #pragma once
 
 #include "Common/PagedArray.h"
-#include "Templates/UniquePtr.h"
-#include "TraceServices/Model/Definitions.h"
-#include "Trace/Analyzer.h"
+#include "Containers/Map.h"
 #include "HAL/CriticalSection.h"
 #include "Misc/ScopeRWLock.h"
+#include "Templates/UniquePtr.h"
+#include "Trace/Analyzer.h"
+#include "TraceServices/Model/Definitions.h"
 
 namespace TraceServices
 {
diff --git a/Engine/Source/Developer/TraceServices/Public/TraceServices/Model/Definitions.h b/Engine/Source/Developer/TraceServices/Public/TraceServices/Model/Definitions.h
--- a/Engine/Source/Developer/TraceServices/Public/TraceServices/Model/Definitions.h
+++ b/Engine/Source/Developer/TraceServices/Public/TraceServices/Model/Definitions.h
@@ -75,7 +75,7 @@ protected:
 };
 
 IDefinitionProvider* GetDefinitionProvider(IAnalysisSession& Session);
-const IDefinitionProvider* ReadDefinitionProvider(IAnalysisSession& Session);
+const IDefinitionProvider* ReadDefinitionProvider(const IAnalysisSession& Session);
 
 } // namespace TraceServices
 
