[Insights] Memory Insights: Improved performance of sorting/filtering/grouping operations involving "top function" of allocation's callstack. Grouping is improved ~100x (from 208s to 2s); sorting is improved ~6x (from 29s to 4.5s). Callstack filtering (for ignoring functions on top of callstack from memory allocators and containers) is now done by TraceServices (ModuleProvider / symbol resolvers).

diff --git a/Engine/Source/Developer/TraceInsights/Private/Insights/MemoryProfiler/ViewModels/MemAllocTable.cpp b/Engine/Source/Developer/TraceInsights/Private/Insights/MemoryProfiler/ViewModels/MemAllocTable.cpp
--- a/Engine/Source/Developer/TraceInsights/Private/Insights/MemoryProfiler/ViewModels/MemAllocTable.cpp
+++ b/Engine/Source/Developer/TraceInsights/Private/Insights/MemoryProfiler/ViewModels/MemAllocTable.cpp
@@ -1,8 +1,6 @@
 // Copyright Epic Games, Inc. All Rights Reserved.
 
 #include "MemAllocTable.h"
-#include "Containers/StringView.h"
-#include "Internationalization/Regex.h"
 #include "Widgets/SBoxPanel.h"
 #include "Widgets/SToolTip.h"
 #include "Widgets/Text/STextBlock.h"
@@ -785,76 +783,18 @@ void FMemAllocTable::AddDefaultColumns()
 						}
 
 						const TraceServices::FStackFrame* Frame = nullptr;
-						const TraceServices::FStackFrame* FirstUnknownSymbol = nullptr;
 						for (uint32 FrameIndex = 0; FrameIndex < Callstack->Num(); ++FrameIndex)
 						{
 							Frame = Callstack->Frame(FrameIndex);
 							check(Frame != nullptr);
 
-							if (!Frame->Symbol || !Frame->Symbol->Name)
-							{
-								if (FirstUnknownSymbol == nullptr)
-								{
-									FirstUnknownSymbol = Frame;
-								}
-								continue;
-							}
-
-							bool bIgnoreSymbol = false;
-
-							// Ignore symbols by function prefix.
-							FStringView IgnoreSymbolsByFunctionName[] =
-							{
-								TEXTVIEW("FMemory::"),
-								TEXTVIEW("FMallocWrapper::"),
-								TEXTVIEW("FMallocPoisonProxy::"),
-								TEXTVIEW("Malloc"),
-								TEXTVIEW("Realloc"),
-								TEXTVIEW("MemoryTrace_"),
-							};
-							for (uint32 StringIndex = 0; StringIndex < UE_ARRAY_COUNT(IgnoreSymbolsByFunctionName); ++StringIndex)
-							{
-								if (FCString::Strnicmp(Frame->Symbol->Name, IgnoreSymbolsByFunctionName[StringIndex].GetData(), IgnoreSymbolsByFunctionName[StringIndex].Len()) == 0)
-								{
-									bIgnoreSymbol = true;
-									break;
-								}
-							}
-
-#if 1 // TODO: check perf impact
-							if (!bIgnoreSymbol && Frame->Symbol->File)
-							{
-								// Ignore symbols by file, specified as RegexPattern strings.
-								FStringView IgnoreSymbolsByFilePath[] =
-								{
-									TEXTVIEW(".*/Containers/.*"),
-								};
-								for (uint32 StringIndex = 0; StringIndex < UE_ARRAY_COUNT(IgnoreSymbolsByFilePath); ++StringIndex)
-								{
-									const FString Pattern(IgnoreSymbolsByFilePath[StringIndex]);
-									const FRegexPattern RegexPattern(Pattern);
-									FString File(Frame->Symbol->File);
-									File.ReplaceCharInline(TEXT('\\'), TEXT('/'), ESearchCase::CaseSensitive);
-									FRegexMatcher RegexMatcher(RegexPattern, File);
-									if (RegexMatcher.FindNext())
-									{
-										bIgnoreSymbol = true;
-										break;
-									}
-								}
-							}
-#endif
-
-							if (!bIgnoreSymbol)
+							if (Frame->Symbol &&
+								Frame->Symbol->Name &&
+								Frame->Symbol->FilterStatus != TraceServices::EResolvedSymbolFilterStatus::Filtered)
 							{
 								break;
 							}
 						}
-						if (!Frame)
-						{
-							check(FirstUnknownSymbol != nullptr);
-							Frame = FirstUnknownSymbol;
-						}
 
 						TStringBuilder<1024> Str;
 						FormatStackFrame(*Frame, Str, EStackFrameFormatFlags::Module);
diff --git a/Engine/Source/Developer/TraceServices/Private/Model/CallstacksProvider.cpp b/Engine/Source/Developer/TraceServices/Private/Model/CallstacksProvider.cpp
--- a/Engine/Source/Developer/TraceServices/Private/Model/CallstacksProvider.cpp
+++ b/Engine/Source/Developer/TraceServices/Private/Model/CallstacksProvider.cpp
@@ -12,9 +12,9 @@ namespace TraceServices
 {
 
 /////////////////////////////////////////////////////////////////////
-static const FResolvedSymbol GNeverResolveSymbol(ESymbolQueryResult::NotLoaded, nullptr, nullptr, nullptr, 0);
-static const FResolvedSymbol GNotFoundSymbol(ESymbolQueryResult::NotFound, TEXT("Unknown"), nullptr, nullptr, 0);
-static const FResolvedSymbol GNoSymbol(ESymbolQueryResult::NotFound, TEXT("No callstack recorded"), nullptr, nullptr, 0);
+static const FResolvedSymbol GNeverResolveSymbol(ESymbolQueryResult::NotLoaded, nullptr, nullptr, nullptr, 0, EResolvedSymbolFilterStatus::NotFiltered);
+static const FResolvedSymbol GNotFoundSymbol(ESymbolQueryResult::NotFound, TEXT("Unknown"), nullptr, nullptr, 0, EResolvedSymbolFilterStatus::NotFiltered);
+static const FResolvedSymbol GNoSymbol(ESymbolQueryResult::NotFound, TEXT("No callstack recorded"), nullptr, nullptr, 0, EResolvedSymbolFilterStatus::NotFiltered);
 static constexpr FStackFrame GNotFoundStackFrame = { 0, &GNotFoundSymbol };
 static constexpr FStackFrame GNoStackFrame = { 0, &GNoSymbol };
 static const FCallstack GNotFoundCallstack(&GNotFoundStackFrame, 1);
diff --git a/Engine/Source/Developer/TraceServices/Private/Model/DbgHelpResolver.cpp b/Engine/Source/Developer/TraceServices/Private/Model/DbgHelpResolver.cpp
--- a/Engine/Source/Developer/TraceServices/Private/Model/DbgHelpResolver.cpp
+++ b/Engine/Source/Developer/TraceServices/Private/Model/DbgHelpResolver.cpp
@@ -35,10 +35,11 @@ static const TCHAR* GUnknownModuleTextDbgHelp = TEXT("Unknown");
 
 
 /////////////////////////////////////////////////////////////////////
-FDbgHelpResolver::FDbgHelpResolver(IAnalysisSession& InSession)
+FDbgHelpResolver::FDbgHelpResolver(IAnalysisSession& InSession, IResolvedSymbolFilter& InSymbolFilter)
 	: bRunWorkerThread(false)
 	, bDrainThenStop(false)
 	, Session(InSession)
+	, SymbolFilter(InSymbolFilter)
 {
 	// Setup search paths. The SearchPaths array is a priority stack, which
 	// means paths are searched in reversed order.
@@ -54,7 +55,7 @@ FDbgHelpResolver::FDbgHelpResolver(IAnalysisSession& InSession)
 	{
 		GConfig->GetArray(TEXT("Insights.MemoryProfiler"), TEXT("SymbolSearchPaths"), SymbolSearchPaths, SettingsIni);
 	}
-	
+
 	// Paths from environment
 	FString SymbolPathEnvVar =  FPlatformMisc::GetEnvironmentVariable(TEXT("UE_INSIGHTS_SYMBOL_PATH"));
 	UE_LOG(LogDbgHelp, Log, TEXT("UE_INSIGHTS_SYMBOL_PATH: '%s'"), *SymbolPathEnvVar);
@@ -63,7 +64,7 @@ FDbgHelpResolver::FDbgHelpResolver(IAnalysisSession& InSession)
 	{
 		SymbolSearchPaths.Emplace(SymbolPathPart);
 	}
-	
+
 	Start();
 }
 
@@ -81,9 +82,9 @@ FDbgHelpResolver::~FDbgHelpResolver()
 void FDbgHelpResolver::QueueModuleLoad(const uint8* ImageId, uint32 ImageIdSize, FModule* Module)
 {
 	check(Module != nullptr);
-	
+
 	FScopeLock _(&ModulesCs);
-	
+
 	const FStringView ModuleName = FPathViews::GetCleanFilename(Module->FullName);
 
 	// Add module and sort list according to base address
@@ -112,21 +113,21 @@ void FDbgHelpResolver::QueueModuleReload(const FModule* Module, const TCHAR* InP
 		Entry->Module->Status.store(EModuleStatus::Pending);
 		LoadSymbolsQueue.Enqueue(FQueuedModule{Module, Session.StoreString(InPath), TArrayView<const uint8>(Entry->ImageId)});
 	}
-	
+
 	SymbolArray SymbolsToResolve;
 	ResolveOnSuccess(SymbolsToResolve);
 	for(TTuple<uint64, FResolvedSymbol*> Pair : SymbolsToResolve)
 	{
 		QueueSymbolResolve(Pair.Get<0>(), Pair.Get<1>());
 	}
-	
+
 	if (!bRunWorkerThread && Thread)
 	{
 		// Restart the worker thread if it has stopped.
 		Start();
 	}
 }
-	
+
 /////////////////////////////////////////////////////////////////////
 void FDbgHelpResolver::QueueSymbolResolve(uint64 Address, FResolvedSymbol* Symbol)
 {
@@ -155,7 +156,7 @@ void FDbgHelpResolver::EnumerateSymbolSearchPaths(TFunctionRef<void(FStringView
 	FScopeLock _(&SymbolSearchPathsLock);
 	Algo::ForEach(SymbolSearchPaths, Callback);
 }
-	
+
 /////////////////////////////////////////////////////////////////////
 void FDbgHelpResolver::OnAnalysisComplete()
 {
@@ -166,7 +167,7 @@ void FDbgHelpResolver::OnAnalysisComplete()
 
 
 /////////////////////////////////////////////////////////////////////
-bool FDbgHelpResolver::SetupSyms() 
+bool FDbgHelpResolver::SetupSyms()
 {
 	// Create a unique handle
 	static UPTRINT BaseHandle = 0x493;
@@ -182,7 +183,7 @@ bool FDbgHelpResolver::SetupSyms()
 	SymOpts |= SYMOPT_UNDNAME;
 
 	SymSetOptions(SymOpts);
-	
+
 	return SymInitialize((HANDLE)Handle, NULL, FALSE);
 }
 
@@ -199,7 +200,7 @@ void FDbgHelpResolver::FreeSyms() const
 uint32 FDbgHelpResolver::Run()
 {
 	const bool bInitialized = SetupSyms();
-	
+
 	while (bInitialized && bRunWorkerThread)
 	{
 		// Prioritize queued module loads
@@ -218,7 +219,7 @@ uint32 FDbgHelpResolver::Run()
 			FQueuedAddress Item;
 			if (ResolveQueue.Dequeue(Item))
 			{
-				ResolveSymbol(Item.Address, Item.Target);
+				ResolveSymbol(Item.Address, *Item.Target);
 			}
 		}
 
@@ -226,14 +227,14 @@ uint32 FDbgHelpResolver::Run()
 		{
 			bRunWorkerThread = false;
 		}
-		
+
 		// ...and breathe...
 		FPlatformProcess::Sleep(0.2f);
 	}
 
 	// We don't need the syms library anymore
 	FreeSyms();
-	
+
 	return 0;
 }
 
@@ -244,41 +245,44 @@ void FDbgHelpResolver::Start()
 	bRunWorkerThread = true;
 	Thread = FRunnableThread::Create(this, TEXT("DbgHelpWorker"), 0, TPri_Normal);
 }
-	
+
 /////////////////////////////////////////////////////////////////////
 void FDbgHelpResolver::Stop()
 {
-	bRunWorkerThread = false; 
+	bRunWorkerThread = false;
 }
 
 /////////////////////////////////////////////////////////////////////
-void FDbgHelpResolver::UpdateResolvedSymbol(FResolvedSymbol* Symbol, ESymbolQueryResult Result, const TCHAR* Module, const TCHAR* Name, const TCHAR* File, uint16 Line)
+void FDbgHelpResolver::UpdateResolvedSymbol(FResolvedSymbol& Symbol, ESymbolQueryResult Result, const TCHAR* Module, const TCHAR* Name, const TCHAR* File, uint16 Line)
 {
-	Symbol->Module = Module;
-	Symbol->Name = Name;
-	Symbol->File = File;
-	Symbol->Line = Line;
-	Symbol->Result.store(Result, std::memory_order_release);
+	Symbol.Module = Module;
+	Symbol.Name = Name;
+	Symbol.File = File;
+	Symbol.Line = Line;
+	Symbol.Result.store(Result, std::memory_order_release);
 }
 
 /////////////////////////////////////////////////////////////////////
-void FDbgHelpResolver::ResolveSymbol(uint64 Address, FResolvedSymbol* Target)
+void FDbgHelpResolver::ResolveSymbol(uint64 Address, FResolvedSymbol& Target)
 {
-	check(Target);
-
-	if (Target->Result.load() == ESymbolQueryResult::OK)
+	if (Target.Result.load() == ESymbolQueryResult::OK)
 	{
 		return;
 	}
-	
+
 	const FModuleEntry* Entry = GetModuleForAddress(Address);
 	if (!Entry)
 	{
-		UpdateResolvedSymbol(Target, ESymbolQueryResult::NotFound, GUnknownModuleTextDbgHelp,
-		                     GUnknownModuleTextDbgHelp, GUnknownModuleTextDbgHelp, 0);
+		UpdateResolvedSymbol(Target,
+			ESymbolQueryResult::NotFound,
+			GUnknownModuleTextDbgHelp,
+			GUnknownModuleTextDbgHelp,
+			GUnknownModuleTextDbgHelp,
+			0);
+		SymbolFilter.Update(Target);
 		return;
 	}
-	
+
 	const EModuleStatus ModuleStatus = Entry->Module->Status.load();
 	if (ModuleStatus != EModuleStatus::Loaded)
 	{
@@ -286,13 +290,18 @@ void FDbgHelpResolver::ResolveSymbol(uint64 Address, FResolvedSymbol* Target)
 			                                  ? ESymbolQueryResult::Mismatch
 			                                  : ESymbolQueryResult::NotLoaded;
 		++Entry->Module->Stats.Failed;
-		UpdateResolvedSymbol(Target, Result, GUnknownModuleTextDbgHelp,
-		                     GUnknownModuleTextDbgHelp, GUnknownModuleTextDbgHelp, 0);
+		UpdateResolvedSymbol(Target,
+			Result,
+			GUnknownModuleTextDbgHelp,
+			GUnknownModuleTextDbgHelp,
+			GUnknownModuleTextDbgHelp,
+			0);
+		SymbolFilter.Update(Target);
 		return;
 	}
 
 	++Entry->Module->Stats.Discovered;
-	
+
 	uint8 InfoBuffer[sizeof(SYMBOL_INFO) + (MaxNameLen * sizeof(char) + 1)];
 	SYMBOL_INFO* Info = (SYMBOL_INFO*)InfoBuffer;
 	Info->SizeOfStruct = sizeof(SYMBOL_INFO);
@@ -302,13 +311,18 @@ void FDbgHelpResolver::ResolveSymbol(uint64 Address, FResolvedSymbol* Target)
 	if (!SymFromAddr((HANDLE)Handle, Address, NULL, Info))
 	{
 		++Entry->Module->Stats.Failed;
-		UpdateResolvedSymbol(Target, ESymbolQueryResult::NotFound, Entry->Name, GUnknownModuleTextDbgHelp,
-		                     GUnknownModuleTextDbgHelp, 0);
+		UpdateResolvedSymbol(Target,
+			ESymbolQueryResult::NotFound,
+			Entry->Name,
+			GUnknownModuleTextDbgHelp,
+			GUnknownModuleTextDbgHelp,
+			0);
+		SymbolFilter.Update(Target);
 		return;
 	}
 
 	const TCHAR* SymbolNameStr = Session.StoreString(ANSI_TO_TCHAR(Info->Name));
-	
+
 	// Find the source file and line
 	DWORD  dwDisplacement;
 	IMAGEHLP_LINE Line;
@@ -317,16 +331,28 @@ void FDbgHelpResolver::ResolveSymbol(uint64 Address, FResolvedSymbol* Target)
 	if (!SymGetLineFromAddr((HANDLE)Handle, Address, &dwDisplacement, &Line))
 	{
 		++Entry->Module->Stats.Failed;
-		UpdateResolvedSymbol(Target, ESymbolQueryResult::OK, Entry->Name, SymbolNameStr, GUnknownModuleTextDbgHelp, 0);
+		UpdateResolvedSymbol(Target,
+			ESymbolQueryResult::OK,
+			Entry->Name,
+			SymbolNameStr,
+			GUnknownModuleTextDbgHelp,
+			0);
+		SymbolFilter.Update(Target);
 		return;
 	}
-	
+
 	const TCHAR* SymbolFileStr = Session.StoreString(ANSI_TO_TCHAR(Line.FileName));
-	
+
 	++Entry->Module->Stats.Resolved;
-	UpdateResolvedSymbol(Target, ESymbolQueryResult::OK, Entry->Name, SymbolNameStr, SymbolFileStr, Line.LineNumber);
+	UpdateResolvedSymbol(Target,
+		ESymbolQueryResult::OK,
+		Entry->Name,
+		SymbolNameStr,
+		SymbolFileStr,
+		Line.LineNumber);
+	SymbolFilter.Update(Target);
 }
-	
+
 /////////////////////////////////////////////////////////////////////
 void FDbgHelpResolver::LoadModuleSymbols(const FModule* Module, const TCHAR* Path, const TArrayView<const uint8> ImageId)
 {
@@ -349,7 +375,7 @@ void FDbgHelpResolver::LoadModuleSymbols(const FModule* Module, const TCHAR* Pat
 		SymGetSearchPathW((HANDLE) Handle, OutPath, 1024);
 		UE_LOG(LogDbgHelp, Display, TEXT("Search path: %s"), OutPath);
 	}
-	
+
 	// Attempt to load symbols
 	const DWORD64 LoadedBaseAddress = SymLoadModuleEx((HANDLE)Handle, NULL, TCHAR_TO_ANSI(Module->Name), NULL, Base, Size, NULL, 0);
 	const bool bModuleLoaded = Base == LoadedBaseAddress;
diff --git a/Engine/Source/Developer/TraceServices/Private/Model/DbgHelpResolver.h b/Engine/Source/Developer/TraceServices/Private/Model/DbgHelpResolver.h
--- a/Engine/Source/Developer/TraceServices/Private/Model/DbgHelpResolver.h
+++ b/Engine/Source/Developer/TraceServices/Private/Model/DbgHelpResolver.h
@@ -25,7 +25,7 @@ public:
 	typedef TArray<TTuple<uint64, FResolvedSymbol*>> SymbolArray;
 
 	void Start();
-	FDbgHelpResolver(IAnalysisSession& InSession);
+	FDbgHelpResolver(IAnalysisSession& InSession, IResolvedSymbolFilter& InSymbolFilter);
 	~FDbgHelpResolver();
 	void QueueModuleLoad(const uint8* ImageId, uint32 ImageIdSize, FModule* Module);
 	void QueueModuleReload(const FModule* Module, const TCHAR* InPath, TFunction<void(SymbolArray&)> ResolveOnSuccess);
@@ -66,9 +66,9 @@ private:
 	void FreeSyms() const;
 	virtual uint32 Run() override;
 	virtual void Stop() override;
-	static void UpdateResolvedSymbol(FResolvedSymbol* Symbol, ESymbolQueryResult Result, const TCHAR* Module, const TCHAR* Name, const TCHAR* File, uint16 Line);
+	static void UpdateResolvedSymbol(FResolvedSymbol& Symbol, ESymbolQueryResult Result, const TCHAR* Module, const TCHAR* Name, const TCHAR* File, uint16 Line);
 
-	void ResolveSymbol(uint64 Address, FResolvedSymbol* Target);
+	void ResolveSymbol(uint64 Address, FResolvedSymbol& Target);
 	void LoadModuleSymbols(const FModule* Module, const TCHAR* Path, const TArrayView<const uint8> ImageId);
 	const FModuleEntry* GetModuleForAddress(uint64 Address) const;
 
@@ -87,6 +87,7 @@ private:
 	bool bDrainThenStop;
 	UPTRINT Handle;
 	IAnalysisSession& Session;
+	IResolvedSymbolFilter& SymbolFilter;
 	FRunnableThread* Thread;
 };
 
diff --git a/Engine/Source/Developer/TraceServices/Private/Model/ModuleProvider.cpp b/Engine/Source/Developer/TraceServices/Private/Model/ModuleProvider.cpp
--- a/Engine/Source/Developer/TraceServices/Private/Model/ModuleProvider.cpp
+++ b/Engine/Source/Developer/TraceServices/Private/Model/ModuleProvider.cpp
@@ -9,8 +9,10 @@
 #include "Common/CachedStringStore.h"
 #include "Common/Utils.h"
 #include "Containers/Map.h"
+#include "Containers/StringView.h"
 #include "GenericPlatform/GenericPlatformFile.h"
 #include "HAL/PlatformFileManager.h"
+#include "Internationalization/Regex.h"
 #include "Misc/Paths.h"
 #include "Misc/PathViews.h"
 #include "TraceServices/Model/AnalysisCache.h"
@@ -31,6 +33,13 @@
 
 namespace TraceServices {
 
+/////////////////////////////////////////////////////////////////////
+class FResolvedSymbolFilter : public IResolvedSymbolFilter
+{
+public:
+	virtual void Update(FResolvedSymbol& InSymbol) const override;
+};
+
 /////////////////////////////////////////////////////////////////////
 template<typename SymbolResolverType>
 class TModuleProvider : public IModuleAnalysisProvider
@@ -86,6 +95,8 @@ private:
 	TUniquePtr<SymbolResolverType>	Resolver;
 	FGraphEventRef				LoadSymbolsTask;
 	bool						LoadSymbolsAbort = false;
+
+	FResolvedSymbolFilter		SymbolFilter;
 };
 
 /////////////////////////////////////////////////////////////////////
@@ -97,7 +108,7 @@ TModuleProvider<SymbolResolverType>::TModuleProvider(IAnalysisSession& Session)
 	, NumCachedSymbols(0)
 	, Session(Session)
 {
-	Resolver = TUniquePtr<SymbolResolverType>(new SymbolResolverType(Session));
+	Resolver = TUniquePtr<SymbolResolverType>(new SymbolResolverType(Session, SymbolFilter));
 	LoadSymbolsFromCache(Session.GetCache());
 }
 
@@ -140,7 +151,7 @@ const FResolvedSymbol* TModuleProvider<SymbolResolverType>::GetSymbol(uint64 Add
 		}
 
 		// Add a pending entry to our cache.
-		ResolvedSymbol = &SymbolCache.EmplaceBack(ESymbolQueryResult::Pending, nullptr, nullptr, nullptr, 0);
+		ResolvedSymbol = &SymbolCache.EmplaceBack(ESymbolQueryResult::Pending, nullptr, nullptr, nullptr, 0, EResolvedSymbolFilterStatus::Unknown);
 		SymbolCacheLookup.Add(Address, ResolvedSymbol);
 		++SymbolsDiscovered;
 	}
@@ -355,7 +366,8 @@ void TModuleProvider<SymbolResolverType>::LoadSymbolsFromCache(IAnalysisCache& C
 			UE_LOG(LogTraceServices, Warning, TEXT("Found cached symbol (adress %llx) which referenced unknown string."), Symbol.Address);
 			continue;
 		}
-		FResolvedSymbol& Resolved = SymbolCache.EmplaceBack(ESymbolQueryResult::OK, Module, Name, File, Symbol.Line);
+		FResolvedSymbol& Resolved = SymbolCache.EmplaceBack(ESymbolQueryResult::OK, Module, Name, File, Symbol.Line, EResolvedSymbolFilterStatus::Unknown);
+		SymbolFilter.Update(Resolved);
 		SymbolCacheLookup.Add(Symbol.Address, &Resolved);
 	}
 	NumCachedSymbols = SymbolCacheLookup.Num();
@@ -381,6 +393,60 @@ uint32 TModuleProvider<SymbolResolverType>::GetNumCachedSymbolsFromModule(uint64
 	return Count;
 }
 
+/////////////////////////////////////////////////////////////////////
+void FResolvedSymbolFilter::Update(FResolvedSymbol& InSymbol) const
+{
+	bool bIsFiltered = false;
+
+	if (!bIsFiltered && InSymbol.Name)
+	{
+		// Ignore symbols by function name prefix.
+		FStringView IgnoreSymbolsByFunctionName[] =
+		{
+			TEXTVIEW("FMemory::"),
+			TEXTVIEW("FMallocWrapper::"),
+			TEXTVIEW("FMallocPoisonProxy::"),
+			TEXTVIEW("Malloc"),
+			TEXTVIEW("Realloc"),
+			TEXTVIEW("MemoryTrace_"),
+			TEXTVIEW("operator new"),
+		};
+		for (uint32 StringIndex = 0; StringIndex < UE_ARRAY_COUNT(IgnoreSymbolsByFunctionName); ++StringIndex)
+		{
+			if (FCString::Strnicmp(InSymbol.Name, IgnoreSymbolsByFunctionName[StringIndex].GetData(), IgnoreSymbolsByFunctionName[StringIndex].Len()) == 0)
+			{
+				bIsFiltered = true;
+				break;
+			}
+		}
+	}
+
+	if (!bIsFiltered && InSymbol.File)
+	{
+		// Ignore symbols by file path, specified as RegexPattern strings.
+		FStringView IgnoreSymbolsByFilePath[] =
+		{
+			TEXTVIEW(".*/Containers/.*"),
+		};
+		for (uint32 StringIndex = 0; StringIndex < UE_ARRAY_COUNT(IgnoreSymbolsByFilePath); ++StringIndex)
+		{
+			const FString Pattern(IgnoreSymbolsByFilePath[StringIndex]);
+			const FRegexPattern RegexPattern(Pattern);
+			FString File(InSymbol.File);
+			File.ReplaceCharInline(TEXT('\\'), TEXT('/'), ESearchCase::CaseSensitive);
+			FRegexMatcher RegexMatcher(RegexPattern, File);
+			if (RegexMatcher.FindNext())
+			{
+				bIsFiltered = true;
+				break;
+			}
+		}
+	}
+
+	EResolvedSymbolFilterStatus FilterStatus = bIsFiltered ? TraceServices::EResolvedSymbolFilterStatus::Filtered : TraceServices::EResolvedSymbolFilterStatus::NotFiltered;
+	InSymbol.FilterStatus.store(FilterStatus, std::memory_order_release);
+}
+
 /////////////////////////////////////////////////////////////////////
 IModuleAnalysisProvider* CreateModuleProvider(IAnalysisSession& InSession, const FAnsiStringView& InSymbolFormat)
 {
diff --git a/Engine/Source/Developer/TraceServices/Private/Model/SymslibResolver.cpp b/Engine/Source/Developer/TraceServices/Private/Model/SymslibResolver.cpp
--- a/Engine/Source/Developer/TraceServices/Private/Model/SymslibResolver.cpp
+++ b/Engine/Source/Developer/TraceServices/Private/Model/SymslibResolver.cpp
@@ -475,10 +475,11 @@ private:
 };
 
 /////////////////////////////////////////////////////////////////////
-FSymslibResolver::FSymslibResolver(IAnalysisSession& InSession)
+FSymslibResolver::FSymslibResolver(IAnalysisSession& InSession, IResolvedSymbolFilter& InSymbolFilter)
 	: Modules(InSession.GetLinearAllocator(), 128)
 	, CancelTasks(false)
 	, Session(InSession)
+	, SymbolFilter(InSymbolFilter)
 {
 	// Setup search paths. The SearchPaths array is a priority stack, which
 	// means paths are searched in reversed order.
@@ -653,7 +654,7 @@ void FSymslibResolver::QueueModuleReload(const FModule* Module, const TCHAR* InP
 					ResolveOnSuccess(SymbolsToResolve);
 					for (TTuple<uint64, FResolvedSymbol*> Pair : SymbolsToResolve)
 					{
-						ResolveSymbolTracked(Pair.Get<0>(), Pair.Get<1>(), StringAllocator);
+						ResolveSymbolTracked(Pair.Get<0>(), *Pair.Get<1>(), StringAllocator);
 					}
 				}
 
@@ -785,7 +786,7 @@ void FSymslibResolver::ResolveSymbols(TArrayView<FQueuedAddress>& QueuedWork)
 		{
 			break;
 		}
-		ResolveSymbolTracked(ToResolve.Address, ToResolve.Target, StringAllocator);
+		ResolveSymbolTracked(ToResolve.Address, *ToResolve.Target, StringAllocator);
 	}
 	UE_LOG(LogSymslib, VeryVerbose, TEXT("String allocator used: %.02f kb, wasted: %.02f kb using %d blocks"),
 		((StringAllocator.BlockUsed * StringAllocator.BlockSize - StringAllocator.BlockRemaining) * sizeof(TCHAR)) / 1024.0f,
@@ -806,13 +807,13 @@ FSymslibResolver::FModuleEntry* FSymslibResolver::GetModuleForAddress(uint64 Add
 	return SortedModules[EntryIdx];
 }
 
-void FSymslibResolver::UpdateResolvedSymbol(FResolvedSymbol* Symbol, ESymbolQueryResult Result, const TCHAR* Module, const TCHAR* Name, const TCHAR* File, uint16 Line)
+void FSymslibResolver::UpdateResolvedSymbol(FResolvedSymbol& Symbol, ESymbolQueryResult Result, const TCHAR* Module, const TCHAR* Name, const TCHAR* File, uint16 Line)
 {
-	Symbol->Module = Module;
-	Symbol->Name = Name;
-	Symbol->File = File;
-	Symbol->Line = Line;
-	Symbol->Result.store(Result, std::memory_order_release);
+	Symbol.Module = Module;
+	Symbol.Name = Name;
+	Symbol.File = File;
+	Symbol.Line = Line;
+	Symbol.Result.store(Result, std::memory_order_release);
 }
 
 void FSymslibResolver::LoadModuleTracked(FModuleEntry* Entry, FStringView OverrideSearchPath)
@@ -1031,9 +1032,9 @@ EModuleStatus FSymslibResolver::LoadModule(FModuleEntry* Entry, FStringView Sear
 	return EModuleStatus::Loaded;
 }
 
-void FSymslibResolver::ResolveSymbolTracked(uint64 Address, FResolvedSymbol* Target, FSymbolStringAllocator& StringAllocator)
+void FSymslibResolver::ResolveSymbolTracked(uint64 Address, FResolvedSymbol& Target, FSymbolStringAllocator& StringAllocator)
 {
-	const ESymbolQueryResult PreviousResult = Target->Result.load();
+	const ESymbolQueryResult PreviousResult = Target.Result.load();
 	if (PreviousResult == ESymbolQueryResult::OK)
 	{
 		return;
@@ -1043,7 +1044,13 @@ void FSymslibResolver::ResolveSymbolTracked(uint64 Address, FResolvedSymbol* Tar
 	if (!Entry)
 	{
 		UE_LOG(LogSymslib, Warning, TEXT("No module mapped to address 0x%016llx."), Address);
-		UpdateResolvedSymbol(Target, ESymbolQueryResult::NotLoaded, GUnknownModuleTextSymsLib, GUnknownModuleTextSymsLib, GUnknownModuleTextSymsLib, 0);
+		UpdateResolvedSymbol(Target,
+			ESymbolQueryResult::NotLoaded,
+			GUnknownModuleTextSymsLib,
+			GUnknownModuleTextSymsLib,
+			GUnknownModuleTextSymsLib,
+			0);
+		SymbolFilter.Update(Target);
 		return;
 	}
 
@@ -1059,7 +1066,7 @@ void FSymslibResolver::ResolveSymbolTracked(uint64 Address, FResolvedSymbol* Tar
 	}
 }
 
-bool FSymslibResolver::ResolveSymbol(uint64 Address, FResolvedSymbol* Target, FSymbolStringAllocator& StringAllocator, FModuleEntry* Entry) const
+bool FSymslibResolver::ResolveSymbol(uint64 Address, FResolvedSymbol& Target, FSymbolStringAllocator& StringAllocator, FModuleEntry* Entry) const
 {
 	EModuleStatus Status = Entry->Module->Status.load();
 	while ((Status == EModuleStatus::Pending || Status == EModuleStatus::Discovered) && !CancelTasks.load())
@@ -1071,11 +1078,25 @@ bool FSymslibResolver::ResolveSymbol(uint64 Address, FResolvedSymbol* Target, FS
 	switch (Status)
 	{
 	case EModuleStatus::Failed:
-		UpdateResolvedSymbol(Target, ESymbolQueryResult::NotLoaded, Entry->Module->Name, GUnknownModuleTextSymsLib, GUnknownModuleTextSymsLib, 0);
+		UpdateResolvedSymbol(Target,
+			ESymbolQueryResult::NotLoaded,
+			Entry->Module->Name,
+			GUnknownModuleTextSymsLib,
+			GUnknownModuleTextSymsLib,
+			0);
+		SymbolFilter.Update(Target);
 		return false;
+
 	case EModuleStatus::VersionMismatch:
-		UpdateResolvedSymbol(Target, ESymbolQueryResult::Mismatch, Entry->Module->Name, GUnknownModuleTextSymsLib, GUnknownModuleTextSymsLib, 0);
+		UpdateResolvedSymbol(Target,
+			ESymbolQueryResult::Mismatch,
+			Entry->Module->Name,
+			GUnknownModuleTextSymsLib,
+			GUnknownModuleTextSymsLib,
+			0);
+		SymbolFilter.Update(Target);
 		return false;
+
 	default:
 		break;
 	}
@@ -1136,7 +1157,13 @@ bool FSymslibResolver::ResolveSymbol(uint64 Address, FResolvedSymbol* Target, FS
 	// this includes skipping symbols without name (empty string)
 	if (!SymsSymbol || !SourceFilePersistent || SymsSymbol && SymsSymbol->Name[0] == 0)
 	{
-		UpdateResolvedSymbol(Target, ESymbolQueryResult::NotFound, Entry->Module->Name, GUnknownModuleTextSymsLib, GUnknownModuleTextSymsLib, 0);
+		UpdateResolvedSymbol(Target,
+			ESymbolQueryResult::NotFound,
+			Entry->Module->Name,
+			GUnknownModuleTextSymsLib,
+			GUnknownModuleTextSymsLib,
+			0);
+		SymbolFilter.Update(Target);
 		return false;
 	}
 
@@ -1146,14 +1173,13 @@ bool FSymslibResolver::ResolveSymbol(uint64 Address, FResolvedSymbol* Target, FS
 	const TCHAR* SymbolNamePersistent =  StringAllocator.Store(ANSI_TO_TCHAR(SymbolName));
 
 	// Store the strings and update the target data
-	UpdateResolvedSymbol(
-		Target,
+	UpdateResolvedSymbol(Target,
 		ESymbolQueryResult::OK,
 		Entry->Module->Name,
 		SymbolNamePersistent,
 		SourceFilePersistent,
-		SourceFileLine
-	);
+		SourceFileLine);
+	SymbolFilter.Update(Target);
 
 	return true;
 }
diff --git a/Engine/Source/Developer/TraceServices/Private/Model/SymslibResolver.h b/Engine/Source/Developer/TraceServices/Private/Model/SymslibResolver.h
--- a/Engine/Source/Developer/TraceServices/Private/Model/SymslibResolver.h
+++ b/Engine/Source/Developer/TraceServices/Private/Model/SymslibResolver.h
@@ -24,7 +24,7 @@ class FSymslibResolver
 public:
 	typedef TArray<TTuple<uint64, FResolvedSymbol*>> SymbolArray;
 
-	FSymslibResolver(IAnalysisSession& InSession);
+	FSymslibResolver(IAnalysisSession& InSession, IResolvedSymbolFilter& InSymbolFilter);
 	~FSymslibResolver();
 	void QueueModuleLoad(const uint8* ImageId, uint32 ImageIdSize, FModule* Module);
 	void QueueModuleReload(const FModule* Module, const TCHAR* Path, TFunction<void(SymbolArray&)> ResolveOnSuccess);
@@ -94,11 +94,11 @@ private:
 	void DispatchQueuedAddresses();
 	void ResolveSymbols(TArrayView<FQueuedAddress>& QueuedWork);
 	FModuleEntry* GetModuleForAddress(uint64 Address);
-	static void UpdateResolvedSymbol(FResolvedSymbol* Symbol, ESymbolQueryResult Result, const TCHAR* Module, const TCHAR* Name, const TCHAR* File, uint16 Line);
+	static void UpdateResolvedSymbol(FResolvedSymbol& Symbol, ESymbolQueryResult Result, const TCHAR* Module, const TCHAR* Name, const TCHAR* File, uint16 Line);
 	void LoadModuleTracked(FModuleEntry* Entry, FStringView OverrideSearchPath);
 	EModuleStatus LoadModule(FModuleEntry* Entry, FStringView SearchPathView, FStringBuilderBase& OutStatusMessage) const;
-	void ResolveSymbolTracked(uint64 Address, FResolvedSymbol* Target, class FSymbolStringAllocator& StringAllocator);
-	bool ResolveSymbol(uint64 Address, FResolvedSymbol* Target, FSymbolStringAllocator& StringAllocator, FModuleEntry* Entry) const;
+	void ResolveSymbolTracked(uint64 Address, FResolvedSymbol& Target, class FSymbolStringAllocator& StringAllocator);
+	bool ResolveSymbol(uint64 Address, FResolvedSymbol& Target, FSymbolStringAllocator& StringAllocator, FModuleEntry* Entry) const;
 	void WaitForTasks();
 
 	mutable FRWLock ModulesLock;
@@ -121,6 +121,7 @@ private:
 	std::atomic<uint64> SymbolBytesWasted;
 
 	IAnalysisSession& Session;
+	IResolvedSymbolFilter& SymbolFilter;
 	FString Platform;
 	FString Project;
 };
diff --git a/Engine/Source/Developer/TraceServices/Public/TraceServices/Model/Modules.h b/Engine/Source/Developer/TraceServices/Public/TraceServices/Model/Modules.h
--- a/Engine/Source/Developer/TraceServices/Public/TraceServices/Model/Modules.h
+++ b/Engine/Source/Developer/TraceServices/Public/TraceServices/Model/Modules.h
@@ -41,6 +41,15 @@ inline const TCHAR* QueryResultToString(ESymbolQueryResult Result)
 	return DisplayStrings[(uint8)Result];
 }
 
+////////////////////////////////////////////////////////////////////////////////
+
+enum class EResolvedSymbolFilterStatus : uint8
+{
+	Unknown = 0,
+	Filtered,
+	NotFiltered,
+};
+
 ////////////////////////////////////////////////////////////////////////////////
 /**
   * Represent a resolved symbol. The resolve status and string values may change
@@ -49,26 +58,35 @@ inline const TCHAR* QueryResultToString(ESymbolQueryResult Result)
   */
 struct FResolvedSymbol
 {
-	std::atomic<ESymbolQueryResult> Result;
 	const TCHAR* Module;
 	const TCHAR* Name;
 	const TCHAR* File;
 	uint16 Line;
+	std::atomic<ESymbolQueryResult> Result;
+	std::atomic<EResolvedSymbolFilterStatus> FilterStatus;
 
 	inline ESymbolQueryResult GetResult() const
 	{
 		return Result.load(std::memory_order_acquire);
 	}
 
-	FResolvedSymbol(ESymbolQueryResult InResult, const TCHAR* InModule, const TCHAR* InName, const TCHAR* InFile, uint16 InLine)
-		: Result(InResult)
-		, Module(InModule)
+	FResolvedSymbol(ESymbolQueryResult InResult, const TCHAR* InModule, const TCHAR* InName, const TCHAR* InFile, uint16 InLine, EResolvedSymbolFilterStatus InFilterStatus)
+		: Module(InModule)
 		, Name(InName)
 		, File(InFile)
 		, Line(InLine)
+		, Result(InResult)
+		, FilterStatus(InFilterStatus)
 	{}
 };
 
+////////////////////////////////////////////////////////////////////////////////
+class IResolvedSymbolFilter
+{
+public:
+	virtual void Update(FResolvedSymbol& InSymbol) const = 0;
+};
+
 ////////////////////////////////////////////////////////////////////////////////
 enum class EModuleStatus
 {
