[Insights] CpuProfiler: Fixed dispatch of timing events with metadata.
Bug was caused by incorrect merge of "pending events" (with metadata) with batched cpu timing events (normal events). At the end of processing of a batch of normal cpu timing events, all "pending events" (with metadata) were dispatched and that was a problem for the next batch of timing events (if it has timestamps older than the added pending events). Those new cpu events were detected as being "in the past", treated as relative timestamps and their values added with previous timestamp --> ends up doubling the timestamp values --> huge session duration and incorrect (very large time values) for cpu timing events after that point in time.
It was fixed by delaying the dispatch of remaining pending events to be merged with the next batch of normal cpu timing events.

diff --git a/Engine/Source/Developer/TraceServices/Private/Analyzers/CpuProfilerTraceAnalysis.cpp b/Engine/Source/Developer/TraceServices/Private/Analyzers/CpuProfilerTraceAnalysis.cpp
--- a/Engine/Source/Developer/TraceServices/Private/Analyzers/CpuProfilerTraceAnalysis.cpp
+++ b/Engine/Source/Developer/TraceServices/Private/Analyzers/CpuProfilerTraceAnalysis.cpp
@@ -164,6 +164,18 @@ bool FCpuProfilerAnalyzer::OnEvent(uint16 RouteId, EStyle Style, const FOnEventC
 			break;
 		}
 
+#if CPUPROFILER_SAFE_DISPATCH_PENDING_EVENTS
+		int32 RemainingPending = ThreadState.PendingEvents.Num();
+		if (RemainingPending > 0)
+		{
+			uint64 LastCycle = ThreadState.LastCycle;
+			const FPendingEvent* PendingCursor = ThreadState.PendingEvents.GetData();
+			DispatchPendingEvents(LastCycle, ~0ull, Context.EventTime, ThreadState, PendingCursor, RemainingPending);
+			check(RemainingPending == 0);
+			ThreadState.PendingEvents.Reset();
+		}
+#endif
+
 		ensure(ThreadState.LastCycle != 0 || ThreadState.ScopeStack.Num() == 0);
 		if (ThreadState.LastCycle != 0)
 		{
@@ -177,6 +189,7 @@ bool FCpuProfilerAnalyzer::OnEvent(uint16 RouteId, EStyle Style, const FOnEventC
 				CPUPROFILER_DEBUG_END_EVENT(Timestamp);
 			}
 		}
+
 		ThreadState.LastCycle = ~0ull;
 		break;
 	}
@@ -253,6 +266,9 @@ uint64 FCpuProfilerAnalyzer::ProcessBuffer(const FEventTime& EventTime, FThreadS
 
 	CPUPROFILER_DEBUG_LOGF(TEXT("[%u] ProcessBuffer %llu (%.9f)\n"), ThreadState.ThreadId, LastCycle, EventTime.AsSeconds(LastCycle));
 
+	check(EventTime.GetTimestamp() == 0);
+	const uint64 BaseCycle = EventTime.AsCycle64();
+
 	int32 RemainingPending = ThreadState.PendingEvents.Num();
 	const FPendingEvent* PendingCursor = ThreadState.PendingEvents.GetData();
 
@@ -271,65 +287,13 @@ uint64 FCpuProfilerAnalyzer::ProcessBuffer(const FEventTime& EventTime, FThreadS
 
 		// If we late connect we will be joining the cycle stream mid-flow and
 		// will have missed out on it's base timestamp. Reconstruct it here.
-		check(EventTime.GetTimestamp() == 0);
-		uint64 BaseCycle = EventTime.AsCycle64();
 		if (ActualCycle < BaseCycle)
 		{
 			ActualCycle += BaseCycle;
 		}
 
-		// Dispatch pending events that are younger than the one we've just decoded
-		for (; RemainingPending > 0; RemainingPending--, PendingCursor++)
-		{
-			bool bEnter = true;
-			uint64 PendingCycle = PendingCursor->Cycle;
-			if (int64(PendingCycle) < 0)
-			{
-				PendingCycle = ~PendingCycle;
-				bEnter = false;
-			}
-
-			if (PendingCycle > ActualCycle)
-			{
-				break;
-			}
-
-#if CPUPROFILER_SAFE_DISPATCH_PENDING_EVENTS
-			if (ensure(PendingCycle >= LastCycle))
-			{
-				// Update LastCycle in order to verify time (of following pending events) increases monotonically.
-				LastCycle = PendingCycle;
-			}
-			else
-			{
-				// Time needs to increase monotonically.
-				// We are not allowing events to "go back in time".
-				PendingCycle = LastCycle;
-			}
-#else
-			if (PendingCycle < LastCycle)
-			{
-				PendingCycle = LastCycle;
-			}
-#endif
-
-			double PendingTime = EventTime.AsSeconds(PendingCycle);
-
-			if (bEnter)
-			{
-				CPUPROFILER_DEBUG_LOGF(TEXT("[%u] <B=%llu (%.9f)\n"), ThreadState.ThreadId, PendingCycle, PendingTime);
-				FTimingProfilerEvent Event;
-				Event.TimerIndex = PendingCursor->TimerId;
-				ThreadState.Timeline->AppendBeginEvent(PendingTime, Event);
-				CPUPROFILER_DEBUG_BEGIN_EVENT(PendingTime, Event);
-			}
-			else
-			{
-				CPUPROFILER_DEBUG_LOGF(TEXT("[%u] <E=%llu (%.9f)\n"), ThreadState.ThreadId, PendingCycle, PendingTime);
-				ThreadState.Timeline->AppendEndEvent(PendingTime);
-				CPUPROFILER_DEBUG_END_EVENT(PendingTime);
-			}
-		}
+		// Dispatch pending events that are younger than the one we've just decoded.
+		DispatchPendingEvents(LastCycle, ActualCycle, EventTime, ThreadState, PendingCursor, RemainingPending);
 
 		double ActualTime = EventTime.AsSeconds(ActualCycle);
 
@@ -366,49 +330,26 @@ uint64 FCpuProfilerAnalyzer::ProcessBuffer(const FEventTime& EventTime, FThreadS
 	}
 	check(BufferPtr == BufferEnd);
 
-	// Dispatch remaining pending events...
-	for (; RemainingPending > 0; RemainingPending--, PendingCursor++)
-	{
-		bool bEnter = true;
-		uint64 PendingCycle = PendingCursor->Cycle;
-		if (int64(PendingCycle) < 0)
-		{
-			PendingCycle = ~PendingCycle;
-			bEnter = false;
-		}
-
 #if CPUPROFILER_SAFE_DISPATCH_PENDING_EVENTS
-		if (ensure(PendingCycle >= LastCycle))
-		{
-			// Update LastCycle in order to verify time (of following pending events) increases monotonically.
-			LastCycle = PendingCycle;
-		}
-		else
-		{
-			// Time needs to increase monotonically.
-			// We are not allowing events to "go back in time".
-			PendingCycle = LastCycle;
-		}
-#endif
-
-		double PendingTime = EventTime.AsSeconds(PendingCycle);
-
-		if (bEnter)
-		{
-			CPUPROFILER_DEBUG_LOGF(TEXT("[%u] >B=%llu (%.9f)\n"), ThreadState.ThreadId, PendingCycle, PendingTime);
-			FTimingProfilerEvent Event;
-			Event.TimerIndex = PendingCursor->TimerId;
-			ThreadState.Timeline->AppendBeginEvent(PendingTime, Event);
-			CPUPROFILER_DEBUG_BEGIN_EVENT(PendingTime, Event);
-		}
-		else
-		{
-			CPUPROFILER_DEBUG_LOGF(TEXT("[%u] >E=%llu (%.9f)\n"), ThreadState.ThreadId, PendingCycle, PendingTime);
-			ThreadState.Timeline->AppendEndEvent(PendingTime);
-			CPUPROFILER_DEBUG_END_EVENT(PendingTime);
-		}
+	if (RemainingPending == 0)
+	{
+		//CPUPROFILER_DEBUG_LOGF(TEXT("[%u] MetaEvents: %d added\n"), ThreadState.ThreadId, ThreadState.PendingEvents.Num());
+		ThreadState.PendingEvents.Reset();
 	}
-	ThreadState.PendingEvents.Reset();
+	else
+	{
+		const int32 NumEventsToRemove = ThreadState.PendingEvents.Num() - RemainingPending;
+		CPUPROFILER_DEBUG_LOGF(TEXT("[%u] MetaEvents: %d added, %d still pending\n"), ThreadState.ThreadId, NumEventsToRemove, RemainingPending);
+		ThreadState.PendingEvents.RemoveAt(0, NumEventsToRemove);
+	}
+#else
+	if (RemainingPending > 0)
+	{
+		DispatchPendingEvents(LastCycle, ~0ull, EventTime, ThreadState, PendingCursor, RemainingPending);
+		check(RemainingPending == 0);
+		ThreadState.PendingEvents.Reset();
+	}
+#endif
 
 	ThreadState.LastCycle = LastCycle;
 	return LastCycle;
@@ -422,6 +363,9 @@ uint64 FCpuProfilerAnalyzer::ProcessBufferV2(const FEventTime& EventTime, FThrea
 
 	CPUPROFILER_DEBUG_LOGF(TEXT("[%u] ProcessBuffer %llu (%.9f)\n"), ThreadState.ThreadId, LastCycle, EventTime.AsSeconds(LastCycle));
 
+	check(EventTime.GetTimestamp() == 0);
+	const uint64 BaseCycle = EventTime.AsCycle64();
+
 	int32 RemainingPending = ThreadState.PendingEvents.Num();
 	const FPendingEvent* PendingCursor = ThreadState.PendingEvents.GetData();
 
@@ -440,65 +384,13 @@ uint64 FCpuProfilerAnalyzer::ProcessBufferV2(const FEventTime& EventTime, FThrea
 
 		// If we late connect we will be joining the cycle stream mid-flow and
 		// will have missed out on it's base timestamp. Reconstruct it here.
-		check(EventTime.GetTimestamp() == 0);
-		uint64 BaseCycle = EventTime.AsCycle64();
 		if (ActualCycle < BaseCycle)
 		{
 			ActualCycle += BaseCycle;
 		}
 
-		// Dispatch pending events that are younger than the one we've just decoded
-		for (; RemainingPending > 0; RemainingPending--, PendingCursor++)
-		{
-			bool bEnter = true;
-			uint64 PendingCycle = PendingCursor->Cycle;
-			if (int64(PendingCycle) < 0)
-			{
-				PendingCycle = ~PendingCycle;
-				bEnter = false;
-			}
-
-			if (PendingCycle > ActualCycle)
-			{
-				break;
-			}
-
-#if CPUPROFILER_SAFE_DISPATCH_PENDING_EVENTS
-			if (ensure(PendingCycle >= LastCycle))
-			{
-				// Update LastCycle in order to verify time (of following pending events) increases monotonically.
-				LastCycle = PendingCycle;
-			}
-			else
-			{
-				// Time needs to increase monotonically.
-				// We are not allowing events to "go back in time".
-				PendingCycle = LastCycle;
-			}
-#else
-			if (PendingCycle < LastCycle)
-			{
-				PendingCycle = LastCycle;
-			}
-#endif
-
-			double PendingTime = EventTime.AsSeconds(PendingCycle);
-
-			if (bEnter)
-			{
-				CPUPROFILER_DEBUG_LOGF(TEXT("[%u] <B=%llu (%.9f)\n"), ThreadState.ThreadId, PendingCycle, PendingTime);
-				FTimingProfilerEvent Event;
-				Event.TimerIndex = PendingCursor->TimerId;
-				ThreadState.Timeline->AppendBeginEvent(PendingTime, Event);
-				CPUPROFILER_DEBUG_BEGIN_EVENT(PendingTime, Event);
-			}
-			else
-			{
-				CPUPROFILER_DEBUG_LOGF(TEXT("[%u] <E=%llu (%.9f)\n"), ThreadState.ThreadId, PendingCycle, PendingTime);
-				ThreadState.Timeline->AppendEndEvent(PendingTime);
-				CPUPROFILER_DEBUG_END_EVENT(PendingTime);
-			}
-		}
+		// Dispatch pending events that are younger than the one we've just decoded.
+		DispatchPendingEvents(LastCycle, ActualCycle, EventTime, ThreadState, PendingCursor, RemainingPending);
 
 		double ActualTime = EventTime.AsSeconds(ActualCycle);
 
@@ -647,7 +539,41 @@ uint64 FCpuProfilerAnalyzer::ProcessBufferV2(const FEventTime& EventTime, FThrea
 	}
 	check(BufferPtr == BufferEnd);
 
-	// Dispatch remaining pending events...
+#if CPUPROFILER_SAFE_DISPATCH_PENDING_EVENTS
+	if (RemainingPending == 0)
+	{
+		//CPUPROFILER_DEBUG_LOGF(TEXT("[%u] MetaEvents: %d added\n"), ThreadState.ThreadId, ThreadState.PendingEvents.Num());
+		ThreadState.PendingEvents.Reset();
+	}
+	else
+	{
+		const int32 NumEventsToRemove = ThreadState.PendingEvents.Num() - RemainingPending;
+		CPUPROFILER_DEBUG_LOGF(TEXT("[%u] MetaEvents: %d added, %d still pending\n"), ThreadState.ThreadId, NumEventsToRemove, RemainingPending);
+		ThreadState.PendingEvents.RemoveAt(0, NumEventsToRemove);
+	}
+#else
+	if (RemainingPending > 0)
+	{
+		DispatchPendingEvents(LastCycle, ~0ull, EventTime, ThreadState, PendingCursor, RemainingPending);
+		check(RemainingPending == 0);
+		ThreadState.PendingEvents.Reset();
+	}
+#endif
+
+	ThreadState.LastCycle = LastCycle;
+	return LastCycle;
+}
+
+////////////////////////////////////////////////////////////////////////////////////////////////////
+
+void FCpuProfilerAnalyzer::DispatchPendingEvents(
+	uint64& LastCycle,
+	uint64 CurrentCycle,
+	const FEventTime& EventTime,
+	FThreadState& ThreadState,
+	const FPendingEvent*& PendingCursor,
+	int32& RemainingPending)
+{
 	for (; RemainingPending > 0; RemainingPending--, PendingCursor++)
 	{
 		bool bEnter = true;
@@ -658,6 +584,11 @@ uint64 FCpuProfilerAnalyzer::ProcessBufferV2(const FEventTime& EventTime, FThrea
 			bEnter = false;
 		}
 
+		if (PendingCycle > CurrentCycle)
+		{
+			break;
+		}
+
 #if CPUPROFILER_SAFE_DISPATCH_PENDING_EVENTS
 		if (ensure(PendingCycle >= LastCycle))
 		{
@@ -670,6 +601,11 @@ uint64 FCpuProfilerAnalyzer::ProcessBufferV2(const FEventTime& EventTime, FThrea
 			// We are not allowing events to "go back in time".
 			PendingCycle = LastCycle;
 		}
+#else
+		if (PendingCycle < LastCycle)
+		{
+			PendingCycle = LastCycle;
+		}
 #endif
 
 		double PendingTime = EventTime.AsSeconds(PendingCycle);
@@ -689,10 +625,6 @@ uint64 FCpuProfilerAnalyzer::ProcessBufferV2(const FEventTime& EventTime, FThrea
 			CPUPROFILER_DEBUG_END_EVENT(PendingTime);
 		}
 	}
-	ThreadState.PendingEvents.Reset();
-
-	ThreadState.LastCycle = LastCycle;
-	return LastCycle;
 }
 
 ////////////////////////////////////////////////////////////////////////////////////////////////////
diff --git a/Engine/Source/Developer/TraceServices/Private/Analyzers/CpuProfilerTraceAnalysis.h b/Engine/Source/Developer/TraceServices/Private/Analyzers/CpuProfilerTraceAnalysis.h
--- a/Engine/Source/Developer/TraceServices/Private/Analyzers/CpuProfilerTraceAnalysis.h
+++ b/Engine/Source/Developer/TraceServices/Private/Analyzers/CpuProfilerTraceAnalysis.h
@@ -53,6 +53,7 @@ private:
 	FThreadState& GetThreadState(uint32 ThreadId);
 	uint64 ProcessBuffer(const FEventTime& EventTime, FThreadState& ThreadState, const uint8* BufferPtr, uint32 BufferSize);
 	uint64 ProcessBufferV2(const FEventTime& EventTime, FThreadState& ThreadState, const uint8* BufferPtr, uint32 BufferSize);
+	void DispatchPendingEvents(uint64& LastCycle, uint64 CurrentCycle, const FEventTime& EventTime, FThreadState& ThreadState, const FPendingEvent*& PendingCursor, int32& RemainingPending);
 
 	enum : uint16
 	{
