[Insights] More improvements to symbol filtering (for detecting "top function") when updates symbols loaded from cache (much faster by caching the regex patterns and by running the update in a parallel for loop).

diff --git a/Engine/Source/Developer/TraceServices/Private/Model/ModuleProvider.cpp b/Engine/Source/Developer/TraceServices/Private/Model/ModuleProvider.cpp
--- a/Engine/Source/Developer/TraceServices/Private/Model/ModuleProvider.cpp
+++ b/Engine/Source/Developer/TraceServices/Private/Model/ModuleProvider.cpp
@@ -5,6 +5,7 @@
 
 #include "Algo/Find.h"
 #include "Algo/Transform.h"
+#include "Async/ParallelFor.h"
 #include "Common/CachedPagedArray.h"
 #include "Common/CachedStringStore.h"
 #include "Common/Utils.h"
@@ -37,7 +38,14 @@ namespace TraceServices {
 class FResolvedSymbolFilter : public IResolvedSymbolFilter
 {
 public:
+	FResolvedSymbolFilter();
+	virtual ~FResolvedSymbolFilter();
+
 	virtual void Update(FResolvedSymbol& InSymbol) const override;
+
+private:
+	TArray<FString> IgnoreSymbolsByFunctionName;
+	TArray<FRegexPattern> IgnoreSymbolsByFilePath;
 };
 
 /////////////////////////////////////////////////////////////////////
@@ -367,10 +375,17 @@ void TModuleProvider<SymbolResolverType>::LoadSymbolsFromCache(IAnalysisCache& C
 			continue;
 		}
 		FResolvedSymbol& Resolved = SymbolCache.EmplaceBack(ESymbolQueryResult::OK, Module, Name, File, Symbol.Line, EResolvedSymbolFilterStatus::Unknown);
-		SymbolFilter.Update(Resolved);
 		SymbolCacheLookup.Add(Symbol.Address, &Resolved);
 	}
 	NumCachedSymbols = SymbolCacheLookup.Num();
+
+	// Update filter for all cached symbols.
+	ParallelFor(SymbolCache.Num(), [this](uint32 Index)
+		{
+			SymbolFilter.Update(SymbolCache[Index]);
+		},
+		EParallelForFlags::BackgroundPriority);
+
 	UE_LOG(LogTraceServices, Display, TEXT("Loaded %d symbols from cache."), SymbolCacheLookup.Num());
 }
 
@@ -393,6 +408,26 @@ uint32 TModuleProvider<SymbolResolverType>::GetNumCachedSymbolsFromModule(uint64
 	return Count;
 }
 
+/////////////////////////////////////////////////////////////////////
+FResolvedSymbolFilter::FResolvedSymbolFilter()
+{
+	IgnoreSymbolsByFunctionName.Add(TEXT("FMemory::"));
+	IgnoreSymbolsByFunctionName.Add(TEXT("FMallocWrapper::"));
+	IgnoreSymbolsByFunctionName.Add(TEXT("FMallocPoisonProxy::"));
+	IgnoreSymbolsByFunctionName.Add(TEXT("Malloc"));
+	IgnoreSymbolsByFunctionName.Add(TEXT("Realloc"));
+	IgnoreSymbolsByFunctionName.Add(TEXT("MemoryTrace_"));
+	IgnoreSymbolsByFunctionName.Add(TEXT("operator new"));
+	IgnoreSymbolsByFunctionName.Add(TEXT("std::"));
+
+	IgnoreSymbolsByFilePath.Add(FRegexPattern(FString(TEXT(".*/Containers/.*"))));
+}
+
+/////////////////////////////////////////////////////////////////////
+FResolvedSymbolFilter::~FResolvedSymbolFilter()
+{
+}
+
 /////////////////////////////////////////////////////////////////////
 void FResolvedSymbolFilter::Update(FResolvedSymbol& InSymbol) const
 {
@@ -401,19 +436,9 @@ void FResolvedSymbolFilter::Update(FResolvedSymbol& InSymbol) const
 	if (!bIsFiltered && InSymbol.Name)
 	{
 		// Ignore symbols by function name prefix.
-		FStringView IgnoreSymbolsByFunctionName[] =
+		for (const FString& Prefix : IgnoreSymbolsByFunctionName)
 		{
-			TEXTVIEW("FMemory::"),
-			TEXTVIEW("FMallocWrapper::"),
-			TEXTVIEW("FMallocPoisonProxy::"),
-			TEXTVIEW("Malloc"),
-			TEXTVIEW("Realloc"),
-			TEXTVIEW("MemoryTrace_"),
-			TEXTVIEW("operator new"),
-		};
-		for (uint32 StringIndex = 0; StringIndex < UE_ARRAY_COUNT(IgnoreSymbolsByFunctionName); ++StringIndex)
-		{
-			if (FCString::Strnicmp(InSymbol.Name, IgnoreSymbolsByFunctionName[StringIndex].GetData(), IgnoreSymbolsByFunctionName[StringIndex].Len()) == 0)
+			if (FCString::Strnicmp(InSymbol.Name, *Prefix, Prefix.Len()) == 0)
 			{
 				bIsFiltered = true;
 				break;
@@ -424,14 +449,8 @@ void FResolvedSymbolFilter::Update(FResolvedSymbol& InSymbol) const
 	if (!bIsFiltered && InSymbol.File)
 	{
 		// Ignore symbols by file path, specified as RegexPattern strings.
-		FStringView IgnoreSymbolsByFilePath[] =
-		{
-			TEXTVIEW(".*/Containers/.*"),
-		};
-		for (uint32 StringIndex = 0; StringIndex < UE_ARRAY_COUNT(IgnoreSymbolsByFilePath); ++StringIndex)
+		for (const FRegexPattern& RegexPattern : IgnoreSymbolsByFilePath)
 		{
-			const FString Pattern(IgnoreSymbolsByFilePath[StringIndex]);
-			const FRegexPattern RegexPattern(Pattern);
 			FString File(InSymbol.File);
 			File.ReplaceCharInline(TEXT('\\'), TEXT('/'), ESearchCase::CaseSensitive);
 			FRegexMatcher RegexMatcher(RegexPattern, File);
