Trace+Insights: Reduce header inclusion in AnalysisService.h

### Features

No functional changes. AnalysisService is only really used in a handful of places, it includes a bunch of model classes, and is generally costly to understand (and probably compile).

### Notes

This change is preliminary for a future change that adds new includes that would be circular without this change as a preliminary change. This change is also preliminary to simplify reviewing the future functional change.

#rnx

diff --git a/Engine/Plugins/Developer/TraceDataFiltering/Source/TraceDataFiltering/Private/STraceDataFilterWidget.h b/Engine/Plugins/Developer/TraceDataFiltering/Source/TraceDataFiltering/Private/STraceDataFilterWidget.h
--- a/Engine/Plugins/Developer/TraceDataFiltering/Source/TraceDataFiltering/Private/STraceDataFilterWidget.h
+++ b/Engine/Plugins/Developer/TraceDataFiltering/Source/TraceDataFiltering/Private/STraceDataFilterWidget.h
@@ -3,7 +3,7 @@
 #pragma once
 
 #include "Templates/SharedPointer.h"
-#include "TraceServices/AnalysisService.h"
+#include "TraceServices/Model/AnalysisSession.h"
 #include "Containers/UnrealString.h"
 #include "Misc/IFilter.h"
 #include "Framework/Views/TreeFilterHandler.h"
diff --git a/Engine/Plugins/Developer/TraceDataFiltering/Source/TraceDataFiltering/Private/TraceDataFilteringModule.cpp b/Engine/Plugins/Developer/TraceDataFiltering/Source/TraceDataFiltering/Private/TraceDataFilteringModule.cpp
--- a/Engine/Plugins/Developer/TraceDataFiltering/Source/TraceDataFiltering/Private/TraceDataFilteringModule.cpp
+++ b/Engine/Plugins/Developer/TraceDataFiltering/Source/TraceDataFiltering/Private/TraceDataFilteringModule.cpp
@@ -17,7 +17,6 @@
 #include "Misc/ConfigContext.h"
 
 #include "TraceServices/ModuleService.h"
-#include "TraceServices/AnalysisService.h"
 #include "TraceServices/ITraceServicesModule.h"
 #include "Insights/IUnrealInsightsModule.h"
 
diff --git a/Engine/Plugins/RenderGraphInsights/Source/RenderGraphInsights/Private/RenderGraphTrackDrawHelper.cpp b/Engine/Plugins/RenderGraphInsights/Source/RenderGraphInsights/Private/RenderGraphTrackDrawHelper.cpp
--- a/Engine/Plugins/RenderGraphInsights/Source/RenderGraphInsights/Private/RenderGraphTrackDrawHelper.cpp
+++ b/Engine/Plugins/RenderGraphInsights/Source/RenderGraphInsights/Private/RenderGraphTrackDrawHelper.cpp
@@ -7,7 +7,6 @@
 #include "Framework/Application/SlateApplication.h"
 #include "Rendering/DrawElements.h"
 #include "Styling/CoreStyle.h"
-#include "TraceServices/AnalysisService.h"
 
 #include "Insights/Common/PaintUtils.h"
 #include "Insights/Common/TimeUtils.h"
diff --git a/Engine/Source/Developer/TraceInsights/Private/Insights/ContextSwitches/ViewModels/ContextSwitchesFilterConverters.cpp b/Engine/Source/Developer/TraceInsights/Private/Insights/ContextSwitches/ViewModels/ContextSwitchesFilterConverters.cpp
--- a/Engine/Source/Developer/TraceInsights/Private/Insights/ContextSwitches/ViewModels/ContextSwitchesFilterConverters.cpp
+++ b/Engine/Source/Developer/TraceInsights/Private/Insights/ContextSwitches/ViewModels/ContextSwitchesFilterConverters.cpp
@@ -3,6 +3,7 @@
 #include "ContextSwitchesFilterConverters.h"
 
 #include "TraceServices/Model/ContextSwitches.h"
+#include "TraceServices/Model/Threads.h"
 
 #include "Insights/InsightsManager.h"
 
diff --git a/Engine/Source/Developer/TraceInsights/Private/Insights/ContextSwitches/ViewModels/ContextSwitchesSharedState.cpp b/Engine/Source/Developer/TraceInsights/Private/Insights/ContextSwitches/ViewModels/ContextSwitchesSharedState.cpp
--- a/Engine/Source/Developer/TraceInsights/Private/Insights/ContextSwitches/ViewModels/ContextSwitchesSharedState.cpp
+++ b/Engine/Source/Developer/TraceInsights/Private/Insights/ContextSwitches/ViewModels/ContextSwitchesSharedState.cpp
@@ -5,6 +5,7 @@
 #include "Framework/Commands/Commands.h"
 #include "Framework/Commands/UICommandList.h"
 #include "Framework/MultiBox/MultiBoxBuilder.h"
+#include "TraceServices/Model/Threads.h"
 #include "TraceServices/Model/ContextSwitches.h"
 
 // Insights
diff --git a/Engine/Source/Developer/TraceInsights/Private/Insights/CookProfiler/CookProfilerManager.cpp b/Engine/Source/Developer/TraceInsights/Private/Insights/CookProfiler/CookProfilerManager.cpp
--- a/Engine/Source/Developer/TraceInsights/Private/Insights/CookProfiler/CookProfilerManager.cpp
+++ b/Engine/Source/Developer/TraceInsights/Private/Insights/CookProfiler/CookProfilerManager.cpp
@@ -5,7 +5,6 @@
 #include "Features/IModularFeatures.h"
 #include "Framework/Docking/TabManager.h"
 #include "Modules/ModuleManager.h"
-#include "TraceServices/AnalysisService.h"
 #include "TraceServices/Model/CookProfilerProvider.h"
 #include "Widgets/Docking/SDockTab.h"
 
diff --git a/Engine/Source/Developer/TraceInsights/Private/Insights/CookProfiler/Widgets/SPackageTableTreeView.cpp b/Engine/Source/Developer/TraceInsights/Private/Insights/CookProfiler/Widgets/SPackageTableTreeView.cpp
--- a/Engine/Source/Developer/TraceInsights/Private/Insights/CookProfiler/Widgets/SPackageTableTreeView.cpp
+++ b/Engine/Source/Developer/TraceInsights/Private/Insights/CookProfiler/Widgets/SPackageTableTreeView.cpp
@@ -4,7 +4,6 @@
 
 #include "Framework/MultiBox/MultiBoxBuilder.h"
 #include "SlateOptMacros.h"
-#include "TraceServices/AnalysisService.h"
 
 // Insights
 #include "Insights/InsightsStyle.h"
diff --git a/Engine/Source/Developer/TraceInsights/Private/Insights/InsightsManager.cpp b/Engine/Source/Developer/TraceInsights/Private/Insights/InsightsManager.cpp
--- a/Engine/Source/Developer/TraceInsights/Private/Insights/InsightsManager.cpp
+++ b/Engine/Source/Developer/TraceInsights/Private/Insights/InsightsManager.cpp
@@ -14,6 +14,7 @@
 #include "Modules/ModuleManager.h"
 #include "Templates/UniquePtr.h"
 #include "Trace/StoreClient.h"
+#include "TraceServices/AnalysisService.h"
 #include "TraceServices/Model/Diagnostics.h"
 #include "TraceServices/Model/NetProfiler.h"
 #include "WorkspaceMenuStructure.h"
diff --git a/Engine/Source/Developer/TraceInsights/Private/Insights/InsightsManager.h b/Engine/Source/Developer/TraceInsights/Private/Insights/InsightsManager.h
--- a/Engine/Source/Developer/TraceInsights/Private/Insights/InsightsManager.h
+++ b/Engine/Source/Developer/TraceInsights/Private/Insights/InsightsManager.h
@@ -6,8 +6,6 @@
 #include "CoreMinimal.h"
 #include "Framework/Commands/UICommandList.h"
 #include "Input/DragAndDrop.h"
-#include "TraceServices/AnalysisService.h"
-#include "TraceServices/ModuleService.h"
 
 // Insights
 #include "Insights/Common/Stopwatch.h"
diff --git a/Engine/Source/Developer/TraceInsights/Private/Insights/LoadingProfiler/LoadingProfilerManager.cpp b/Engine/Source/Developer/TraceInsights/Private/Insights/LoadingProfiler/LoadingProfilerManager.cpp
--- a/Engine/Source/Developer/TraceInsights/Private/Insights/LoadingProfiler/LoadingProfilerManager.cpp
+++ b/Engine/Source/Developer/TraceInsights/Private/Insights/LoadingProfiler/LoadingProfilerManager.cpp
@@ -4,7 +4,6 @@
 
 #include "MessageLog/Public/MessageLogModule.h"
 #include "Modules/ModuleManager.h"
-#include "TraceServices/AnalysisService.h"
 #include "TraceServices/Model/LoadTimeProfiler.h"
 #include "Widgets/Docking/SDockTab.h"
 #include "WorkspaceMenuStructure.h"
diff --git a/Engine/Source/Developer/TraceInsights/Private/Insights/LoadingProfiler/Widgets/SLoadingProfilerWindow.cpp b/Engine/Source/Developer/TraceInsights/Private/Insights/LoadingProfiler/Widgets/SLoadingProfilerWindow.cpp
--- a/Engine/Source/Developer/TraceInsights/Private/Insights/LoadingProfiler/Widgets/SLoadingProfilerWindow.cpp
+++ b/Engine/Source/Developer/TraceInsights/Private/Insights/LoadingProfiler/Widgets/SLoadingProfilerWindow.cpp
@@ -6,6 +6,7 @@
 #include "Framework/Docking/TabManager.h"
 #include "Framework/Docking/WorkspaceItem.h"
 #include "SlateOptMacros.h"
+#include "TraceServices/Model/LoadTimeProfiler.h"
 #include "Widgets/Docking/SDockTab.h"
 
 // Insights
diff --git a/Engine/Source/Developer/TraceInsights/Private/Insights/MemoryProfiler/MemoryProfilerManager.cpp b/Engine/Source/Developer/TraceInsights/Private/Insights/MemoryProfiler/MemoryProfilerManager.cpp
--- a/Engine/Source/Developer/TraceInsights/Private/Insights/MemoryProfiler/MemoryProfilerManager.cpp
+++ b/Engine/Source/Developer/TraceInsights/Private/Insights/MemoryProfiler/MemoryProfilerManager.cpp
@@ -4,7 +4,6 @@
 
 #include "MessageLog/Public/MessageLogModule.h"
 #include "Modules/ModuleManager.h"
-#include "TraceServices/AnalysisService.h"
 #include "TraceServices/Model/AllocationsProvider.h"
 #include "TraceServices/Model/Memory.h"
 #include "Widgets/Docking/SDockTab.h"
diff --git a/Engine/Source/Developer/TraceInsights/Private/Insights/MemoryProfiler/ViewModels/MemoryGraphTrack.cpp b/Engine/Source/Developer/TraceInsights/Private/Insights/MemoryProfiler/ViewModels/MemoryGraphTrack.cpp
--- a/Engine/Source/Developer/TraceInsights/Private/Insights/MemoryProfiler/ViewModels/MemoryGraphTrack.cpp
+++ b/Engine/Source/Developer/TraceInsights/Private/Insights/MemoryProfiler/ViewModels/MemoryGraphTrack.cpp
@@ -5,7 +5,6 @@
 #include "Fonts/FontMeasure.h"
 #include "Fonts/SlateFontInfo.h"
 #include "Styling/SlateBrush.h"
-#include "TraceServices/AnalysisService.h"
 #include "TraceServices/Model/AllocationsProvider.h"
 #include "TraceServices/Model/TimingProfiler.h"
 #include "TraceServices/Model/Counters.h"
diff --git a/Engine/Source/Developer/TraceInsights/Private/Insights/MemoryProfiler/ViewModels/MemorySharedState.cpp b/Engine/Source/Developer/TraceInsights/Private/Insights/MemoryProfiler/ViewModels/MemorySharedState.cpp
--- a/Engine/Source/Developer/TraceInsights/Private/Insights/MemoryProfiler/ViewModels/MemorySharedState.cpp
+++ b/Engine/Source/Developer/TraceInsights/Private/Insights/MemoryProfiler/ViewModels/MemorySharedState.cpp
@@ -5,7 +5,6 @@
 #include "Framework/Commands/UICommandList.h"
 #include "Framework/MultiBox/MultiBoxBuilder.h"
 #include "Logging/MessageLog.h"
-#include "TraceServices/AnalysisService.h"
 #include "TraceServices/Model/Memory.h"
 
 // Insights
diff --git a/Engine/Source/Developer/TraceInsights/Private/Insights/MemoryProfiler/ViewModels/MemoryTag.cpp b/Engine/Source/Developer/TraceInsights/Private/Insights/MemoryProfiler/ViewModels/MemoryTag.cpp
--- a/Engine/Source/Developer/TraceInsights/Private/Insights/MemoryProfiler/ViewModels/MemoryTag.cpp
+++ b/Engine/Source/Developer/TraceInsights/Private/Insights/MemoryProfiler/ViewModels/MemoryTag.cpp
@@ -1,7 +1,6 @@
 // Copyright Epic Games, Inc. All Rights Reserved.
 
 #include "MemoryTag.h"
-#include "TraceServices/AnalysisService.h"
 #include "TraceServices/Model/Memory.h"
 
 // Insights
diff --git a/Engine/Source/Developer/TraceInsights/Private/Insights/MemoryProfiler/Widgets/SMemAllocTableTreeView.cpp b/Engine/Source/Developer/TraceInsights/Private/Insights/MemoryProfiler/Widgets/SMemAllocTableTreeView.cpp
--- a/Engine/Source/Developer/TraceInsights/Private/Insights/MemoryProfiler/Widgets/SMemAllocTableTreeView.cpp
+++ b/Engine/Source/Developer/TraceInsights/Private/Insights/MemoryProfiler/Widgets/SMemAllocTableTreeView.cpp
@@ -9,7 +9,6 @@
 #include "Modules/ModuleManager.h"
 #include "SlateOptMacros.h"
 #include "Styling/AppStyle.h"
-#include "TraceServices/AnalysisService.h"
 #include "TraceServices/Model/AllocationsProvider.h"
 #include "TraceServices/Model/Callstack.h"
 #include "TraceServices/Model/Modules.h"
diff --git a/Engine/Source/Developer/TraceInsights/Private/Insights/NetworkingProfiler/NetworkingProfilerManager.cpp b/Engine/Source/Developer/TraceInsights/Private/Insights/NetworkingProfiler/NetworkingProfilerManager.cpp
--- a/Engine/Source/Developer/TraceInsights/Private/Insights/NetworkingProfiler/NetworkingProfilerManager.cpp
+++ b/Engine/Source/Developer/TraceInsights/Private/Insights/NetworkingProfiler/NetworkingProfilerManager.cpp
@@ -4,7 +4,6 @@
 
 #include "MessageLog/Public/MessageLogModule.h"
 #include "Modules/ModuleManager.h"
-#include "TraceServices/AnalysisService.h"
 #include "Widgets/Docking/SDockTab.h"
 #include "WorkspaceMenuStructure.h"
 #include "WorkspaceMenuStructureModule.h"
diff --git a/Engine/Source/Developer/TraceInsights/Private/Insights/NetworkingProfiler/ViewModels/PacketContentViewDrawHelper.cpp b/Engine/Source/Developer/TraceInsights/Private/Insights/NetworkingProfiler/ViewModels/PacketContentViewDrawHelper.cpp
--- a/Engine/Source/Developer/TraceInsights/Private/Insights/NetworkingProfiler/ViewModels/PacketContentViewDrawHelper.cpp
+++ b/Engine/Source/Developer/TraceInsights/Private/Insights/NetworkingProfiler/ViewModels/PacketContentViewDrawHelper.cpp
@@ -6,7 +6,6 @@
 #include "Framework/Application/SlateApplication.h"
 #include "Rendering/DrawElements.h"
 #include "Styling/AppStyle.h"
-#include "TraceServices/AnalysisService.h"
 
 // Insights
 #include "Insights/Common/PaintUtils.h"
diff --git a/Engine/Source/Developer/TraceInsights/Private/Insights/NetworkingProfiler/ViewModels/PacketViewDrawHelper.cpp b/Engine/Source/Developer/TraceInsights/Private/Insights/NetworkingProfiler/ViewModels/PacketViewDrawHelper.cpp
--- a/Engine/Source/Developer/TraceInsights/Private/Insights/NetworkingProfiler/ViewModels/PacketViewDrawHelper.cpp
+++ b/Engine/Source/Developer/TraceInsights/Private/Insights/NetworkingProfiler/ViewModels/PacketViewDrawHelper.cpp
@@ -4,7 +4,6 @@
 
 #include "Framework/Application/SlateApplication.h"
 #include "Rendering/DrawElements.h"
-#include "TraceServices/AnalysisService.h"
 
 // Insights
 #include "Insights/Common/PaintUtils.h"
diff --git a/Engine/Source/Developer/TraceInsights/Private/Insights/NetworkingProfiler/Widgets/SNetStatsCountersView.cpp b/Engine/Source/Developer/TraceInsights/Private/Insights/NetworkingProfiler/Widgets/SNetStatsCountersView.cpp
--- a/Engine/Source/Developer/TraceInsights/Private/Insights/NetworkingProfiler/Widgets/SNetStatsCountersView.cpp
+++ b/Engine/Source/Developer/TraceInsights/Private/Insights/NetworkingProfiler/Widgets/SNetStatsCountersView.cpp
@@ -5,7 +5,6 @@
 #include "Framework/MultiBox/MultiBoxBuilder.h"
 #include "SlateOptMacros.h"
 #include "Templates/UniquePtr.h"
-#include "TraceServices/AnalysisService.h"
 #include "TraceServices/Model/NetProfiler.h"
 #include "Widgets/Input/SCheckBox.h"
 #include "Widgets/Layout/SScrollBox.h"
diff --git a/Engine/Source/Developer/TraceInsights/Private/Insights/NetworkingProfiler/Widgets/SNetStatsView.cpp b/Engine/Source/Developer/TraceInsights/Private/Insights/NetworkingProfiler/Widgets/SNetStatsView.cpp
--- a/Engine/Source/Developer/TraceInsights/Private/Insights/NetworkingProfiler/Widgets/SNetStatsView.cpp
+++ b/Engine/Source/Developer/TraceInsights/Private/Insights/NetworkingProfiler/Widgets/SNetStatsView.cpp
@@ -5,7 +5,6 @@
 #include "Framework/MultiBox/MultiBoxBuilder.h"
 #include "SlateOptMacros.h"
 #include "Templates/UniquePtr.h"
-#include "TraceServices/AnalysisService.h"
 #include "TraceServices/Model/NetProfiler.h"
 #include "Widgets/Input/SCheckBox.h"
 #include "Widgets/Layout/SScrollBox.h"
diff --git a/Engine/Source/Developer/TraceInsights/Private/Insights/NetworkingProfiler/Widgets/SPacketView.cpp b/Engine/Source/Developer/TraceInsights/Private/Insights/NetworkingProfiler/Widgets/SPacketView.cpp
--- a/Engine/Source/Developer/TraceInsights/Private/Insights/NetworkingProfiler/Widgets/SPacketView.cpp
+++ b/Engine/Source/Developer/TraceInsights/Private/Insights/NetworkingProfiler/Widgets/SPacketView.cpp
@@ -8,6 +8,7 @@
 #include "Framework/MultiBox/MultiBoxBuilder.h"
 #include "HAL/PlatformTime.h"
 #include "Rendering/DrawElements.h"
+#include "TraceServices/Model/Frames.h"
 #include "TraceServices/Model/NetProfiler.h"
 #include "Styling/AppStyle.h"
 #include "Widgets/Layout/SScrollBar.h"
diff --git a/Engine/Source/Developer/TraceInsights/Private/Insights/Table/ViewModels/TableTreeNode.h b/Engine/Source/Developer/TraceInsights/Private/Insights/Table/ViewModels/TableTreeNode.h
--- a/Engine/Source/Developer/TraceInsights/Private/Insights/Table/ViewModels/TableTreeNode.h
+++ b/Engine/Source/Developer/TraceInsights/Private/Insights/Table/ViewModels/TableTreeNode.h
@@ -3,7 +3,6 @@
 #pragma once
 
 #include "CoreMinimal.h"
-#include "TraceServices/AnalysisService.h"
 
 // Insights
 #include "Insights/Table/ViewModels/BaseTreeNode.h"
diff --git a/Engine/Source/Developer/TraceInsights/Private/Insights/Table/Widgets/STableTreeView.cpp b/Engine/Source/Developer/TraceInsights/Private/Insights/Table/Widgets/STableTreeView.cpp
--- a/Engine/Source/Developer/TraceInsights/Private/Insights/Table/Widgets/STableTreeView.cpp
+++ b/Engine/Source/Developer/TraceInsights/Private/Insights/Table/Widgets/STableTreeView.cpp
@@ -12,7 +12,6 @@
 #include "Layout/WidgetPath.h"
 #include "Logging/MessageLog.h"
 #include "SlateOptMacros.h"
-#include "TraceServices/AnalysisService.h"
 #include "Widgets/Input/SCheckBox.h"
 #include "Widgets/Layout/SScrollBox.h"
 #include "Widgets/Input/SSearchBox.h"
diff --git a/Engine/Source/Developer/TraceInsights/Private/Insights/Table/Widgets/SUntypedTableTreeView.cpp b/Engine/Source/Developer/TraceInsights/Private/Insights/Table/Widgets/SUntypedTableTreeView.cpp
--- a/Engine/Source/Developer/TraceInsights/Private/Insights/Table/Widgets/SUntypedTableTreeView.cpp
+++ b/Engine/Source/Developer/TraceInsights/Private/Insights/Table/Widgets/SUntypedTableTreeView.cpp
@@ -2,7 +2,8 @@
 
 #include "SUntypedTableTreeView.h"
 
-#include "TraceServices/AnalysisService.h"
+#include "TraceServices/Containers/Tables.h"
+#include "TraceServices/Model/Threads.h"
 
 // Insights
 #include "Insights/Common/Stopwatch.h"
diff --git a/Engine/Source/Developer/TraceInsights/Private/Insights/TaskGraphProfiler/TaskGraphProfilerManager.cpp b/Engine/Source/Developer/TraceInsights/Private/Insights/TaskGraphProfiler/TaskGraphProfilerManager.cpp
--- a/Engine/Source/Developer/TraceInsights/Private/Insights/TaskGraphProfiler/TaskGraphProfilerManager.cpp
+++ b/Engine/Source/Developer/TraceInsights/Private/Insights/TaskGraphProfiler/TaskGraphProfilerManager.cpp
@@ -6,7 +6,6 @@
 #include "Features/IModularFeatures.h"
 #include "Framework/Docking/TabManager.h"
 #include "Modules/ModuleManager.h"
-#include "TraceServices/AnalysisService.h"
 #include "TraceServices/Model/TasksProfiler.h"
 #include "Widgets/Docking/SDockTab.h"
 
diff --git a/Engine/Source/Developer/TraceInsights/Private/Insights/TaskGraphProfiler/Widgets/STaskTableTreeView.cpp b/Engine/Source/Developer/TraceInsights/Private/Insights/TaskGraphProfiler/Widgets/STaskTableTreeView.cpp
--- a/Engine/Source/Developer/TraceInsights/Private/Insights/TaskGraphProfiler/Widgets/STaskTableTreeView.cpp
+++ b/Engine/Source/Developer/TraceInsights/Private/Insights/TaskGraphProfiler/Widgets/STaskTableTreeView.cpp
@@ -9,7 +9,6 @@
 #include "Misc/Paths.h"
 #include "Modules/ModuleManager.h"
 #include "SlateOptMacros.h"
-#include "TraceServices/AnalysisService.h"
 #include "TraceServices/Model/TasksProfiler.h"
 #include "Widgets/Input/SComboBox.h"
 
diff --git a/Engine/Source/Developer/TraceInsights/Private/Insights/Tests/InsightsTestUtils.cpp b/Engine/Source/Developer/TraceInsights/Private/Insights/Tests/InsightsTestUtils.cpp
--- a/Engine/Source/Developer/TraceInsights/Private/Insights/Tests/InsightsTestUtils.cpp
+++ b/Engine/Source/Developer/TraceInsights/Private/Insights/Tests/InsightsTestUtils.cpp
@@ -5,6 +5,7 @@
 #include "Misc/AutomationTest.h"
 #include "Modules/ModuleManager.h"
 #include "TraceServices/AnalysisService.h"
+#include "TraceServices/Model/AnalysisSession.h"
 #include "TraceServices/ITraceServicesModule.h"
 #include "TraceServices/ModuleService.h"
 
diff --git a/Engine/Source/Developer/TraceInsights/Private/Insights/Tests/TimingProfilerTests.cpp b/Engine/Source/Developer/TraceInsights/Private/Insights/Tests/TimingProfilerTests.cpp
--- a/Engine/Source/Developer/TraceInsights/Private/Insights/Tests/TimingProfilerTests.cpp
+++ b/Engine/Source/Developer/TraceInsights/Private/Insights/Tests/TimingProfilerTests.cpp
@@ -2,6 +2,7 @@
 
 #include "Insights/Tests/TimingProfilerTests.h"
 
+#include "TraceServices/Model/Threads.h"
 #include "Misc/AutomationTest.h"
 #include "Misc/FileHelper.h"
 
diff --git a/Engine/Source/Developer/TraceInsights/Private/Insights/TimingProfilerManager.cpp b/Engine/Source/Developer/TraceInsights/Private/Insights/TimingProfilerManager.cpp
--- a/Engine/Source/Developer/TraceInsights/Private/Insights/TimingProfilerManager.cpp
+++ b/Engine/Source/Developer/TraceInsights/Private/Insights/TimingProfilerManager.cpp
@@ -4,7 +4,6 @@
 
 #include "MessageLog/Public/MessageLogModule.h"
 #include "Modules/ModuleManager.h"
-#include "TraceServices/AnalysisService.h"
 #include "Widgets/Docking/SDockTab.h"
 #include "WorkspaceMenuStructure.h"
 #include "WorkspaceMenuStructureModule.h"
diff --git a/Engine/Source/Developer/TraceInsights/Private/Insights/ViewModels/EventNameFilterValueConverter.cpp b/Engine/Source/Developer/TraceInsights/Private/Insights/ViewModels/EventNameFilterValueConverter.cpp
--- a/Engine/Source/Developer/TraceInsights/Private/Insights/ViewModels/EventNameFilterValueConverter.cpp
+++ b/Engine/Source/Developer/TraceInsights/Private/Insights/ViewModels/EventNameFilterValueConverter.cpp
@@ -1,7 +1,7 @@
 // Copyright Epic Games, Inc. All Rights Reserved.
 
 #include "EventNameFilterValueConverter.h"
-
+#include "TraceServices/Model/TimingProfiler.h"
 #include "Insights/InsightsManager.h"
 
 #define LOCTEXT_NAMESPACE "Insights::FEventNameFilterValueConverter"
diff --git a/Engine/Source/Developer/TraceInsights/Private/Insights/ViewModels/FileActivityTimingTrack.cpp b/Engine/Source/Developer/TraceInsights/Private/Insights/ViewModels/FileActivityTimingTrack.cpp
--- a/Engine/Source/Developer/TraceInsights/Private/Insights/ViewModels/FileActivityTimingTrack.cpp
+++ b/Engine/Source/Developer/TraceInsights/Private/Insights/ViewModels/FileActivityTimingTrack.cpp
@@ -8,7 +8,7 @@
 #include "Framework/Commands/UICommandList.h"
 #include "Framework/MultiBox/MultiBoxBuilder.h"
 #include "Styling/SlateBrush.h"
-#include "TraceServices/AnalysisService.h"
+#include "TraceServices/Model/LoadTimeProfiler.h"
 
 // Insights
 #include "Insights/Common/InsightsMenuBuilder.h"
diff --git a/Engine/Source/Developer/TraceInsights/Private/Insights/ViewModels/FrameTrackHelper.cpp b/Engine/Source/Developer/TraceInsights/Private/Insights/ViewModels/FrameTrackHelper.cpp
--- a/Engine/Source/Developer/TraceInsights/Private/Insights/ViewModels/FrameTrackHelper.cpp
+++ b/Engine/Source/Developer/TraceInsights/Private/Insights/ViewModels/FrameTrackHelper.cpp
@@ -6,7 +6,7 @@
 #include "Fonts/SlateFontInfo.h"
 #include "Framework/Application/SlateApplication.h"
 #include "Rendering/DrawElements.h"
-#include "TraceServices/AnalysisService.h"
+#include "TraceServices/Model/Frames.h"
 
 // Insights
 #include "Insights/Common/PaintUtils.h"
diff --git a/Engine/Source/Developer/TraceInsights/Private/Insights/ViewModels/LoadingTimingTrack.cpp b/Engine/Source/Developer/TraceInsights/Private/Insights/ViewModels/LoadingTimingTrack.cpp
--- a/Engine/Source/Developer/TraceInsights/Private/Insights/ViewModels/LoadingTimingTrack.cpp
+++ b/Engine/Source/Developer/TraceInsights/Private/Insights/ViewModels/LoadingTimingTrack.cpp
@@ -7,7 +7,7 @@
 #include "Framework/Commands/UICommandList.h"
 #include "Framework/MultiBox/MultiBoxBuilder.h"
 #include "Styling/SlateBrush.h"
-#include "TraceServices/AnalysisService.h"
+#include "TraceServices/Model/Threads.h"
 
 // Insights
 #include "Insights/Common/PaintUtils.h"
diff --git a/Engine/Source/Developer/TraceInsights/Private/Insights/ViewModels/LoadingTimingTrack.h b/Engine/Source/Developer/TraceInsights/Private/Insights/ViewModels/LoadingTimingTrack.h
--- a/Engine/Source/Developer/TraceInsights/Private/Insights/ViewModels/LoadingTimingTrack.h
+++ b/Engine/Source/Developer/TraceInsights/Private/Insights/ViewModels/LoadingTimingTrack.h
@@ -4,7 +4,7 @@
 
 #include "CoreMinimal.h"
 #include "Framework/Commands/Commands.h"
-#include "TraceServices/AnalysisService.h"
+#include "TraceServices/Model/LoadTimeProfiler.h"
 
 // Insights
 #include "Insights/ITimingViewExtender.h"
diff --git a/Engine/Source/Developer/TraceInsights/Private/Insights/ViewModels/LogMessage.cpp b/Engine/Source/Developer/TraceInsights/Private/Insights/ViewModels/LogMessage.cpp
--- a/Engine/Source/Developer/TraceInsights/Private/Insights/ViewModels/LogMessage.cpp
+++ b/Engine/Source/Developer/TraceInsights/Private/Insights/ViewModels/LogMessage.cpp
@@ -4,7 +4,7 @@
 
 #include "Misc/OutputDeviceHelper.h"
 #include "Misc/ScopeLock.h"
-#include "TraceServices/AnalysisService.h"
+#include "TraceServices/Model/Log.h"
 
 // Insights
 #include "Insights/Common/TimeUtils.h"
diff --git a/Engine/Source/Developer/TraceInsights/Private/Insights/ViewModels/MarkersTimingTrack.cpp b/Engine/Source/Developer/TraceInsights/Private/Insights/ViewModels/MarkersTimingTrack.cpp
--- a/Engine/Source/Developer/TraceInsights/Private/Insights/ViewModels/MarkersTimingTrack.cpp
+++ b/Engine/Source/Developer/TraceInsights/Private/Insights/ViewModels/MarkersTimingTrack.cpp
@@ -9,7 +9,7 @@
 #include "IImageWrapperModule.h"
 #include "Modules/ModuleManager.h"
 #include "Styling/AppStyle.h"
-#include "TraceServices/AnalysisService.h"
+#include "TraceServices/Model/Log.h"
 #include "TraceServices/Model/Screenshot.h"
 
 // Insights
diff --git a/Engine/Source/Developer/TraceInsights/Private/Insights/ViewModels/ThreadTimingTrack.cpp b/Engine/Source/Developer/TraceInsights/Private/Insights/ViewModels/ThreadTimingTrack.cpp
--- a/Engine/Source/Developer/TraceInsights/Private/Insights/ViewModels/ThreadTimingTrack.cpp
+++ b/Engine/Source/Developer/TraceInsights/Private/Insights/ViewModels/ThreadTimingTrack.cpp
@@ -9,7 +9,9 @@
 #include "HAL/PlatformApplicationMisc.h"
 #include "Serialization/MemoryReader.h"
 #include "Styling/SlateBrush.h"
+#include "TraceServices/Model/LoadTimeProfiler.h"
 #include "TraceServices/Model/TasksProfiler.h"
+#include "TraceServices/Model/Threads.h"
 #include "Async/TaskGraphInterfaces.h"
 
 // Insights
diff --git a/Engine/Source/Developer/TraceInsights/Private/Insights/ViewModels/TimingExporter.cpp b/Engine/Source/Developer/TraceInsights/Private/Insights/ViewModels/TimingExporter.cpp
--- a/Engine/Source/Developer/TraceInsights/Private/Insights/ViewModels/TimingExporter.cpp
+++ b/Engine/Source/Developer/TraceInsights/Private/Insights/ViewModels/TimingExporter.cpp
@@ -3,6 +3,7 @@
 #include "TimingExporter.h"
 #include "HAL/PlatformFileManager.h"
 #include "Logging/MessageLog.h"
+#include "TraceServices/Model/Threads.h"
 
 // Insights
 #include "Insights/Common/Stopwatch.h"
diff --git a/Engine/Source/Developer/TraceInsights/Private/Insights/ViewModels/TimingGraphTrack.cpp b/Engine/Source/Developer/TraceInsights/Private/Insights/ViewModels/TimingGraphTrack.cpp
--- a/Engine/Source/Developer/TraceInsights/Private/Insights/ViewModels/TimingGraphTrack.cpp
+++ b/Engine/Source/Developer/TraceInsights/Private/Insights/ViewModels/TimingGraphTrack.cpp
@@ -4,7 +4,7 @@
 
 #include <limits>
 
-#include "TraceServices/AnalysisService.h"
+#include "TraceServices/Model/Frames.h"
 #include "TraceServices/Model/TimingProfiler.h"
 #include "TraceServices/Model/Counters.h"
 
diff --git a/Engine/Source/Developer/TraceInsights/Private/Insights/ViewModels/TimingViewDrawHelper.cpp b/Engine/Source/Developer/TraceInsights/Private/Insights/ViewModels/TimingViewDrawHelper.cpp
--- a/Engine/Source/Developer/TraceInsights/Private/Insights/ViewModels/TimingViewDrawHelper.cpp
+++ b/Engine/Source/Developer/TraceInsights/Private/Insights/ViewModels/TimingViewDrawHelper.cpp
@@ -6,7 +6,6 @@
 #include "Framework/Application/SlateApplication.h"
 #include "Rendering/DrawElements.h"
 #include "Styling/AppStyle.h"
-#include "TraceServices/AnalysisService.h"
 
 // Insights
 #include "Insights/Common/PaintUtils.h"
diff --git a/Engine/Source/Developer/TraceInsights/Private/Insights/Widgets/SFrameTrack.cpp b/Engine/Source/Developer/TraceInsights/Private/Insights/Widgets/SFrameTrack.cpp
--- a/Engine/Source/Developer/TraceInsights/Private/Insights/Widgets/SFrameTrack.cpp
+++ b/Engine/Source/Developer/TraceInsights/Private/Insights/Widgets/SFrameTrack.cpp
@@ -10,7 +10,7 @@
 #include "Misc/StringBuilder.h"
 #include "Rendering/DrawElements.h"
 #include "Styling/AppStyle.h"
-#include "TraceServices/AnalysisService.h"
+#include "TraceServices/Model/Frames.h"
 #include "Widgets/Layout/SScrollBar.h"
 
 // Insights
diff --git a/Engine/Source/Developer/TraceInsights/Private/Insights/Widgets/SLogView.cpp b/Engine/Source/Developer/TraceInsights/Private/Insights/Widgets/SLogView.cpp
--- a/Engine/Source/Developer/TraceInsights/Private/Insights/Widgets/SLogView.cpp
+++ b/Engine/Source/Developer/TraceInsights/Private/Insights/Widgets/SLogView.cpp
@@ -18,6 +18,7 @@
 #include "SlateOptMacros.h"
 #include "Styling/AppStyle.h"
 #include "Styling/StyleColors.h"
+#include "TraceServices/Model/Log.h"
 #include "Widgets/Images/SImage.h"
 #include "Widgets/Input/SComboButton.h"
 #include "Widgets/Layout/SScrollBox.h"
diff --git a/Engine/Source/Developer/TraceInsights/Private/Insights/Widgets/SStatsView.cpp b/Engine/Source/Developer/TraceInsights/Private/Insights/Widgets/SStatsView.cpp
--- a/Engine/Source/Developer/TraceInsights/Private/Insights/Widgets/SStatsView.cpp
+++ b/Engine/Source/Developer/TraceInsights/Private/Insights/Widgets/SStatsView.cpp
@@ -8,7 +8,7 @@
 #include "HAL/PlatformApplicationMisc.h"
 #include "SlateOptMacros.h"
 #include "Styling/AppStyle.h"
-#include "TraceServices/AnalysisService.h"
+#include "TraceServices/Model/Counters.h"
 #include "Widgets/Input/SCheckBox.h"
 #include "Widgets/Layout/SScrollBox.h"
 #include "Widgets/Input/SSearchBox.h"
diff --git a/Engine/Source/Developer/TraceInsights/Private/Insights/Widgets/STimersView.cpp b/Engine/Source/Developer/TraceInsights/Private/Insights/Widgets/STimersView.cpp
--- a/Engine/Source/Developer/TraceInsights/Private/Insights/Widgets/STimersView.cpp
+++ b/Engine/Source/Developer/TraceInsights/Private/Insights/Widgets/STimersView.cpp
@@ -14,7 +14,7 @@
 #include "Modules/ModuleManager.h"
 #include "SlateOptMacros.h"
 #include "Styling/AppStyle.h"
-#include "TraceServices/AnalysisService.h"
+#include "TraceServices/Model/Threads.h"
 #include "TraceServices/Model/TimingProfiler.h"
 #include "Widgets/Input/SCheckBox.h"
 #include "Widgets/Layout/SScrollBox.h"
diff --git a/Engine/Source/Developer/TraceInsights/Private/Insights/Widgets/STimingView.cpp b/Engine/Source/Developer/TraceInsights/Private/Insights/Widgets/STimingView.cpp
--- a/Engine/Source/Developer/TraceInsights/Private/Insights/Widgets/STimingView.cpp
+++ b/Engine/Source/Developer/TraceInsights/Private/Insights/Widgets/STimingView.cpp
@@ -23,7 +23,6 @@
 #include "Styling/SlateBrush.h"
 #include "Styling/StyleColors.h"
 #include "Styling/ToolBarStyle.h"
-#include "TraceServices/AnalysisService.h"
 #include "Widgets/Docking/SDockTab.h"
 #include "Widgets/Input/SButton.h"
 #include "Widgets/Input/SCheckBox.h"
diff --git a/Engine/Source/Developer/TraceInsights/Private/Insights/Widgets/STimingView.h b/Engine/Source/Developer/TraceInsights/Private/Insights/Widgets/STimingView.h
--- a/Engine/Source/Developer/TraceInsights/Private/Insights/Widgets/STimingView.h
+++ b/Engine/Source/Developer/TraceInsights/Private/Insights/Widgets/STimingView.h
@@ -9,7 +9,6 @@
 #include "Layout/Geometry.h"
 #include "Styling/SlateTypes.h"
 #include "Templates/Function.h"
-#include "TraceServices/AnalysisService.h"
 #include "Widgets/DeclarativeSyntaxSupport.h"
 #include "Widgets/SCompoundWidget.h"
 #include "Widgets/SWidget.h"
diff --git a/Engine/Source/Developer/TraceServices/Private/AnalysisCache.h b/Engine/Source/Developer/TraceServices/Private/AnalysisCache.h
--- a/Engine/Source/Developer/TraceServices/Private/AnalysisCache.h
+++ b/Engine/Source/Developer/TraceServices/Private/AnalysisCache.h
@@ -5,7 +5,6 @@
 #include "IO/IoHash.h"
 #include "Memory/SharedBuffer.h"
 #include "Templates/UniquePtr.h"
-#include "TraceServices/AnalysisService.h"
 #include "TraceServices/Model/AnalysisCache.h"
 
 class IMappedFileRegion;
diff --git a/Engine/Source/Developer/TraceServices/Private/AnalysisServicePrivate.h b/Engine/Source/Developer/TraceServices/Private/AnalysisServicePrivate.h
--- a/Engine/Source/Developer/TraceServices/Private/AnalysisServicePrivate.h
+++ b/Engine/Source/Developer/TraceServices/Private/AnalysisServicePrivate.h
@@ -3,6 +3,7 @@
 #pragma once
 
 #include "AnalysisCache.h"
+#include "TraceServices/Model/AnalysisSession.h"
 #include "TraceServices/AnalysisService.h"
 #include "Containers/HashTable.h"
 #include "Common/StringStore.h"
diff --git a/Engine/Source/Developer/TraceServices/Private/Analyzers/CsvProfilerTraceAnalysis.cpp b/Engine/Source/Developer/TraceServices/Private/Analyzers/CsvProfilerTraceAnalysis.cpp
--- a/Engine/Source/Developer/TraceServices/Private/Analyzers/CsvProfilerTraceAnalysis.cpp
+++ b/Engine/Source/Developer/TraceServices/Private/Analyzers/CsvProfilerTraceAnalysis.cpp
@@ -7,6 +7,7 @@
 #include "HAL/LowLevelMemTracker.h"
 #include "TraceServices/Model/Counters.h"
 #include "TraceServices/Model/Frames.h"
+#include "TraceServices/Model/Threads.h"
 
 namespace TraceServices
 {
diff --git a/Engine/Source/Developer/TraceServices/Private/Model/BookmarksPrivate.h b/Engine/Source/Developer/TraceServices/Private/Model/BookmarksPrivate.h
--- a/Engine/Source/Developer/TraceServices/Private/Model/BookmarksPrivate.h
+++ b/Engine/Source/Developer/TraceServices/Private/Model/BookmarksPrivate.h
@@ -2,7 +2,7 @@
 
 #pragma once
 
-#include "TraceServices/AnalysisService.h"
+#include "TraceServices/Model/Bookmarks.h"
 #include "Templates/SharedPointer.h"
 
 namespace TraceServices
diff --git a/Engine/Source/Developer/TraceServices/Private/Model/CookProfilerProviderPrivate.h b/Engine/Source/Developer/TraceServices/Private/Model/CookProfilerProviderPrivate.h
--- a/Engine/Source/Developer/TraceServices/Private/Model/CookProfilerProviderPrivate.h
+++ b/Engine/Source/Developer/TraceServices/Private/Model/CookProfilerProviderPrivate.h
@@ -3,7 +3,6 @@
 #pragma once
 
 #include "TraceServices/Model/CookProfilerProvider.h"
-#include "TraceServices/AnalysisService.h"
 
 namespace TraceServices
 {
diff --git a/Engine/Source/Developer/TraceServices/Private/Model/DiagnosticsPrivate.h b/Engine/Source/Developer/TraceServices/Private/Model/DiagnosticsPrivate.h
--- a/Engine/Source/Developer/TraceServices/Private/Model/DiagnosticsPrivate.h
+++ b/Engine/Source/Developer/TraceServices/Private/Model/DiagnosticsPrivate.h
@@ -1,7 +1,6 @@
 // Copyright Epic Games, Inc. All Rights Reserved.
 #pragma once
 
-#include "TraceServices/AnalysisService.h"
 #include "TraceServices/Model/Diagnostics.h"
 
 namespace TraceServices
diff --git a/Engine/Source/Developer/TraceServices/Private/Model/FramesPrivate.h b/Engine/Source/Developer/TraceServices/Private/Model/FramesPrivate.h
--- a/Engine/Source/Developer/TraceServices/Private/Model/FramesPrivate.h
+++ b/Engine/Source/Developer/TraceServices/Private/Model/FramesPrivate.h
@@ -2,13 +2,15 @@
 
 #pragma once
 
-#include "TraceServices/AnalysisService.h"
+#include "TraceServices/Model/Frames.h"
 #include "ProfilingDebugging/MiscTrace.h"
 #include "Common/PagedArray.h"
 
 namespace TraceServices
 {
 
+class IAnalysisSession;
+
 class FFrameProvider
 	: public IFrameProvider
 {
diff --git a/Engine/Source/Developer/TraceServices/Private/Model/LoadTimeProfilerPrivate.h b/Engine/Source/Developer/TraceServices/Private/Model/LoadTimeProfilerPrivate.h
--- a/Engine/Source/Developer/TraceServices/Private/Model/LoadTimeProfilerPrivate.h
+++ b/Engine/Source/Developer/TraceServices/Private/Model/LoadTimeProfilerPrivate.h
@@ -2,8 +2,8 @@
 
 #pragma once
 
+#include "TraceServices/Model/Counters.h"
 #include "TraceServices/Model/LoadTimeProfiler.h"
-#include "TraceServices/AnalysisService.h"
 #include "Common/SlabAllocator.h"
 #include "Common/PagedArray.h"
 #include "Common/StringStore.h"
diff --git a/Engine/Source/Developer/TraceServices/Private/Model/LogPrivate.h b/Engine/Source/Developer/TraceServices/Private/Model/LogPrivate.h
--- a/Engine/Source/Developer/TraceServices/Private/Model/LogPrivate.h
+++ b/Engine/Source/Developer/TraceServices/Private/Model/LogPrivate.h
@@ -2,7 +2,7 @@
 
 #pragma once
 
-#include "TraceServices/AnalysisService.h"
+#include "TraceServices/Model/Log.h"
 #include "Templates/SharedPointer.h"
 #include "Common/PagedArray.h"
 #include "Model/Tables.h"
diff --git a/Engine/Source/Developer/TraceServices/Private/Model/MemoryPrivate.h b/Engine/Source/Developer/TraceServices/Private/Model/MemoryPrivate.h
--- a/Engine/Source/Developer/TraceServices/Private/Model/MemoryPrivate.h
+++ b/Engine/Source/Developer/TraceServices/Private/Model/MemoryPrivate.h
@@ -1,7 +1,6 @@
 // Copyright Epic Games, Inc. All Rights Reserved.
 #pragma once
 
-#include "TraceServices/AnalysisService.h"
 #include "TraceServices/Model/Memory.h"
 #include "Common/PagedArray.h"
 
diff --git a/Engine/Source/Developer/TraceServices/Private/Model/ScreenshotProviderPrivate.h b/Engine/Source/Developer/TraceServices/Private/Model/ScreenshotProviderPrivate.h
--- a/Engine/Source/Developer/TraceServices/Private/Model/ScreenshotProviderPrivate.h
+++ b/Engine/Source/Developer/TraceServices/Private/Model/ScreenshotProviderPrivate.h
@@ -2,7 +2,6 @@
 
 #pragma once
 
-#include "TraceServices/AnalysisService.h"
 #include "Templates/SharedPointer.h"
 #include "TraceServices/Model/Screenshot.h"
 
diff --git a/Engine/Source/Developer/TraceServices/Private/Model/Tables.h b/Engine/Source/Developer/TraceServices/Private/Model/Tables.h
--- a/Engine/Source/Developer/TraceServices/Private/Model/Tables.h
+++ b/Engine/Source/Developer/TraceServices/Private/Model/Tables.h
@@ -2,8 +2,7 @@
 
 #pragma once
 
-#include "CoreTypes.h"
-#include "TraceServices/AnalysisService.h"
+#include "TraceServices/Containers/Tables.h"
 #include "Templates/Function.h"
 #include "Common/PagedArray.h"
 #include "Common/SlabAllocator.h"
diff --git a/Engine/Source/Developer/TraceServices/Private/Modules/CsvProfilerModule.cpp b/Engine/Source/Developer/TraceServices/Private/Modules/CsvProfilerModule.cpp
--- a/Engine/Source/Developer/TraceServices/Private/Modules/CsvProfilerModule.cpp
+++ b/Engine/Source/Developer/TraceServices/Private/Modules/CsvProfilerModule.cpp
@@ -2,8 +2,9 @@
 
 #include "CsvProfilerModule.h"
 #include "Analyzers/CsvProfilerTraceAnalysis.h"
-#include "AnalysisServicePrivate.h"
-#include "TraceServices/ModuleService.h"
+#include "TraceServices/Model/Counters.h"
+#include "TraceServices/Model/Frames.h"
+#include "TraceServices/Model/Threads.h"
 
 namespace TraceServices
 {
diff --git a/Engine/Source/Developer/TraceServices/Private/Modules/LoadTimeProfilerModule.cpp b/Engine/Source/Developer/TraceServices/Private/Modules/LoadTimeProfilerModule.cpp
--- a/Engine/Source/Developer/TraceServices/Private/Modules/LoadTimeProfilerModule.cpp
+++ b/Engine/Source/Developer/TraceServices/Private/Modules/LoadTimeProfilerModule.cpp
@@ -5,6 +5,7 @@
 #include "Analyzers/LoadTimeTraceAnalysis.h"
 #include "AnalysisServicePrivate.h"
 #include "Model/FileActivity.h"
+#include "TraceServices/Model/Bookmarks.h"
 #include "HAL/FileManager.h"
 
 namespace TraceServices
diff --git a/Engine/Source/Developer/TraceServices/Private/Modules/StatsModule.cpp b/Engine/Source/Developer/TraceServices/Private/Modules/StatsModule.cpp
--- a/Engine/Source/Developer/TraceServices/Private/Modules/StatsModule.cpp
+++ b/Engine/Source/Developer/TraceServices/Private/Modules/StatsModule.cpp
@@ -2,8 +2,7 @@
 
 #include "StatsModule.h"
 #include "Analyzers/StatsTraceAnalysis.h"
-#include "AnalysisServicePrivate.h"
-#include "TraceServices/ModuleService.h"
+#include "TraceServices/Model/Counters.h"
 
 namespace TraceServices
 {
diff --git a/Engine/Source/Developer/TraceServices/Public/Model/MonotonicTimeline.h b/Engine/Source/Developer/TraceServices/Public/Model/MonotonicTimeline.h
--- a/Engine/Source/Developer/TraceServices/Public/Model/MonotonicTimeline.h
+++ b/Engine/Source/Developer/TraceServices/Public/Model/MonotonicTimeline.h
@@ -7,7 +7,6 @@
 #include "Common/PagedArray.h"
 #include "Model/AsyncEnumerateTask.h"
 #include "Model/MonotonicTimelineData.h"
-#include "TraceServices/AnalysisService.h"
 #include "TraceServices/Containers/Timelines.h"
 
 namespace TraceServices
diff --git a/Engine/Source/Developer/TraceServices/Public/TraceServices/AnalysisService.h b/Engine/Source/Developer/TraceServices/Public/TraceServices/AnalysisService.h
--- a/Engine/Source/Developer/TraceServices/Public/TraceServices/AnalysisService.h
+++ b/Engine/Source/Developer/TraceServices/Public/TraceServices/AnalysisService.h
@@ -5,21 +5,21 @@
 #include "CoreTypes.h"
 #include "Templates/SharedPointer.h"
 #include "Templates/UniquePtr.h"
-#include "Delegates/Delegate.h"
 #include "Trace/DataStream.h"
-#include "TraceServices/Model/AnalysisSession.h"
-#include "TraceServices/Model/AnalysisCache.h"
-#include "TraceServices/Model/Log.h"
-#include "TraceServices/Model/Bookmarks.h"
-#include "TraceServices/Model/Frames.h"
-#include "TraceServices/Model/Threads.h"
-#include "TraceServices/Model/TimingProfiler.h"
-#include "TraceServices/Model/LoadTimeProfiler.h"
-#include "TraceServices/Model/Counters.h"
+
+namespace UE
+{
+namespace Trace
+{
+	class IInDataStream;
+}
+}
 
 namespace TraceServices
 {
 
+class IAnalysisSession;
+
 class IAnalysisService
 {
 public:
diff --git a/Engine/Source/Developer/TraceServices/Public/TraceServices/Model/Threads.h b/Engine/Source/Developer/TraceServices/Public/TraceServices/Model/Threads.h
--- a/Engine/Source/Developer/TraceServices/Public/TraceServices/Model/Threads.h
+++ b/Engine/Source/Developer/TraceServices/Public/TraceServices/Model/Threads.h
@@ -2,6 +2,7 @@
 
 #pragma once
 
+#include "HAL/PlatformAffinity.h"
 #include "TraceServices/Model/AnalysisSession.h"
 #include "Templates/Function.h"
 
diff --git a/Engine/Source/Developer/TraceServices/Public/TraceServices/Model/TimingProfiler.h b/Engine/Source/Developer/TraceServices/Public/TraceServices/Model/TimingProfiler.h
--- a/Engine/Source/Developer/TraceServices/Public/TraceServices/Model/TimingProfiler.h
+++ b/Engine/Source/Developer/TraceServices/Public/TraceServices/Model/TimingProfiler.h
@@ -2,6 +2,7 @@
 
 #pragma once
 
+#include "Model/MonotonicTimeline.h"
 #include "TraceServices/Model/AnalysisSession.h"
 #include "TraceServices/Containers/Timelines.h"
 #include "TraceServices/Containers/Tables.h"
