Non-unity build fix.

diff --git a/Engine/Source/Developer/TraceServices/Private/Model/Diagnostics.cpp b/Engine/Source/Developer/TraceServices/Private/Model/Diagnostics.cpp
--- a/Engine/Source/Developer/TraceServices/Private/Model/Diagnostics.cpp
+++ b/Engine/Source/Developer/TraceServices/Private/Model/Diagnostics.cpp
@@ -1,6 +1,7 @@
 // Copyright Epic Games, Inc. All Rights Reserved.
 #include "TraceServices/Model/Diagnostics.h"
 #include "Model/DiagnosticsPrivate.h"
+#include "UObject/NameTypes.h"
 
 namespace TraceServices
 {
