Non-unity fix.

diff --git a/Engine/Source/Developer/TraceServices/Private/AnalysisCache.h b/Engine/Source/Developer/TraceServices/Private/AnalysisCache.h
--- a/Engine/Source/Developer/TraceServices/Private/AnalysisCache.h
+++ b/Engine/Source/Developer/TraceServices/Private/AnalysisCache.h
@@ -5,6 +5,7 @@
 #include "IO/IoHash.h"
 #include "Memory/SharedBuffer.h"
 #include "Templates/UniquePtr.h"
+#include "Containers/Map.h"
 #include "TraceServices/Model/AnalysisCache.h"
 
 class IMappedFileRegion;
