Remove accidental debug code and added comment.

diff --git a/Engine/Source/Developer/TraceServices/Private/Analyzers/StringsAnalyzer.cpp b/Engine/Source/Developer/TraceServices/Private/Analyzers/StringsAnalyzer.cpp
--- a/Engine/Source/Developer/TraceServices/Private/Analyzers/StringsAnalyzer.cpp
+++ b/Engine/Source/Developer/TraceServices/Private/Analyzers/StringsAnalyzer.cpp
@@ -33,10 +33,6 @@ bool FStringsAnalyzer::OnEvent(uint16 RouteId, EStyle Style, const FOnEventConte
 			FWideStringView DisplayWide;
 			FAnsiStringView DisplayAnsi;
 			const TCHAR* Display = nullptr;
-			if (DisplayAnsi.Compare("LyraTestAssetName1") == 0)
-			{
-				UE_DEBUG_BREAK();
-			}
 			if (EventData.GetString("DisplayWide", DisplayWide) && DisplayWide.Len() > 0)
 			{
 				Display = Session.StoreString(DisplayWide);
@@ -58,10 +54,6 @@ bool FStringsAnalyzer::OnEvent(uint16 RouteId, EStyle Style, const FOnEventConte
 			if (RouteId == RouteId_FName)
 			{
 				auto Id = EventData.GetDefinitionId<uint32>();
-				if (Id.Id == 1651656)
-				{
-					UE_DEBUG_BREAK();
-				}
 				// Overwrite type id of Strings.FNameNoSync to Strings.FNames type id.
 				//Id.RefTypeId = ???;
 				DefinitionProvider->Register<FStringDefinition>(Instance, Id);
diff --git a/Engine/Source/Runtime/Core/Public/ProfilingDebugging/StringsTrace.h b/Engine/Source/Runtime/Core/Public/ProfilingDebugging/StringsTrace.h
--- a/Engine/Source/Runtime/Core/Public/ProfilingDebugging/StringsTrace.h
+++ b/Engine/Source/Runtime/Core/Public/ProfilingDebugging/StringsTrace.h
@@ -16,6 +16,8 @@ CORE_API UE_TRACE_EVENT_BEGIN_EXTERN(Strings, FName, Definition32bit)
 	UE_TRACE_EVENT_FIELD(UE::Trace::AnsiString, DisplayAnsi)
 UE_TRACE_EVENT_END()
 
+// The following events are used when on connection callback is triggered. Do not use or reference
+
 UE_TRACE_EVENT_BEGIN_EXTERN(Strings, StaticStringNoSync, Definition64bit | NoSync)
 	UE_TRACE_EVENT_FIELD(UE::Trace::WideString, DisplayWide)
 	UE_TRACE_EVENT_FIELD(UE::Trace::AnsiString, DisplayAnsi)
