Fixed non unity file compilation error

diff --git a/Engine/Source/Developer/TraceServices/Private/Analyzers/NetTraceAnalyzer.h b/Engine/Source/Developer/TraceServices/Private/Analyzers/NetTraceAnalyzer.h
--- a/Engine/Source/Developer/TraceServices/Private/Analyzers/NetTraceAnalyzer.h
+++ b/Engine/Source/Developer/TraceServices/Private/Analyzers/NetTraceAnalyzer.h
@@ -3,6 +3,7 @@
 #pragma once
 
 #include "Trace/Analyzer.h"
+#include "Containers/Map.h"
 #include "Containers/UnrealString.h"
 #include "Model/NetProfilerProvider.h"
 
