Fixed non-unity (-DisableUnity -NoPCH) compile errors in TraceServices module.

diff --git a/Engine/Source/Developer/TraceServices/Private/Analyzers/NetTraceAnalyzer.h b/Engine/Source/Developer/TraceServices/Private/Analyzers/NetTraceAnalyzer.h
--- a/Engine/Source/Developer/TraceServices/Private/Analyzers/NetTraceAnalyzer.h
+++ b/Engine/Source/Developer/TraceServices/Private/Analyzers/NetTraceAnalyzer.h
@@ -2,10 +2,12 @@
 
 #pragma once
 
-#include "Trace/Analyzer.h"
 #include "Containers/Map.h"
 #include "Containers/UnrealString.h"
+#include "Templates/SharedPointer.h"
+
 #include "Model/NetProfilerProvider.h"
+#include "Trace/Analyzer.h"
 
 namespace TraceServices
 {
@@ -133,7 +135,7 @@ private:
 
 	void FlushPacketEvents(FNetTraceConnectionState& ConnectionState, FNetProfilerConnectionData& ConnectionData, const ENetProfilerConnectionMode ConnectionMode);
 
-	void FlushFrameStatsCounters(FNetTraceAnalyzer::FNetTraceGameInstanceState& GameInstanceState);
+	void FlushFrameStatsCounters(FNetTraceGameInstanceState& GameInstanceState);
 	uint32 GetCurrentNetProfilerFrameIndexAndFlushFrameStatsCountersIfNeeded(uint32 GameInstanceId, uint32 EngineFrameIndex);
 
 	void HandlePacketEvent(const FOnEventContext& Context, const FEventData& EventData);	
diff --git a/Engine/Source/Developer/TraceServices/Private/Model/CookProfilerProviderPrivate.h b/Engine/Source/Developer/TraceServices/Private/Model/CookProfilerProviderPrivate.h
--- a/Engine/Source/Developer/TraceServices/Private/Model/CookProfilerProviderPrivate.h
+++ b/Engine/Source/Developer/TraceServices/Private/Model/CookProfilerProviderPrivate.h
@@ -2,6 +2,7 @@
 
 #pragma once
 
+#include "Containers/Map.h"
 #include "TraceServices/Model/CookProfilerProvider.h"
 
 namespace TraceServices
diff --git a/Engine/Source/Developer/TraceServices/Private/Model/MemoryPrivate.h b/Engine/Source/Developer/TraceServices/Private/Model/MemoryPrivate.h
--- a/Engine/Source/Developer/TraceServices/Private/Model/MemoryPrivate.h
+++ b/Engine/Source/Developer/TraceServices/Private/Model/MemoryPrivate.h
@@ -1,8 +1,11 @@
 // Copyright Epic Games, Inc. All Rights Reserved.
 #pragma once
 
-#include "TraceServices/Model/Memory.h"
+#include "Containers/Map.h"
+#include "UObject/NameTypes.h"
+
 #include "Common/PagedArray.h"
+#include "TraceServices/Model/Memory.h"
 
 namespace TraceServices
 {
