[Insights] Fixed missing TRACESERVICES_API for IDefinitionProvider API.

diff --git a/Engine/Source/Developer/TraceServices/Public/TraceServices/Model/Definitions.h b/Engine/Source/Developer/TraceServices/Public/TraceServices/Model/Definitions.h
--- a/Engine/Source/Developer/TraceServices/Public/TraceServices/Model/Definitions.h
+++ b/Engine/Source/Developer/TraceServices/Public/TraceServices/Model/Definitions.h
@@ -74,8 +74,8 @@ protected:
 	virtual const void* FindEntry(uint64 Hash) const = 0;
 };
 
-IDefinitionProvider* GetDefinitionProvider(IAnalysisSession& Session);
-const IDefinitionProvider* ReadDefinitionProvider(const IAnalysisSession& Session);
+TRACESERVICES_API IDefinitionProvider* GetDefinitionProvider(IAnalysisSession& Session);
+TRACESERVICES_API const IDefinitionProvider* ReadDefinitionProvider(const IAnalysisSession& Session);
 
 } // namespace TraceServices
 
