Changes to how definition provider works

diff --git a/Engine/Source/Developer/TraceAnalysis/Public/Trace/Analyzer.h b/Engine/Source/Developer/TraceAnalysis/Public/Trace/Analyzer.h
--- a/Engine/Source/Developer/TraceAnalysis/Public/Trace/Analyzer.h
+++ b/Engine/Source/Developer/TraceAnalysis/Public/Trace/Analyzer.h
@@ -176,6 +176,9 @@ public:
 		template<typename DefinitionType>
 		TEventRef<DefinitionType> GetReferenceValue(const ANSICHAR* FieldName) const;
 
+		template<typename DefinitionType>
+		TEventRef<DefinitionType> GetReferenceValue(uint32 FieldIndex) const;
+		
 		/** If this is a spec event, gets the unique Id for this spec.
 		 * @return A valid spec id if the event is valid, otherwise an empty id.
 		 */
@@ -205,7 +208,9 @@ public:
 
 	private:
 		bool IsDefinitionImpl(uint32& OutTypeId) const;
+		uint32 GetReferenceTypeIdImpl(FEventFieldHandle Handle) const;
 		const void* GetReferenceValueImpl(const char* FieldName, uint16& OutSizeType, uint32& OutTypeUid) const;
+		const void* GetReferenceValueImpl(uint32 FieldIndex, uint32& OutTypeUid) const;
 		const void* GetValueImpl(const ANSICHAR* FieldName, int8& SizeAndType) const;
 		const FArrayReader* GetArrayImpl(const ANSICHAR* FieldName) const;
 	};
@@ -437,5 +442,14 @@ TEventRef<DefinitionType> IAnalyzer::FEventData::GetReferenceValue(const ANSICHA
 	return MakeEventRef<DefinitionType>(0,0);
 }
 
+////////////////////////////////////////////////////////////////////////////////
+template <typename DefinitionType>
+TEventRef<DefinitionType> IAnalyzer::FEventData::GetReferenceValue(uint32 FieldIndex) const
+{
+	uint32 RefTypeUid;
+	DefinitionType* Id = (DefinitionType*) GetReferenceValueImpl(FieldIndex, RefTypeUid);
+	return MakeEventRef<DefinitionType>(*Id, RefTypeUid);
+}
+	
 } // namespace Trace
 } // namespace UE
diff --git a/Engine/Source/Developer/TraceServices/Private/Analyzers/StringsAnalyzer.cpp b/Engine/Source/Developer/TraceServices/Private/Analyzers/StringsAnalyzer.cpp
--- a/Engine/Source/Developer/TraceServices/Private/Analyzers/StringsAnalyzer.cpp
+++ b/Engine/Source/Developer/TraceServices/Private/Analyzers/StringsAnalyzer.cpp
@@ -48,21 +48,18 @@ bool FStringsAnalyzer::OnEvent(uint16 RouteId, EStyle Style, const FOnEventConte
 				return true;
 			}
 
+			IDefinitionProvider::FEditScopeLock Lock(*DefinitionProvider);
 			FStringDefinition* Instance = DefinitionProvider->Create<FStringDefinition>();
 			Instance->Display = Display;
 
 			if (RouteId == RouteId_FName)
 			{
 				auto Id = EventData.GetDefinitionId<uint32>();
-				// Overwrite type id of Strings.FNameNoSync to Strings.FNames type id.
-				//Id.RefTypeId = ???;
 				DefinitionProvider->Register<FStringDefinition>(Instance, Id);
 			}
 			else
 			{
 				auto Id = EventData.GetDefinitionId<uint64>();
-				// Overwrite type id of Strings.StaticStringNoSync to Strings.StaticString type id.
-				//Id.RefTypeId = ???;
 				DefinitionProvider->Register<FStringDefinition>(Instance, Id);
 			}
 		}
diff --git a/Engine/Source/Developer/TraceServices/Private/Model/DefinitionProvider.cpp b/Engine/Source/Developer/TraceServices/Private/Model/DefinitionProvider.cpp
--- a/Engine/Source/Developer/TraceServices/Private/Model/DefinitionProvider.cpp
+++ b/Engine/Source/Developer/TraceServices/Private/Model/DefinitionProvider.cpp
@@ -13,16 +13,14 @@ FDefinitionProvider::FDefinitionProvider(IAnalysisSession* InSession)
 }
 
 //////////////////////////////////////////////////////////////////////////////////////////
-void FDefinitionProvider::AddEntry(uint64 Hash, void* Ptr)
+void FDefinitionProvider::AddEntry(uint64 Hash, const void* Ptr)
 {
-	FWriteScopeLock _(DefinitionsLock);
 	Definitions.Add(Hash, Ptr);
 }
 
 //////////////////////////////////////////////////////////////////////////////////////////
 const void* FDefinitionProvider::FindEntry(uint64 Hash) const
 {
-	FReadScopeLock _(DefinitionsLock);
 	return Definitions.FindRef(Hash);
 }
 
diff --git a/Engine/Source/Developer/TraceServices/Private/Model/DefinitionProvider.h b/Engine/Source/Developer/TraceServices/Private/Model/DefinitionProvider.h
--- a/Engine/Source/Developer/TraceServices/Private/Model/DefinitionProvider.h
+++ b/Engine/Source/Developer/TraceServices/Private/Model/DefinitionProvider.h
@@ -17,8 +17,12 @@ class FDefinitionProvider : public IDefinitionProvider
 public:
 	FDefinitionProvider(IAnalysisSession* InSession);
 
+	virtual void BeginEdit() const override { Lock.WriteLock(); };
+	virtual void EndEdit() const override { Lock.WriteUnlock(); };
+	virtual void BeginRead() const override { Lock.ReadLock(); };
+	virtual void EndRead() const override { Lock.ReadUnlock(); };
 private:
-	virtual void AddEntry(uint64 Hash, void* Ptr) override;
+	virtual void AddEntry(uint64 Hash, const void* Ptr) override;
 	virtual const void* FindEntry(uint64 Hash) const override;
 	virtual void* Allocate(uint32 Size, uint32 Alignment) override;
 
@@ -26,8 +30,8 @@ private:
 	static constexpr uint32 PageSize = 1024;
 	uint32 PageRemain;
 
-	mutable FRWLock DefinitionsLock;
-	TMap<uint64, void*> Definitions;
+	TMap<uint64, const void*> Definitions;
+	mutable FRWLock Lock;
 };
 
 } // namespace TraceServices
diff --git a/Engine/Source/Developer/TraceServices/Public/TraceServices/Model/Definitions.h b/Engine/Source/Developer/TraceServices/Public/TraceServices/Model/Definitions.h
--- a/Engine/Source/Developer/TraceServices/Public/TraceServices/Model/Definitions.h
+++ b/Engine/Source/Developer/TraceServices/Public/TraceServices/Model/Definitions.h
@@ -30,6 +30,38 @@ namespace TraceServices
 class IDefinitionProvider : public IProvider
 {
 public:
+	
+	struct TRACESERVICES_API FEditScopeLock
+	{
+		FEditScopeLock(const IDefinitionProvider& InMetadataProvider)
+			: MetadataProvider(InMetadataProvider)
+		{
+			MetadataProvider.BeginEdit();
+		}
+
+		~FEditScopeLock()
+		{
+			MetadataProvider.EndEdit();
+		}
+
+		const IDefinitionProvider& MetadataProvider;
+	};
+
+	struct TRACESERVICES_API FReadScopeLock
+	{
+		FReadScopeLock(const IDefinitionProvider& InMetadataProvider)
+			: MetadataProvider(InMetadataProvider)
+		{
+			MetadataProvider.BeginRead();
+		}
+
+		~FReadScopeLock()
+		{
+			MetadataProvider.EndRead();
+		}
+
+		const IDefinitionProvider& MetadataProvider;
+	};
 
 	/**
 	 * Allocates memory for the definition. The memory is guaranteed to be valid during the lifetime of the
@@ -51,8 +83,9 @@ public:
 	 * @param Id Id that should be used to reference the definition.
 	 */
 	template<typename T, typename DefinitionType>
-	void Register(T* Instance, UE::Trace::TEventRef<DefinitionType> Id)
+	void Register(const T* Instance, const UE::Trace::TEventRef<DefinitionType>& Id)
 	{
+		Stringifiers.FindOrAdd(Id.RefTypeId, &T::ToString);
 		AddEntry(Id.GetHash(), Instance);
 	}
 
@@ -62,16 +95,36 @@ public:
 	 * @return Pointer to the definition or null if the id was not found.
 	 */
 	template<typename T, typename DefinitionType>
-	const T* Get(UE::Trace::TEventRef<DefinitionType> Reference) const
+	const T* Get(const UE::Trace::TEventRef<DefinitionType>& Reference) const
 	{
 		const void* Value = FindEntry(Reference.GetHash());
 		return (T*) Value;
 	}
 
+	/**
+	 * Gets the definition as a owned string.
+	 * @param Reference Id used to uniqely identify the definition.
+	 * @return A string describing the referenced value
+	 */
+	template<typename DefinitionType>
+	FString GetAsString(const UE::Trace::TEventRef<DefinitionType>& Reference) const
+	{
+		const void* Entry = FindEntry(Reference.GetHash());
+		auto Stringifier = Stringifiers.Find(Reference.RefTypeId);
+		if (Entry && Stringifier)
+		{
+			return (*Stringifier)(Entry);
+		}
+		return FString();
+	}
+
 protected:
 	virtual void* Allocate(uint32 Size, uint32 Alignment) = 0;
-	virtual void AddEntry(uint64 Hash, void* Ptr) = 0;
+	virtual void AddEntry(uint64 Hash, const void* Ptr) = 0;
 	virtual const void* FindEntry(uint64 Hash) const = 0;
+
+	using StringifierFn = TFunction<FString (const void*)>;
+	TMap<uint32, StringifierFn> Stringifiers;
 };
 
 TRACESERVICES_API IDefinitionProvider* GetDefinitionProvider(IAnalysisSession& Session);
diff --git a/Engine/Source/Developer/TraceServices/Public/TraceServices/Model/Strings.h b/Engine/Source/Developer/TraceServices/Public/TraceServices/Model/Strings.h
--- a/Engine/Source/Developer/TraceServices/Public/TraceServices/Model/Strings.h
+++ b/Engine/Source/Developer/TraceServices/Public/TraceServices/Model/Strings.h
@@ -12,5 +12,11 @@ namespace TraceServices
 		 * Display string. The pointer is valid during the analysis session.
 		 */
 		const TCHAR* Display;
+
+		static FString ToString(const void* Data)
+		{
+			const FStringDefinition* Def = (const FStringDefinition*) Data;
+			return FString(Def->Display);
+		}
 	};
 }
\ No newline at end of file
