[Insights] Removed duplicated FEditScopeLock/FReadScopeLock from IMetadataProvider and IAllocationsProvider and used the ones from ProviderLock.h instead.

diff --git a/Engine/Source/Developer/TraceInsights/Private/Insights/MemoryProfiler/MemoryProfilerManager.cpp b/Engine/Source/Developer/TraceInsights/Private/Insights/MemoryProfiler/MemoryProfilerManager.cpp
--- a/Engine/Source/Developer/TraceInsights/Private/Insights/MemoryProfiler/MemoryProfilerManager.cpp
+++ b/Engine/Source/Developer/TraceInsights/Private/Insights/MemoryProfiler/MemoryProfilerManager.cpp
@@ -4,12 +4,15 @@
 
 #include "MessageLog/Public/MessageLogModule.h"
 #include "Modules/ModuleManager.h"
-#include "TraceServices/Model/AllocationsProvider.h"
-#include "TraceServices/Model/Memory.h"
 #include "Widgets/Docking/SDockTab.h"
 #include "WorkspaceMenuStructure.h"
 #include "WorkspaceMenuStructureModule.h"
 
+// TraceServices
+#include "Common/ProviderLock.h"
+#include "TraceServices/Model/AllocationsProvider.h"
+#include "TraceServices/Model/Memory.h"
+
 // Insights
 #include "Insights/Common/InsightsMenuBuilder.h"
 #include "Insights/InsightsManager.h"
@@ -273,7 +276,7 @@ bool FMemoryProfilerManager::Tick(float DeltaTime)
 			const TraceServices::IAllocationsProvider* AllocationsProvider = TraceServices::ReadAllocationsProvider(*Session.Get());
 			if (AllocationsProvider)
 			{
-				TraceServices::IAllocationsProvider::FReadScopeLock _(*AllocationsProvider);
+				TraceServices::FProviderReadScopeLock _(*AllocationsProvider);
 				if (AllocationsProvider->IsInitialized())
 				{
 					bShouldBeAvailable = true;
diff --git a/Engine/Source/Developer/TraceInsights/Private/Insights/MemoryProfiler/ViewModels/MemoryGraphTrack.cpp b/Engine/Source/Developer/TraceInsights/Private/Insights/MemoryProfiler/ViewModels/MemoryGraphTrack.cpp
--- a/Engine/Source/Developer/TraceInsights/Private/Insights/MemoryProfiler/ViewModels/MemoryGraphTrack.cpp
+++ b/Engine/Source/Developer/TraceInsights/Private/Insights/MemoryProfiler/ViewModels/MemoryGraphTrack.cpp
@@ -5,6 +5,9 @@
 #include "Fonts/FontMeasure.h"
 #include "Fonts/SlateFontInfo.h"
 #include "Styling/SlateBrush.h"
+
+// TraceServices
+#include "Common/ProviderLock.h"
 #include "TraceServices/Model/AllocationsProvider.h"
 #include "TraceServices/Model/TimingProfiler.h"
 #include "TraceServices/Model/Counters.h"
@@ -478,7 +481,7 @@ void FMemoryGraphTrack::PreUpdateAllocationsTimelineSeries(FMemoryGraphSeries& S
 		const TraceServices::IAllocationsProvider* AllocationsProvider = TraceServices::ReadAllocationsProvider(*Session.Get());
 		if (AllocationsProvider)
 		{
-			TraceServices::IAllocationsProvider::FReadScopeLock ProviderReadScope(*AllocationsProvider);
+			TraceServices::FProviderReadScopeLock ProviderReadScope(*AllocationsProvider);
 
 			int32 StartIndex = -1;
 			int32 EndIndex = -1;
@@ -582,7 +585,7 @@ void FMemoryGraphTrack::UpdateAllocationsTimelineSeries(FMemoryGraphSeries& Seri
 
 			TArray<FTimelineEvent> TimelineEvents;
 			{
-				TraceServices::IAllocationsProvider::FReadScopeLock ProviderReadScope(*AllocationsProvider);
+				TraceServices::FProviderReadScopeLock ProviderReadScope(*AllocationsProvider);
 				AllocationsProvider->GetTimelineIndexRange(Viewport.GetStartTime(), Viewport.GetEndTime(), StartIndex, EndIndex);
 
 				if (EndIndex >= 0)
diff --git a/Engine/Source/Developer/TraceInsights/Private/Insights/MemoryProfiler/Widgets/SMemAllocTableTreeView.cpp b/Engine/Source/Developer/TraceInsights/Private/Insights/MemoryProfiler/Widgets/SMemAllocTableTreeView.cpp
--- a/Engine/Source/Developer/TraceInsights/Private/Insights/MemoryProfiler/Widgets/SMemAllocTableTreeView.cpp
+++ b/Engine/Source/Developer/TraceInsights/Private/Insights/MemoryProfiler/Widgets/SMemAllocTableTreeView.cpp
@@ -9,15 +9,18 @@
 #include "Modules/ModuleManager.h"
 #include "SlateOptMacros.h"
 #include "Styling/AppStyle.h"
+#include "Widgets/Images/SImage.h"
+#include "Widgets/Input/SCheckBox.h"
+#include "Widgets/SToolTip.h"
+
+// TraceServices
+#include "Common/ProviderLock.h"
 #include "TraceServices/Model/AllocationsProvider.h"
 #include "TraceServices/Model/Callstack.h"
 #include "TraceServices/Model/Definitions.h"
 #include "TraceServices/Model/MetadataProvider.h"
 #include "TraceServices/Model/Modules.h"
 #include "TraceServices/Model/Strings.h"
-#include "Widgets/Images/SImage.h"
-#include "Widgets/Input/SCheckBox.h"
-#include "Widgets/SToolTip.h"
 
 // Insights
 #include "Insights/InsightsStyle.h"
@@ -260,7 +263,7 @@ void SMemAllocTableTreeView::StartQuery()
 
 	{
 		const TraceServices::IAllocationsProvider& Provider = *AllocationsProvider;
-		TraceServices::IAllocationsProvider::FReadScopeLock _(Provider);
+		TraceServices::FProviderReadScopeLock _(Provider);
 		TraceServices::IAllocationsProvider::FQueryParams Params = { Rule->GetValue(), TimeMarkers[0], TimeMarkers[1], TimeMarkers[2], TimeMarkers[3] };
 		Query = Provider.StartQuery(Params);
 	}
@@ -309,7 +312,7 @@ void SMemAllocTableTreeView::UpdateQuery(TraceServices::IAllocationsProvider::EQ
 	uint16 AssetMetadataType = 0;
 	if (MetadataProvider)
 	{
-		TraceServices::IMetadataProvider::FReadScopeLock MetadataProviderReadLock(*MetadataProvider);
+		TraceServices::FProviderReadScopeLock MetadataProviderReadLock(*MetadataProvider);
 		AssetMetadataType = MetadataProvider->GetRegisteredMetadataType(TEXT("Asset"));
 		if (MetadataProvider->GetRegisteredMetadataName(AssetMetadataType) == nullptr)
 		{
@@ -347,7 +350,7 @@ void SMemAllocTableTreeView::UpdateQuery(TraceServices::IAllocationsProvider::EQ
 		TSharedPtr<Insights::FMemAllocTable> MemAllocTable = GetMemAllocTable();
 		if (MemAllocTable)
 		{
-			TraceServices::IAllocationsProvider::FReadScopeLock _(Provider);
+			TraceServices::FProviderReadScopeLock _(Provider);
 
 			TArray<FMemoryAlloc>& Allocs = MemAllocTable->GetAllocs();
 
@@ -393,7 +396,7 @@ void SMemAllocTableTreeView::UpdateQuery(TraceServices::IAllocationsProvider::EQ
 					Alloc.Asset = nullptr;
 					if (MetadataProvider && DefinitionProvider)
 					{
-						TraceServices::IMetadataProvider::FReadScopeLock MetadataProviderReadLock(*MetadataProvider);
+						TraceServices::FProviderReadScopeLock MetadataProviderReadLock(*MetadataProvider);
 						MetadataProvider->EnumerateMetadata(Allocation->GetThreadId(), Allocation->GetMetadataId(),
 							[AssetMetadataType, &Alloc, DefinitionProvider](uint32 StackDepth, uint16 Type, const void* Data, uint32 Size) -> bool
 							{
@@ -1305,7 +1308,7 @@ void SMemAllocTableTreeView::PopulateLLMTagSuggestionList(const FString& Text, T
 		return;
 	}
 
-	TraceServices::IAllocationsProvider::FReadScopeLock _(*AllocationsProvider);
+	TraceServices::FProviderReadScopeLock _(*AllocationsProvider);
 
 	AllocationsProvider->EnumerateTags([&OutSuggestions, &Text](const TCHAR* Display, const TCHAR* FullPath, TraceServices::TagIdType CurrentTag, TraceServices::TagIdType ParentTag)
 	{
diff --git a/Engine/Source/Developer/TraceServices/Private/Analyzers/AllocationsAnalysis.cpp b/Engine/Source/Developer/TraceServices/Private/Analyzers/AllocationsAnalysis.cpp
--- a/Engine/Source/Developer/TraceServices/Private/Analyzers/AllocationsAnalysis.cpp
+++ b/Engine/Source/Developer/TraceServices/Private/Analyzers/AllocationsAnalysis.cpp
@@ -3,9 +3,12 @@
 #include "AllocationsAnalysis.h"
 
 #include "HAL/LowLevelMemTracker.h"
+#include "ProfilingDebugging/MemoryTrace.h"
+
+// TraceServices
+#include "Common/ProviderLock.h"
 #include "Model/AllocationsProvider.h"
 #include "Model/MetadataProvider.h"
-#include "ProfilingDebugging/MemoryTrace.h"
 #include "TraceServices/Model/AnalysisSession.h"
 #include "TraceServices/Model/Callstack.h"
 
@@ -78,7 +81,7 @@ void FAllocationsAnalyzer::OnAnalysisEnd()
 	}
 	const double Time = FMath::Max(SessionTime, LastMarkerSeconds);
 
-	FAllocationsProvider::FEditScopeLock _(AllocationsProvider);
+	FProviderEditScopeLock _(AllocationsProvider);
 	AllocationsProvider.EditOnAnalysisCompleted(Time);
 }
 
@@ -115,7 +118,7 @@ bool FAllocationsAnalyzer::OnEvent(uint16 RouteId, EStyle Style, const FOnEventC
 			const uint8 MinAlignment = EventData.GetValue<uint8>("MinAlignment");
 			SizeShift = EventData.GetValue<uint8>("SizeShift");
 
-			FAllocationsProvider::FEditScopeLock _(AllocationsProvider);
+			FProviderEditScopeLock _(AllocationsProvider);
 			AllocationsProvider.EditInit(Time, MinAlignment);
 			break;
 		}
@@ -128,7 +131,7 @@ bool FAllocationsAnalyzer::OnEvent(uint16 RouteId, EStyle Style, const FOnEventC
 			FStringView Name;
 			EventData.GetString("Name", Name);
 
-			FAllocationsProvider::FEditScopeLock _(AllocationsProvider);
+			FProviderEditScopeLock _(AllocationsProvider);
 			AllocationsProvider.EditHeapSpec(Id, ParentId, Name, Flags);
 			break;
 		}
@@ -194,7 +197,7 @@ bool FAllocationsAnalyzer::OnEvent(uint16 RouteId, EStyle Style, const FOnEventC
 			const uint32 TraceThreadId = Context.ThreadInfo.GetId();
 			const uint32 SystemThreadId = Context.ThreadInfo.GetSystemId();
 
-			FAllocationsProvider::FEditScopeLock _(AllocationsProvider);
+			FProviderEditScopeLock _(AllocationsProvider);
 			AllocationsProvider.SetCurrentThreadId(TraceThreadId, SystemThreadId);
 			AllocationsProvider.EditAlloc(Time, CallstackId, Address, Size, Alignment, RootHeap);
 			if (RouteId == RouteId_ReallocAlloc || RouteId == RouteId_ReallocAllocSystem)
@@ -238,7 +241,7 @@ bool FAllocationsAnalyzer::OnEvent(uint16 RouteId, EStyle Style, const FOnEventC
 			const uint32 TraceThreadId = Context.ThreadInfo.GetId();
 			const uint32 SystemThreadId = Context.ThreadInfo.GetSystemId();
 
-			FAllocationsProvider::FEditScopeLock _(AllocationsProvider);
+			FProviderEditScopeLock _(AllocationsProvider);
 			AllocationsProvider.SetCurrentThreadId(TraceThreadId, SystemThreadId);
 			if (RouteId == RouteId_ReallocFree || RouteId == RouteId_ReallocFreeSystem)
 			{
@@ -259,7 +262,7 @@ bool FAllocationsAnalyzer::OnEvent(uint16 RouteId, EStyle Style, const FOnEventC
 			const uint32 TraceThreadId = Context.ThreadInfo.GetId();
 			const uint32 SystemThreadId = Context.ThreadInfo.GetSystemId();
 
-			FAllocationsProvider::FEditScopeLock _(AllocationsProvider);
+			FProviderEditScopeLock _(AllocationsProvider);
 			AllocationsProvider.SetCurrentThreadId(TraceThreadId, SystemThreadId);
 			AllocationsProvider.EditMarkAllocationAsHeap(Time, Address, Heap, Flags);
 			break;
@@ -274,7 +277,7 @@ bool FAllocationsAnalyzer::OnEvent(uint16 RouteId, EStyle Style, const FOnEventC
 			const uint32 TraceThreadId = Context.ThreadInfo.GetId();
 			const uint32 SystemThreadId = Context.ThreadInfo.GetSystemId();
 
-			FAllocationsProvider::FEditScopeLock _(AllocationsProvider);
+			FProviderEditScopeLock _(AllocationsProvider);
 			AllocationsProvider.SetCurrentThreadId(TraceThreadId, SystemThreadId);
 			AllocationsProvider.EditUnmarkAllocationAsHeap(Time, Address, Heap);
 			break;
@@ -315,7 +318,7 @@ bool FAllocationsAnalyzer::OnEvent(uint16 RouteId, EStyle Style, const FOnEventC
 			Context.EventData.GetString("Display", Display);
 			const TCHAR* DisplayString = Session.StoreString(*Display);
 
-			FAllocationsProvider::FEditScopeLock _(AllocationsProvider);
+			FProviderEditScopeLock _(AllocationsProvider);
 			AllocationsProvider.EditAddTagSpec(Tag, Parent, DisplayString);
 			break;
 		}
@@ -331,7 +334,7 @@ bool FAllocationsAnalyzer::OnEvent(uint16 RouteId, EStyle Style, const FOnEventC
 				{
 					const TagIdType Tag = Context.EventData.GetValue<TagIdType>("Tag");
 					{
-						FAllocationsProvider::FEditScopeLock _(AllocationsProvider);
+						FProviderEditScopeLock _(AllocationsProvider);
 						AllocationsProvider.EditPushTag(ThreadId, Tracker, Tag);
 					}
 #if INSIGHTS_MEM_TRACE_METADATA_TEST
@@ -345,7 +348,7 @@ bool FAllocationsAnalyzer::OnEvent(uint16 RouteId, EStyle Style, const FOnEventC
 				{
 					const uint64 Ptr = Context.EventData.GetValue<uint64>("Ptr");
 					{
-						FAllocationsProvider::FEditScopeLock _(AllocationsProvider);
+						FProviderEditScopeLock _(AllocationsProvider);
 						AllocationsProvider.EditPushTagFromPtr(ThreadId, Tracker, Ptr);
 					}
 #if INSIGHTS_MEM_TRACE_METADATA_TEST
@@ -360,7 +363,7 @@ bool FAllocationsAnalyzer::OnEvent(uint16 RouteId, EStyle Style, const FOnEventC
 			else // EStyle::LeaveScope
 			{
 				{
-					FAllocationsProvider::FEditScopeLock _(AllocationsProvider);
+					FProviderEditScopeLock _(AllocationsProvider);
 					if (AllocationsProvider.HasTagFromPtrScope(ThreadId, Tracker)) // Is TagFromPtr scope active?
 					{
 						AllocationsProvider.EditPopTagFromPtr(ThreadId, Tracker);
diff --git a/Engine/Source/Developer/TraceServices/Private/Common/ProviderLock.cpp b/Engine/Source/Developer/TraceServices/Private/Common/ProviderLock.cpp
--- a/Engine/Source/Developer/TraceServices/Private/Common/ProviderLock.cpp
+++ b/Engine/Source/Developer/TraceServices/Private/Common/ProviderLock.cpp
@@ -4,10 +4,11 @@
 
 namespace TraceServices
 {
-void FProviderLock::ReadAccessCheck(const FProviderLock* CurrentProviderLock, const int32& CurrentReadProviderLockCount, const int32& CurrentWriteAllocationsProviderLockCount) const
+
+void FProviderLock::ReadAccessCheck(const FProviderLock* CurrentProviderLock, const int32& CurrentReadProviderLockCount, const int32& CurrentWriteProviderLockCount) const
 {
-	checkf(CurrentProviderLock == this && (CurrentReadProviderLockCount > 0 || CurrentWriteAllocationsProviderLockCount > 0),
-		TEXT("Trying to READ from allocations provider outside of a READ scope"));
+	checkf(CurrentProviderLock == this && (CurrentReadProviderLockCount > 0 || CurrentWriteProviderLockCount > 0),
+		TEXT("Trying to READ from provider outside of a READ scope"));
 }
 
 void FProviderLock::WriteAccessCheck(const int32& CurrentWriteAllocationsProviderLockCount) const
@@ -16,9 +17,9 @@ void FProviderLock::WriteAccessCheck(const int32& CurrentWriteAllocationsProvide
 		TEXT("Trying to WRITE to provider outside of an EDIT/WRITE scope"));
 }
 
-void FProviderLock::BeginRead(FProviderLock*& CurrentProviderLock, int32& CurrentReadProviderLockCount, const int32& WriteAllocationsProviderLockCount)
+void FProviderLock::BeginRead(FProviderLock*& CurrentProviderLock, int32& CurrentReadProviderLockCount, const int32& WriteProviderLockCount)
 {
-	checkf(WriteAllocationsProviderLockCount == 0, TEXT("Trying to lock allocations provider for READ while holding EDIT/WRITE access"));
+	checkf(WriteProviderLockCount == 0, TEXT("Trying to lock provider for READ while holding EDIT/WRITE access"));
 	if (CurrentReadProviderLockCount++ == 0)
 	{
 		CurrentProviderLock = this;
@@ -36,24 +37,25 @@ void FProviderLock::EndRead(FProviderLock*& CurrentProviderLock, int32& CurrentR
 	}
 }
 
-void FProviderLock::BeginWrite(FProviderLock*& CurrentProviderLock, const int32& CurrentReadProviderLockCount, int32& WriteAllocationsProviderLockCount)
+void FProviderLock::BeginWrite(FProviderLock*& CurrentProviderLock, const int32& CurrentReadProviderLockCount, int32& WriteProviderLockCount)
 {
 	check(!CurrentProviderLock || CurrentProviderLock == this);
-	checkf(CurrentReadProviderLockCount == 0, TEXT("Trying to lock allocations provider for EDIT/WRITE while holding READ access"));
-	if (WriteAllocationsProviderLockCount++ == 0)
+	checkf(CurrentReadProviderLockCount == 0, TEXT("Trying to lock provider for EDIT/WRITE while holding READ access"));
+	if (WriteProviderLockCount++ == 0)
 	{
 		CurrentProviderLock = this;
 		RWLock.WriteLock();
 	}
 }
 
-void FProviderLock::EndWrite(FProviderLock*& CurrentProviderLock, int32& WriteAllocationsProviderLockCount)
+void FProviderLock::EndWrite(FProviderLock*& CurrentProviderLock, int32& WriteProviderLockCount)
 {
-	check(WriteAllocationsProviderLockCount > 0);
-	if (--WriteAllocationsProviderLockCount == 0)
+	check(WriteProviderLockCount > 0);
+	if (--WriteProviderLockCount == 0)
 	{
 		RWLock.WriteUnlock();
 		CurrentProviderLock = nullptr;
 	}
 }
+
 } // namespace TraceServices
diff --git a/Engine/Source/Developer/TraceServices/Private/Model/AllocationsProvider.cpp b/Engine/Source/Developer/TraceServices/Private/Model/AllocationsProvider.cpp
--- a/Engine/Source/Developer/TraceServices/Private/Model/AllocationsProvider.cpp
+++ b/Engine/Source/Developer/TraceServices/Private/Model/AllocationsProvider.cpp
@@ -2,12 +2,15 @@
 
 #include "AllocationsProvider.h"
 
-#include "AllocationsQuery.h"
-#include "Common/Utils.h"
 #include "Containers/ArrayView.h"
 #include "Misc/PathViews.h"
-#include "Model/MetadataProvider.h"
 #include "ProfilingDebugging/MemoryTrace.h"
+
+// TraceServices
+#include "AllocationsQuery.h"
+#include "Common/ProviderLock.h"
+#include "Common/Utils.h"
+#include "Model/MetadataProvider.h"
 #include "SbTree.h"
 #include "TraceServices/Containers/Allocators.h"
 #include "TraceServices/Model/Callstack.h"
@@ -1286,7 +1289,7 @@ void FAllocationsProvider::EditAlloc(double Time, uint32 CallstackId, uint64 Add
 
 		uint32 MetadataId = MetadataProvider.InvalidMetadataId;
 		{
-			FMetadataProvider::FEditScopeLock _(MetadataProvider);
+			FProviderEditScopeLock _(MetadataProvider);
 			MetadataId = MetadataProvider.PinAndGetId(CurrentSystemThreadId);
 #if INSIGHTS_DEBUG_METADATA
 			if (MetadataId != FMetadataProvider::InvalidMetadataId &&
diff --git a/Engine/Source/Developer/TraceServices/Public/Common/ProviderLock.h b/Engine/Source/Developer/TraceServices/Public/Common/ProviderLock.h
--- a/Engine/Source/Developer/TraceServices/Public/Common/ProviderLock.h
+++ b/Engine/Source/Developer/TraceServices/Public/Common/ProviderLock.h
@@ -44,14 +44,14 @@ struct TRACESERVICES_API FProviderReadScopeLock
 class FProviderLock
 {
 public:
-	void ReadAccessCheck(const FProviderLock* CurrentProviderLock, const int32& CurrentReadProviderLockCount, const int32& CurrentWriteAllocationsProviderLockCount) const;
+	void ReadAccessCheck(const FProviderLock* CurrentProviderLock, const int32& CurrentReadProviderLockCount, const int32& CurrentWriteProviderLockCount) const;
 	void WriteAccessCheck(const int32& CurrentWriteProviderLockCount) const;
 
-	void BeginRead(FProviderLock*& CurrentProviderLock, int32& CurrentReadProviderLockCount, const int32& WriteAllocationsProviderLockCount);
+	void BeginRead(FProviderLock*& CurrentProviderLock, int32& CurrentReadProviderLockCount, const int32& WriteProviderLockCount);
 	void EndRead(FProviderLock*& CurrentProviderLock, int32& CurrentReadProviderLockCount);
 
-	void BeginWrite(FProviderLock*& CurrentProviderLock, const int32& CurrentReadProviderLockCount, int32& WriteAllocationsProviderLockCount);
-	void EndWrite(FProviderLock*& CurrentProviderLock, int32& WriteAllocationsProviderLockCount);
+	void BeginWrite(FProviderLock*& CurrentProviderLock, const int32& CurrentReadProviderLockCount, int32& WriteProviderLockCount);
+	void EndWrite(FProviderLock*& CurrentProviderLock, int32& WriteProviderLockCount);
 
 private:
 	FRWLock RWLock;
diff --git a/Engine/Source/Developer/TraceServices/Public/TraceServices/Model/AllocationsProvider.h b/Engine/Source/Developer/TraceServices/Public/TraceServices/Model/AllocationsProvider.h
--- a/Engine/Source/Developer/TraceServices/Public/TraceServices/Model/AllocationsProvider.h
+++ b/Engine/Source/Developer/TraceServices/Public/TraceServices/Model/AllocationsProvider.h
@@ -22,38 +22,6 @@ typedef uint32 TagIdType;
 class IAllocationsProvider : public IProvider
 {
 public:
-	struct TRACESERVICES_API FEditScopeLock
-	{
-		FEditScopeLock(const IAllocationsProvider& InAllocationsProvider)
-			: AllocationsProvider(InAllocationsProvider)
-		{
-			AllocationsProvider.BeginEdit();
-		}
-
-		~FEditScopeLock()
-		{
-			AllocationsProvider.EndEdit();
-		}
-
-		const IAllocationsProvider& AllocationsProvider;
-	};
-
-	struct TRACESERVICES_API FReadScopeLock
-	{
-		FReadScopeLock(const IAllocationsProvider& InAllocationsProvider)
-			: AllocationsProvider(InAllocationsProvider)
-		{
-			AllocationsProvider.BeginRead();
-		}
-
-		~FReadScopeLock()
-		{
-			AllocationsProvider.EndRead();
-		}
-
-		const IAllocationsProvider& AllocationsProvider;
-	};
-
 	// Allocation query rules.
 	// The enum uses the following naming convention:
 	//     A, B, C, D = time markers
diff --git a/Engine/Source/Developer/TraceServices/Public/TraceServices/Model/MetadataProvider.h b/Engine/Source/Developer/TraceServices/Public/TraceServices/Model/MetadataProvider.h
--- a/Engine/Source/Developer/TraceServices/Public/TraceServices/Model/MetadataProvider.h
+++ b/Engine/Source/Developer/TraceServices/Public/TraceServices/Model/MetadataProvider.h
@@ -16,39 +16,6 @@ namespace TraceServices
 
 class IMetadataProvider : public IProvider
 {
-public:
-	struct TRACESERVICES_API FEditScopeLock
-	{
-		FEditScopeLock(const IMetadataProvider& InMetadataProvider)
-			: MetadataProvider(InMetadataProvider)
-		{
-			MetadataProvider.BeginEdit();
-		}
-
-		~FEditScopeLock()
-		{
-			MetadataProvider.EndEdit();
-		}
-
-		const IMetadataProvider& MetadataProvider;
-	};
-
-	struct TRACESERVICES_API FReadScopeLock
-	{
-		FReadScopeLock(const IMetadataProvider& InMetadataProvider)
-			: MetadataProvider(InMetadataProvider)
-		{
-			MetadataProvider.BeginRead();
-		}
-
-		~FReadScopeLock()
-		{
-			MetadataProvider.EndRead();
-		}
-
-		const IMetadataProvider& MetadataProvider;
-	};
-
 public:
 	virtual ~IMetadataProvider() = default;
 
