TasksInsights: removed `NestedAdded` trace as `SubsequentAdded` is enough to establish both prerequisites and nested tasks relation. This also adds a new task category - a parent task - a task that depends on a nested task.
Parent tasks visualisation comes in the following CL

diff --git a/Engine/Source/Developer/TraceInsights/Private/Insights/TaskGraphProfiler/TaskGraphProfilerManager.cpp b/Engine/Source/Developer/TraceInsights/Private/Insights/TaskGraphProfiler/TaskGraphProfilerManager.cpp
--- a/Engine/Source/Developer/TraceInsights/Private/Insights/TaskGraphProfiler/TaskGraphProfilerManager.cpp
+++ b/Engine/Source/Developer/TraceInsights/Private/Insights/TaskGraphProfiler/TaskGraphProfilerManager.cpp
@@ -721,31 +721,6 @@ double FTaskGraphProfilerManager::GetRelationsOnCriticalPathAscendingRec(const T
 		}
 	}
 
-	if (Task->ParentOfNestedTask.IsValid())
-	{
-		const TraceServices::FTaskInfo* ParentTaskInfo = TasksProvider->TryGetTask(Task->ParentOfNestedTask->RelativeId);
-
-		if (ParentTaskInfo != nullptr)
-		{
-			TArray<FTaskGraphRelation> AscendingRelations;
-
-			double ChainDuration = GetRelationsOnCriticalPathAscendingRec(ParentTaskInfo, TasksProvider, AscendingRelations);
-			if (ChainDuration > MaxChainDuration)
-			{
-				MaxChainDuration = ChainDuration;
-				NextTaskInChain = ParentTaskInfo;
-				// Remove the relations from the shorter branch.
-				if (Relations.Num() > InitialRelationNum)
-				{
-					Relations.RemoveAt(InitialRelationNum, Relations.Num() - InitialRelationNum, false);
-				}
-
-				Relations.Append(AscendingRelations);
-			}
-			AscendingRelations.Empty(AscendingRelations.Max());
-		}
-	}
-
 	if (NextTaskInChain)
 	{
 		int32 SourceDepth = GetDepthOfTaskExecution(NextTaskInChain->StartedTimestamp, NextTaskInChain->FinishedTimestamp, NextTaskInChain->StartedThreadId);
diff --git a/Engine/Source/Developer/TraceServices/Private/Analyzers/TasksAnalysis.cpp b/Engine/Source/Developer/TraceServices/Private/Analyzers/TasksAnalysis.cpp
--- a/Engine/Source/Developer/TraceServices/Private/Analyzers/TasksAnalysis.cpp
+++ b/Engine/Source/Developer/TraceServices/Private/Analyzers/TasksAnalysis.cpp
@@ -6,6 +6,7 @@
 #include "Containers/UnrealString.h"
 #include "HAL/LowLevelMemTracker.h"
 #include "Model/TasksProfilerPrivate.h"
+#include "Common/Utils.h"
 
 namespace TraceServices
 {
@@ -103,10 +104,12 @@ bool FTasksAnalyzer::OnEvent(uint16 RouteId, EStyle Style, const FOnEventContext
 		}
 		case RouteId_NestedAdded:
 		{
-			double Timestamp = Context.EventTime.AsSeconds(EventData.GetValue<uint64>("Timestamp"));
-			TaskTrace::FId TaskId = EventData.GetValue<TaskTrace::FId>("TaskId");
-			TaskTrace::FId NestedId = EventData.GetValue<TaskTrace::FId>("NestedId");
-			TasksProvider.NestedAdded(TaskId, NestedId, Timestamp, ThreadId);
+			static bool bLogged = false;
+			if (!bLogged)
+			{
+				UE_LOG(LogTraceServices, Log, TEXT("An old TaskTrace format detected. Nested tasks will be ignored"));
+				bLogged = true;
+			}
 			break;
 		}
 		case RouteId_Finished:
diff --git a/Engine/Source/Developer/TraceServices/Private/Analyzers/TasksAnalysis.h b/Engine/Source/Developer/TraceServices/Private/Analyzers/TasksAnalysis.h
--- a/Engine/Source/Developer/TraceServices/Private/Analyzers/TasksAnalysis.h
+++ b/Engine/Source/Developer/TraceServices/Private/Analyzers/TasksAnalysis.h
@@ -26,6 +26,7 @@ namespace TraceServices
 			RouteId_Scheduled,
 			RouteId_SubsequentAdded,
 			RouteId_Started,
+			// `NestedAdded` was removed from `TaskTrace` but left here to avoid crashing TasksInsights when used with an old trace file
 			RouteId_NestedAdded,
 			RouteId_Finished,
 			RouteId_Completed,
diff --git a/Engine/Source/Developer/TraceServices/Private/Model/TasksProfiler.cpp b/Engine/Source/Developer/TraceServices/Private/Model/TasksProfiler.cpp
--- a/Engine/Source/Developer/TraceServices/Private/Model/TasksProfiler.cpp
+++ b/Engine/Source/Developer/TraceServices/Private/Model/TasksProfiler.cpp
@@ -117,7 +117,6 @@ namespace TraceServices
 			return;
 		}
 
-		Task->Id = TaskId;
 		Task->CreatedTimestamp = Timestamp;
 		Task->CreatedThreadId = ThreadId;
 	}
@@ -137,9 +136,8 @@ namespace TraceServices
 
 		checkf(Task->LaunchedTimestamp == FTaskInfo::InvalidTimestamp, TEXT("%d"), TaskId);
 			
-		if (Task->Id == TaskTrace::InvalidId) // created and launched in one go
+		if (Task->CreatedTimestamp == FTaskInfo::InvalidTimestamp) // created and launched in one go
 		{
-			Task->Id = TaskId;
 			Task->CreatedTimestamp = Timestamp;
 			Task->CreatedThreadId = ThreadId;
 		}
@@ -192,16 +190,31 @@ namespace TraceServices
 		FTaskInfo* Task = TryGetOrCreateTask(TaskId);
 		if (Task == nullptr)
 		{
-			UE_LOG(LogTraceServices, Log, TEXT("SubsequentAdded(TaskId %d, SubsequentId %d, Timestamp %.6f) skipped"), TaskId, SubsequentId, Timestamp);
+			UE_LOG(LogTraceServices, Log, TEXT("SubsequentAdded(TaskId %d, SubsequentId %d, Timestamp %.6f) skipped, task %d doesn't exists"), TaskId, SubsequentId, Timestamp, TaskId);
 			return;
 		}
 
-		Task->Id = TaskId;
-
-		AddRelative(TEXT("Subsequent"), TaskId, &FTaskInfo::Subsequents, SubsequentId, Timestamp, ThreadId);
+		FTaskInfo* SubsequentTask = TryGetOrCreateTask(SubsequentId);
+		if (SubsequentTask == nullptr)
+		{
+			UE_LOG(LogTraceServices, Log, TEXT("SubsequentAdded(TaskId %d, SubsequentId %d, Timestamp %.6f) skipped, subsequent task %d doesn't exists"), TaskId, SubsequentId, Timestamp, SubsequentId);
+			return;
+		}
 
-		// make a backward link from the subsequent task to this task (prerequisite)
-		AddRelative(TEXT("Prerequisite"), SubsequentId, &FTaskInfo::Prerequisites, TaskId, Timestamp, ThreadId);
+		// if a subsequent is added before this task execution started, it's a subsequent that have this task as a (execution) prerequisite.
+		// otherwise, it's a subsequent that have this task as a nested task (a completion prerequisite)
+		if (SubsequentTask->StartedTimestamp == FTaskInfo::InvalidTimestamp || SubsequentTask->StartedTimestamp >= Timestamp)
+		{
+			AddRelative(TEXT("Subsequent"), TaskId, &FTaskInfo::Subsequents, SubsequentId, Timestamp, ThreadId);
+			// make a backward link from the subsequent task to this task (prerequisite)
+			AddRelative(TEXT("Prerequisite"), SubsequentId, &FTaskInfo::Prerequisites, TaskId, Timestamp, ThreadId);
+		}
+		else
+		{
+			AddRelative(TEXT("Parent"), TaskId, &FTaskInfo::ParentTasks, SubsequentId, Timestamp, ThreadId);
+			// make a backward link from the parent task to this nested task
+			AddRelative(TEXT("Nested"), SubsequentId, &FTaskInfo::NestedTasks, TaskId, Timestamp, ThreadId);
+		}
 	}
 
 	void FTasksProvider::TaskStarted(TaskTrace::FId TaskId, double Timestamp, uint32 ThreadId)
@@ -242,19 +255,6 @@ namespace TraceServices
 		TaskLatencyCounter->SetValue(Timestamp, Task->StartedTimestamp - Task->ScheduledTimestamp);
 	}
 
-	void FTasksProvider::NestedAdded(TaskTrace::FId TaskId, TaskTrace::FId NestedId, double Timestamp, uint32 ThreadId)
-	{
-		InitTaskIdToIndexConversion(TaskId);
-
-		AddRelative(TEXT("Nested"), TaskId, &FTaskInfo::NestedTasks, NestedId, Timestamp, ThreadId);
-
-		FTaskInfo* Task = TryGetTask(NestedId);
-		if (Task != nullptr)
-		{
-			Task->ParentOfNestedTask = MakeUnique<FTaskInfo::FRelationInfo>(TaskId, Timestamp, ThreadId);
-		}
-	}
-
 	void FTasksProvider::TaskFinished(TaskTrace::FId TaskId, double Timestamp)
 	{
 		InitTaskIdToIndexConversion(TaskId);
@@ -371,6 +371,7 @@ namespace TraceServices
 		if (TaskIndex >= Tasks.Num())
 		{
 			Tasks.AddDefaulted(TaskIndex - Tasks.Num() + 1);
+			Tasks[TaskIndex].Id = TaskId;
 		}
 
 		return &Tasks[TaskIndex];
diff --git a/Engine/Source/Developer/TraceServices/Private/Model/TasksProfilerPrivate.h b/Engine/Source/Developer/TraceServices/Private/Model/TasksProfilerPrivate.h
--- a/Engine/Source/Developer/TraceServices/Private/Model/TasksProfilerPrivate.h
+++ b/Engine/Source/Developer/TraceServices/Private/Model/TasksProfilerPrivate.h
@@ -27,7 +27,6 @@ namespace TraceServices
 		void TaskScheduled(TaskTrace::FId TaskId, double Timestamp, uint32 ThreadId);
 		void SubsequentAdded(TaskTrace::FId TaskId, TaskTrace::FId SubsequentId, double Timestamp, uint32 ThreadId);
 		void TaskStarted(TaskTrace::FId TaskId, double Timestamp, uint32 ThreadId);
-		void NestedAdded(TaskTrace::FId TaskId, TaskTrace::FId NestedId, double Timestamp, uint32 ThreadId);
 		void TaskFinished(TaskTrace::FId TaskId, double Timestamp);
 		void TaskCompleted(TaskTrace::FId TaskId, double Timestamp, uint32 ThreadId);
 		void TaskDestroyed(TaskTrace::FId TaskId, double Timestamp, uint32 ThreadId);
diff --git a/Engine/Source/Developer/TraceServices/Public/TraceServices/Model/TasksProfiler.h b/Engine/Source/Developer/TraceServices/Public/TraceServices/Model/TasksProfiler.h
--- a/Engine/Source/Developer/TraceServices/Public/TraceServices/Model/TasksProfiler.h
+++ b/Engine/Source/Developer/TraceServices/Public/TraceServices/Model/TasksProfiler.h
@@ -65,10 +65,10 @@ namespace TraceServices
 		TArray<FRelationInfo> Prerequisites;
 		// other tasks that wait for this task completion before they'll start execution
 		TArray<FRelationInfo> Subsequents;
+		// other tasks that have this task as a nested
+		TArray<FRelationInfo> ParentTasks;
 		// the task is completed only after all nested tasks are completed
 		TArray<FRelationInfo> NestedTasks;
-		// The parent task. Is valid only for nested tasks
-		TUniquePtr<FRelationInfo> ParentOfNestedTask;
 	};
 
 	struct FWaitingForTasks
diff --git a/Engine/Source/Runtime/Core/Private/Async/TaskTrace.cpp b/Engine/Source/Runtime/Core/Private/Async/TaskTrace.cpp
--- a/Engine/Source/Runtime/Core/Private/Async/TaskTrace.cpp
+++ b/Engine/Source/Runtime/Core/Private/Async/TaskTrace.cpp
@@ -46,12 +46,6 @@ namespace TaskTrace
 		UE_TRACE_EVENT_FIELD(uint64, TaskId)
 	UE_TRACE_EVENT_END()
 
-	UE_TRACE_EVENT_BEGIN(TaskTrace, NestedAdded)
-		UE_TRACE_EVENT_FIELD(uint64, Timestamp)
-		UE_TRACE_EVENT_FIELD(uint64, TaskId)
-		UE_TRACE_EVENT_FIELD(uint64, NestedId)
-	UE_TRACE_EVENT_END()
-
 	UE_TRACE_EVENT_BEGIN(TaskTrace, Finished)
 		UE_TRACE_EVENT_FIELD(uint64, Timestamp)
 		UE_TRACE_EVENT_FIELD(uint64, TaskId)
@@ -158,18 +152,6 @@ namespace TaskTrace
 			<< Started.TaskId(TaskId);
 	}
 
-	void NestedAdded(FId TaskId, FId NestedId)
-	{
-		check(bGTaskTraceInitialized);
-		check(TaskId != InvalidId);
-		check(NestedId != InvalidId);
-
-		UE_TRACE_LOG(TaskTrace, NestedAdded, TaskChannel)
-			<< NestedAdded.Timestamp(FPlatformTime::Cycles64())
-			<< NestedAdded.TaskId(TaskId)
-			<< NestedAdded.NestedId(NestedId);
-	}
-
 	void Finished(FId TaskId)
 	{
 		check(bGTaskTraceInitialized);
diff --git a/Engine/Source/Runtime/Core/Public/Async/TaskGraphInterfaces.h b/Engine/Source/Runtime/Core/Public/Async/TaskGraphInterfaces.h
--- a/Engine/Source/Runtime/Core/Public/Async/TaskGraphInterfaces.h
+++ b/Engine/Source/Runtime/Core/Public/Async/TaskGraphInterfaces.h
@@ -1003,7 +1003,7 @@ public:
 	{
 		checkThreadGraph(!IsComplete()); // it is not legal to add a DontCompleteUntil after the event has been completed. Basically, this is only legal within a task function.
 		new (EventsToWaitFor) FGraphEventRef(EventToWaitFor);
-		TaskTrace::NestedAdded(GetTraceId(), EventToWaitFor->GetTraceId());
+		TaskTrace::SubsequentAdded(EventToWaitFor->GetTraceId(), GetTraceId());
 	}
 
 	/**
diff --git a/Engine/Source/Runtime/Core/Public/Async/TaskTrace.h b/Engine/Source/Runtime/Core/Public/Async/TaskTrace.h
--- a/Engine/Source/Runtime/Core/Public/Async/TaskTrace.h
+++ b/Engine/Source/Runtime/Core/Public/Async/TaskTrace.h
@@ -37,7 +37,6 @@ namespace TaskTrace
 	void CORE_API Scheduled(FId TaskId);
 	void CORE_API SubsequentAdded(FId TaskId, FId SubsequentId);
 	void CORE_API Started(FId TaskId);
-	void CORE_API NestedAdded(FId TaskId, FId NestedId);
 	void CORE_API Finished(FId TaskId);
 	void CORE_API Completed(FId TaskId);
 	void CORE_API Destroyed(FId TaskId);
@@ -68,7 +67,6 @@ namespace TaskTrace
 	inline void Scheduled(FId TaskId) {}
 	inline void SubsequentAdded(FId TaskId, FId SubsequentId) {}
 	inline void Started(FId TaskId) {}
-	inline void NestedAdded(FId TaskId, FId NestedId) {}
 	inline void Finished(FId TaskId) {}
 	inline void Completed(FId TaskId) {}
 	inline void Destroyed(FId TaskId) {}
diff --git a/Engine/Source/Runtime/Core/Public/Experimental/Coroutine/Coroutine.h b/Engine/Source/Runtime/Core/Public/Experimental/Coroutine/Coroutine.h
--- a/Engine/Source/Runtime/Core/Public/Experimental/Coroutine/Coroutine.h
+++ b/Engine/Source/Runtime/Core/Public/Experimental/Coroutine/Coroutine.h
@@ -288,7 +288,7 @@ namespace CoroTask_Detail
 			{
 				TraceId = TaskTrace::GenerateTaskId();
 				TaskTrace::Launched(TraceId, GetDebugName(), true, ENamedThreads::Type::AnyThread);
-				TaskTrace::NestedAdded(OldTraceId, TraceId);
+				TaskTrace::SubsequentAdded(TraceId, OldTraceId);
 				if(Prerequisite != nullptr)
 				{
 					TaskTrace::SubsequentAdded(Prerequisite->TraceId, TraceId);
diff --git a/Engine/Source/Runtime/Core/Public/Tasks/TaskPrivate.h b/Engine/Source/Runtime/Core/Public/Tasks/TaskPrivate.h
--- a/Engine/Source/Runtime/Core/Public/Tasks/TaskPrivate.h
+++ b/Engine/Source/Runtime/Core/Public/Tasks/TaskPrivate.h
@@ -256,12 +256,8 @@ namespace UE::Tasks
 			// returns false if the task is already completed and the subsequent wasn't added
 			bool AddSubsequent(FTaskBase& Subsequent)
 			{
-				bool bSucceeded = Subsequents.PushIfNotClosed(&Subsequent);
-				if (bSucceeded)
-				{
-					TaskTrace::SubsequentAdded(GetTraceId(), Subsequent.GetTraceId());
-				}
-				return bSucceeded;
+				TaskTrace::SubsequentAdded(GetTraceId(), Subsequent.GetTraceId()); // doesn't matter if we suceeded below, we need to record task dependency
+				return Subsequents.PushIfNotClosed(&Subsequent);
 			}
 
 			// A piped task is executed after the previous task from this pipe is completed. Tasks from the same pipe are not executed
@@ -420,7 +416,6 @@ namespace UE::Tasks
 
 				if (Nested.AddSubsequent(*this)) // "release" memory order
 				{
-					TaskTrace::NestedAdded(GetTraceId(), Nested.GetTraceId());
 					Nested.AddRef(); // keep it alive as we store it in `Prerequisites` and we can need it to try to retract it. it's released on closing the task
 					Prerequisites.Push(&Nested);
 				}
