Adds metadata schemas

This change adds schema support for metadata. Registering a metadata type now requires the user to build a out a schema for the data. This also allows consumers to deal with the data without knowing the underlying metadata fields.

diff --git a/Engine/Source/Developer/TraceInsights/Private/Insights/MemoryProfiler/Widgets/SMemAllocTableTreeView.cpp b/Engine/Source/Developer/TraceInsights/Private/Insights/MemoryProfiler/Widgets/SMemAllocTableTreeView.cpp
--- a/Engine/Source/Developer/TraceInsights/Private/Insights/MemoryProfiler/Widgets/SMemAllocTableTreeView.cpp
+++ b/Engine/Source/Developer/TraceInsights/Private/Insights/MemoryProfiler/Widgets/SMemAllocTableTreeView.cpp
@@ -310,15 +310,20 @@ void SMemAllocTableTreeView::UpdateQuery(TraceServices::IAllocationsProvider::EQ
 	const TraceServices::IMetadataProvider* MetadataProvider = TraceServices::ReadMetadataProvider(*Session.Get());
 
 	uint16 AssetMetadataType = 0;
+	const TraceServices::FMetadataSchema* Schema = nullptr;
 	if (MetadataProvider)
 	{
 		TraceServices::FProviderReadScopeLock MetadataProviderReadLock(*MetadataProvider);
 		AssetMetadataType = MetadataProvider->GetRegisteredMetadataType(TEXT("Asset"));
-		if (MetadataProvider->GetRegisteredMetadataName(AssetMetadataType) == nullptr)
+		if (AssetMetadataType == TraceServices::IMetadataProvider::InvalidMetadataId)
 		{
-			// If AssetMetadataType is not valid then we do not need to furhter check the Asset metadata for each allocation.
+			// If AssetMetadataType is not valid then we do not need to further check the Asset metadata for each allocation.
 			MetadataProvider = nullptr;
 		}
+		else
+		{
+			Schema = MetadataProvider->GetRegisteredMetadataSchema(AssetMetadataType);
+		}
 	}
 
 	const TraceServices::IDefinitionProvider* DefinitionProvider = TraceServices::ReadDefinitionProvider(*Session.Get());
@@ -394,24 +399,23 @@ void SMemAllocTableTreeView::UpdateQuery(TraceServices::IAllocationsProvider::EQ
 					Alloc.Tag = Provider.GetTagName(Allocation->GetTag());
 
 					Alloc.Asset = nullptr;
-					if (MetadataProvider && DefinitionProvider)
+					const uint32 MetadataId = Allocation->GetMetadataId();
+					if (MetadataId != TraceServices::IMetadataProvider::InvalidMetadataId && MetadataProvider && DefinitionProvider)
 					{
 						TraceServices::FProviderReadScopeLock MetadataProviderReadLock(*MetadataProvider);
-						MetadataProvider->EnumerateMetadata(Allocation->GetThreadId(), Allocation->GetMetadataId(),
-							[AssetMetadataType, &Alloc, DefinitionProvider](uint32 StackDepth, uint16 Type, const void* Data, uint32 Size) -> bool
+						MetadataProvider->EnumerateMetadata(Allocation->GetThreadId(), MetadataId,
+							[AssetMetadataType, Schema, &Alloc, DefinitionProvider](uint32 StackDepth, uint16 Type, const void* Data, uint32 Size) -> bool
 							{
 								if (Type == AssetMetadataType)
 								{
-									////TODO: TraceServices::IDefinitionProvider::FReadScopeLock DefinitionProviderReadLock(*DefinitionProvider);
-									//DefinitionProvider->BeginRead();
-									//const UE::Trace::TEventRef<UE::Trace::FEventRef32>& AssetNameReference = ((const UE::Trace::TEventRef<UE::Trace::FEventRef32>*)Data)[0];
-									////const UE::Trace::TEventRef<UE::Trace::FEventRef32>& ClassNameReference = ((const UE::Trace::TEventRef<UE::Trace::FEventRef32>*)Data)[1];
-									//const TraceServices::FStringDefinition* AssetNameDef = DefinitionProvider->Get<TraceServices::FStringDefinition>(AssetNameReference);
-									//if (AssetNameDef)
-									//{
-									//	Alloc.Asset = AssetNameDef->Display;
-									//}
-									//DefinitionProvider->EndRead();
+									TraceServices::IDefinitionProvider::FReadScopeLock DefinitionProviderReadLock(*DefinitionProvider);
+									const auto Reader = Schema->Reader();
+									const auto AssetNameRef = Reader.GetValueAs<UE::Trace::FEventRef32>((uint8*)Data, 0);
+									const auto AssetName = DefinitionProvider->Get<TraceServices::FStringDefinition>(*AssetNameRef);
+									if (AssetName)
+									{
+										Alloc.Asset = AssetName->Display;
+									}
 									return false;
 								}
 								return true;
diff --git a/Engine/Source/Developer/TraceServices/Private/Model/MetadataProvider.cpp b/Engine/Source/Developer/TraceServices/Private/Model/MetadataProvider.cpp
--- a/Engine/Source/Developer/TraceServices/Private/Model/MetadataProvider.cpp
+++ b/Engine/Source/Developer/TraceServices/Private/Model/MetadataProvider.cpp
@@ -77,6 +77,7 @@ void FMetadataProviderLock::EndWrite()
 
 FMetadataProvider::FMetadataProvider(IAnalysisSession& InSession)
 	: Session(InSession)
+	, RegisteredTypes(Session.GetLinearAllocator(), 16)
 	, MetadataStore(Session.GetLinearAllocator(), 1024)
 {
 }
@@ -123,14 +124,13 @@ FMetadataProvider::~FMetadataProvider()
 
 ////////////////////////////////////////////////////////////////////////////////////////////////////
 
-uint16 FMetadataProvider::RegisterMetadataType(const TCHAR* Name, uint32 Size)
+uint16 FMetadataProvider::RegisterMetadataType(const TCHAR* Name, const FMetadataSchema& Schema)
 {
 	Lock.WriteAccessCheck();
 
 	check(RegisteredTypes.Num() <= MaxMetadataTypeId);
 	const uint16 Type = RegisteredTypes.Num();
-
-	RegisteredTypes.Add({ Name, Size });
+	RegisteredTypes.EmplaceBack(Schema);
 	RegisteredTypesMap.Add(Name, Type);
 
 	return Type;
@@ -138,7 +138,7 @@ uint16 FMetadataProvider::RegisterMetadataType(const TCHAR* Name, uint32 Size)
 
 ////////////////////////////////////////////////////////////////////////////////////////////////////
 
-uint16 FMetadataProvider::GetRegisteredMetadataType(const TCHAR* Name) const
+uint16 FMetadataProvider::GetRegisteredMetadataType(FName Name) const
 {
 	Lock.ReadAccessCheck();
 
@@ -148,12 +148,22 @@ uint16 FMetadataProvider::GetRegisteredMetadataType(const TCHAR* Name) const
 
 ////////////////////////////////////////////////////////////////////////////////////////////////////
 
-const TCHAR* FMetadataProvider::GetRegisteredMetadataName(uint16 Type) const
+FName FMetadataProvider::GetRegisteredMetadataName(uint16 Type) const
 {
 	Lock.ReadAccessCheck();
 
 	const int32 Index = (int32)Type;
-	return Index < RegisteredTypes.Num() ? RegisteredTypes[Index].Name : nullptr;
+	const FName* Name = RegisteredTypesMap.FindKey(Index);
+	return Name ? *Name : FName();
+}
+
+////////////////////////////////////////////////////////////////////////////////////////////////////
+const FMetadataSchema* FMetadataProvider::GetRegisteredMetadataSchema(uint16 Type) const
+{
+	Lock.ReadAccessCheck();
+	
+	const int32 Index = (int32)Type;
+	return Index < RegisteredTypes.Num() ? &RegisteredTypes[Index] : nullptr;
 }
 
 ////////////////////////////////////////////////////////////////////////////////////////////////////
@@ -337,7 +347,6 @@ uint32 FMetadataProvider::PinAndGetId(uint32 ThreadId)
 				Entry.StoreIndex = StackEntry.StoreIndex;
 				Entry.Type = StackEntry.StoreEntry.Type;
 				Entry.StackSize = StackIndex + 1;
-				Entry.Unused = 0;
 			}
 
 			return CurrentStack.Top().PinnedId;
@@ -353,7 +362,6 @@ uint32 FMetadataProvider::PinAndGetId(uint32 ThreadId)
 		Entry.StoreIndex = StackEntry.StoreIndex;
 		Entry.Type = StackEntry.StoreEntry.Type;
 		Entry.StackSize = StackIndex + 1;
-		Entry.Unused = 0;
 	}
 
 	return CurrentStack.Top().PinnedId;
diff --git a/Engine/Source/Developer/TraceServices/Private/Model/MetadataProvider.h b/Engine/Source/Developer/TraceServices/Private/Model/MetadataProvider.h
--- a/Engine/Source/Developer/TraceServices/Private/Model/MetadataProvider.h
+++ b/Engine/Source/Developer/TraceServices/Private/Model/MetadataProvider.h
@@ -3,8 +3,6 @@
 #pragma once
 
 #include "Common/PagedArray.h"
-#include "Containers/Map.h"
-#include "HAL/CriticalSection.h"
 #include "TraceServices/Model/MetadataProvider.h"
 
 namespace TraceServices
@@ -30,28 +28,20 @@ public:
 private:
 	FRWLock RWLock;
 };
-
+	
 ////////////////////////////////////////////////////////////////////////////////////////////////////
 
 class FMetadataProvider : public IMetadataProvider
 {
 public:
 	static constexpr uint32 MaxMetadataSize = 0xFFFF;
-	static constexpr uint32 InvalidMetadataId = 0xFFFFFFFF;
-	static constexpr uint16 InvalidMetadataType = 0xFFFF;
 
 private:
 	static constexpr uint32 MaxInlinedMetadataSize = 12;
 	static constexpr uint32 MaxMetadataTypeId = 0xFFFF;
-	static constexpr int32 MaxMetadataStackSize = 0xFF;
+	static constexpr int32 MaxMetadataStackSize = 0xFFFF;
 	static constexpr uint32 InvalidMetadataStoreIndex = 0xFFFFFFFF;
 
-	struct FMetadataType
-	{
-		const TCHAR* Name;
-		uint32 Size; // 0 for variable size data
-	};
-
 	struct FMetadataStoreEntry
 	{
 		union
@@ -77,8 +67,7 @@ private:
 	{
 		uint32 StoreIndex;
 		uint16 Type;
-		uint8 StackSize;
-		uint8 Unused;
+		uint16 StackSize;
 	};
 	static_assert(sizeof(FMetadataEntry) == 8, "sizeof(FMetadataEntry)");
 
@@ -106,7 +95,7 @@ public:
 	//////////////////////////////////////////////////
 	// Edit operations
 
-	uint16 RegisterMetadataType(const TCHAR* Name, uint32 FixedSize = 0);
+	uint16 RegisterMetadataType(const TCHAR* Name, const FMetadataSchema& Schema);
 
 	void PushScopedMetadata(uint32 ThreadId, uint16 Type, void* Data, uint32 Size);
 	void PopScopedMetadata(uint32 ThreadId, uint16 Type);
@@ -117,8 +106,9 @@ public:
 	//////////////////////////////////////////////////
 	// Read operations
 
-	virtual uint16 GetRegisteredMetadataType(const TCHAR* Name) const override;
-	virtual const TCHAR* GetRegisteredMetadataName(uint16 Type) const override;
+	virtual uint16 GetRegisteredMetadataType(FName Name) const override;
+	virtual FName GetRegisteredMetadataName(uint16 Type) const override;
+	virtual const FMetadataSchema* GetRegisteredMetadataSchema(uint16) const override;
 
 	virtual uint32 GetMetadataStackSize(uint32 InThreadId, uint32 InMetadataId) const override;
 	virtual bool GetMetadata(uint32 InThreadId, uint32 InMetadataId, uint32 InStackDepth, uint16& OutType, const void*& OutData, uint32& OutSize) const override;
@@ -133,9 +123,8 @@ private:
 	IAnalysisSession& Session;
 
 	mutable FMetadataProviderLock Lock;
-
-	TArray<FMetadataType> RegisteredTypes;
-	TMap<const TCHAR*, uint16> RegisteredTypesMap;
+	TPagedArray<FMetadataSchema> RegisteredTypes;
+	TMap<FName, uint16> RegisteredTypesMap;
 
 	TPagedArray<FMetadataStoreEntry> MetadataStore; // stores individual metadata values
 
diff --git a/Engine/Source/Developer/TraceServices/Public/TraceServices/Model/MetadataProvider.h b/Engine/Source/Developer/TraceServices/Public/TraceServices/Model/MetadataProvider.h
--- a/Engine/Source/Developer/TraceServices/Public/TraceServices/Model/MetadataProvider.h
+++ b/Engine/Source/Developer/TraceServices/Public/TraceServices/Model/MetadataProvider.h
@@ -4,19 +4,274 @@
 
 #include "AnalysisSession.h"
 #include "CoreTypes.h"
-#include "Templates/UniquePtr.h"
-#include "ProfilingDebugging/MemoryTrace.h"
-
-#include <new>
+#include "UObject/NameTypes.h"
+#include "Internationalization/Text.h"
+#include <Runtime/Core/Public/Memory/SharedBuffer.h>
 
 namespace TraceServices
 {
 
+////////////////////////////////////////////////////////////////////////////////////////////////////
+class FMetadataSchema
+{
+public:
+	enum class EFieldType : uint8
+	{
+		Integer,
+		SignedInteger,
+		FloatingPoint,
+		WideStringPtr,
+		Reference,
+	};
+	
+	struct FField
+	{
+		FName			Name;
+		EFieldType		Type;
+		uint8			Offset;
+		uint8			Size;
+	};
+
+	/**
+	 * Helper class that assist in defining a schema.
+	 */
+	struct FBuilder
+	{
+		FBuilder(FField* InBuffer, uint8 InMaxFieldCount)
+			: Buffer(InBuffer), MaxFieldCount(InMaxFieldCount), CurrentFieldCount(0), CurrentOffset(0)
+		{
+		}
+
+		template<typename T>
+		FBuilder& AddField(FName FieldName, EFieldType Type)
+		{
+			check(CurrentFieldCount < MaxFieldCount);
+			FField* Field = &Buffer[CurrentFieldCount++];
+			Field->Name = FieldName;
+			Field->Type = Type;
+			Field->Offset = CurrentOffset;
+			Field->Size = sizeof(T);
+			CurrentOffset += sizeof(T);
+			return *this;
+		}
+
+		void Finish()
+		{
+		}
+		
+	private:
+		FField* Buffer;
+		uint8 MaxFieldCount;
+		uint8 CurrentFieldCount;
+		uint8 CurrentOffset;
+	};
+
+	/**
+	 * Helper class that helps writing data based on a schema.
+	 */
+	struct FWriter
+	{
+		FWriter(const FField* InBuffer, uint8 InFieldCount)
+			: Buffer(InBuffer), FieldCount(InFieldCount), Offset(0)
+		{
+		}
+
+		template<typename ArrayType>
+		void WriteField(uint8 Index, const void* Src, uint8 Size, ArrayType& Dst)
+		{
+			const FField& Field = Buffer[Index];
+			if (Dst.Num() < Field.Offset + Size)
+			{
+				Dst.AddZeroed(Size);
+			}
+			void* Dest = Dst.GetData() + Field.Offset;
+			FMemory::Memcpy(Dest, Src, Size);
+		}
+		
+	private:
+		const FField* Buffer;
+		uint8 FieldCount;
+		uint8 Offset;
+	};
+
+	/**
+	 * Helper class that helps reading fields from data based on a schema.
+	 */
+	struct FReader
+	{
+		FReader(const FField* InBuffer, uint8 InFieldCount)
+			: Buffer(InBuffer), FieldCount(InFieldCount)
+		{}
+
+		uint8 GetFieldCount() const { return FieldCount; }
+		const FField& GetFieldDesc(uint8 Index) const { return Buffer[Index]; }
+
+		FString GetString(uint8* MetaData, uint8 Index) const
+		{
+			const FField& Field = Buffer[Index];
+			switch(Field.Type)
+			{
+			case EFieldType::Integer:
+				 {
+					switch(Field.Size)
+					{
+					case 1: return FString::FromInt(*GetValue<uint8>(MetaData, Index));
+					case 2: return FString::FromInt(*GetValue<uint16>(MetaData, Index));
+					case 4: return FString::FromInt(*GetValue<uint32>(MetaData, Index));
+					default: return FString::FromInt(*GetValue<uint64>(MetaData, Index));
+					}
+				 }
+			case EFieldType::SignedInteger:
+				{
+					switch (Field.Size)
+					{
+					case 1: return FString::FromInt(*GetValue<int8>(MetaData, Index));
+					case 2: return FString::FromInt(*GetValue<int16>(MetaData, Index));
+					case 4: return FString::FromInt(*GetValue<int32>(MetaData, Index));
+					default: return FString::FromInt(*GetValue<int64>(MetaData, Index));
+					}
+				}
+			case EFieldType::FloatingPoint:
+				{
+					switch (Field.Size)
+					{
+					case 4: return FString::Format(TEXT("{0}"), {*GetValue<float>(MetaData, Index)});
+					default: return FString::Format(TEXT("{0}"), {*GetValue<double>(MetaData, Index)});
+					}
+				}
+			case EFieldType::WideStringPtr:
+				{
+					const TCHAR* String = *GetValue<TCHAR*>(MetaData, Index);
+					return FString(String);
+				}
+			case EFieldType::Reference:
+				{
+					switch(Field.Size)
+					{
+					case 1: { const auto* Ref = GetValueAs<UE::Trace::FEventRef8>(MetaData, Index); return FString::Format(TEXT("Ref({0},{1})"), {Ref->RefTypeId, Ref->Id});}
+					case 2: { const auto* Ref = GetValueAs<UE::Trace::FEventRef16>(MetaData, Index); return FString::Format(TEXT("Ref({0},{1})"), {Ref->RefTypeId, Ref->Id});}
+					case 4: { const auto* Ref = GetValueAs<UE::Trace::FEventRef32>(MetaData, Index); return FString::Format(TEXT("Ref({0},{1})"), {Ref->RefTypeId, Ref->Id});}
+					default: { const auto* Ref = GetValueAs<UE::Trace::FEventRef64>(MetaData, Index); return FString::Format(TEXT("Ref({0},{1})"), {Ref->RefTypeId, Ref->Id});}
+					}
+				}
+			}
+			checkNoEntry();
+			return FString();
+		}
+
+		FText GetText(uint8* MetaData, uint8 Index) const
+		{
+			const FField& Field = Buffer[Index];
+			switch(Field.Type)
+			{
+			case EFieldType::Integer:
+				 {
+					switch(Field.Size)
+					{
+					case 1: return FText::AsNumber(*GetValue<uint8>(MetaData, Index));
+					case 2: return FText::AsNumber(*GetValue<uint16>(MetaData, Index));
+					case 4: return FText::AsNumber(*GetValue<uint32>(MetaData, Index));
+					default: return FText::AsNumber(*GetValue<uint64>(MetaData, Index));
+					}
+				 }
+				 break;
+			case EFieldType::SignedInteger:
+				{
+					switch(Field.Size)
+					{
+					case 1: return FText::AsNumber(*GetValue<int8>(MetaData, Index));
+					case 2: return FText::AsNumber(*GetValue<int16>(MetaData, Index));
+					case 4: return FText::AsNumber(*GetValue<int32>(MetaData, Index));
+					default: return FText::AsNumber(*GetValue<int64>(MetaData, Index));
+					}
+				}
+			case EFieldType::FloatingPoint:
+				{
+					switch (Field.Size)
+					{
+					case 4: return FText::AsNumber(*GetValue<float>(MetaData, Index));
+					default: return FText::AsNumber(*GetValue<double>(MetaData, Index));
+					}
+				}
+			case EFieldType::WideStringPtr:
+				{
+					const TCHAR* String = *GetValue<TCHAR*>(MetaData, Index);
+					return FText::FromString(FString(String));
+				}
+				break;
+			case EFieldType::Reference:
+				{
+					static const FTextFormat Fmt = FTextFormat::FromString(TEXT("Ref({0},{1})"));
+					switch (Field.Size)
+					{
+					case 1: { const auto* Ref = GetValueAs<UE::Trace::FEventRef8>(MetaData, Index); return FText::Format(Fmt, Ref->RefTypeId, Ref->Id); }
+					case 2: { const auto* Ref = GetValueAs<UE::Trace::FEventRef16>(MetaData, Index); return FText::Format(Fmt, Ref->RefTypeId, Ref->Id); }
+					case 4: { const auto* Ref = GetValueAs<UE::Trace::FEventRef32>(MetaData, Index); return FText::Format(Fmt, Ref->RefTypeId, Ref->Id); }
+					default:  { const auto* Ref = GetValueAs<UE::Trace::FEventRef64>(MetaData, Index); return FText::Format(Fmt, Ref->RefTypeId, Ref->Id); }
+					}
+				}
+			}
+			checkNoEntry();
+			return FText();
+		}
+
+		template<typename T>
+		const T* GetValueAs(const uint8* Metadata, uint8 Index) const
+		{
+			return GetValue<T>(Metadata, Index);
+		}
+
+	private:
+		template<typename T>
+		const T* GetValue(const uint8* Metadata, uint8 Index) const
+		{
+			const FField& Field = Buffer[Index];
+			return (T*)(Metadata + Field.Offset);
+		}
+
+	private:
+		const FField* Buffer;
+		uint8 FieldCount;
+	};
+
+	explicit FMetadataSchema(uint8 InFieldCount)
+	{
+		Fields.AddZeroed(InFieldCount);
+	}
+
+	FMetadataSchema(const FMetadataSchema& InSchema)
+		: Fields(InSchema.Fields)
+	{}
+
+	FBuilder Builder() 
+	{
+		return FBuilder(Fields.GetData(), Fields.Num());
+	}
+
+	FWriter Writer() const
+	{
+		return FWriter(Fields.GetData(), Fields.Num());
+	}
+
+	FReader Reader() const
+	{
+		return FReader(Fields.GetData(), Fields.Num());
+	}
+
+private:
+	FMetadataSchema() = delete;
+	
+	TArray<FField> Fields;
+};
+	
 ////////////////////////////////////////////////////////////////////////////////////////////////////
 
 class IMetadataProvider : public IProvider
 {
 public:
+	static constexpr uint32 InvalidMetadataId = 0xFFFFFFFF;
+	static constexpr uint16 InvalidMetadataType = 0xFF;
+
 	virtual ~IMetadataProvider() = default;
 
 	virtual void BeginEdit() const = 0;
@@ -24,8 +279,9 @@ public:
 	virtual void BeginRead() const = 0;
 	virtual void EndRead() const = 0;
 
-	virtual uint16 GetRegisteredMetadataType(const TCHAR* Name) const = 0;
-	virtual const TCHAR* GetRegisteredMetadataName(uint16 Type) const = 0;
+	virtual uint16 GetRegisteredMetadataType(FName Name) const = 0;
+	virtual FName GetRegisteredMetadataName(uint16 Type) const = 0;
+	virtual const FMetadataSchema* GetRegisteredMetadataSchema(uint16) const = 0 ;
 
 	virtual uint32 GetMetadataStackSize(uint32 InThreadId, uint32 InMetadataId) const = 0;
 	virtual bool GetMetadata(uint32 InThreadId, uint32 InMetadataId, uint32 InStackDepth, uint16& OutType, const void*& OutData, uint32& OutSize) const = 0;
