Fix CIS error

diff --git a/Engine/Source/Developer/TraceServices/Public/TraceServices/Model/Definitions.h b/Engine/Source/Developer/TraceServices/Public/TraceServices/Model/Definitions.h
--- a/Engine/Source/Developer/TraceServices/Public/TraceServices/Model/Definitions.h
+++ b/Engine/Source/Developer/TraceServices/Public/TraceServices/Model/Definitions.h
@@ -1,6 +1,7 @@
 ﻿// Copyright Epic Games, Inc. All Rights Reserved.
 #pragma once
 
+#include "Containers/Map.h"
 #include "TraceServices/Model/AnalysisSession.h"
 #include "Trace/Analyzer.h"
 
