Metadata scopes

Adds runtime events and analysis for generic meta data scopes. Any scope emitted on the "Metadata" writer will be analyzed and accessible in the metadata provider. Currently one built-in metadata type is provided: "Asset", which mimics LLMs implementation of asset tracking.

diff --git a/Engine/Source/Developer/TraceAnalysis/Private/Analysis/Engine.cpp b/Engine/Source/Developer/TraceAnalysis/Private/Analysis/Engine.cpp
--- a/Engine/Source/Developer/TraceAnalysis/Private/Analysis/Engine.cpp
+++ b/Engine/Source/Developer/TraceAnalysis/Private/Analysis/Engine.cpp
@@ -888,6 +888,20 @@ const void* IAnalyzer::FEventData::GetReferenceValueImpl(const char* FieldName,
 	return nullptr;
 }
 
+////////////////////////////////////////////////////////////////////////////////
+const void* IAnalyzer::FEventData::GetReferenceValueImpl(uint32 FieldIndex, uint32& OutTypeUid) const
+{
+	const auto* Info = (const FEventDataInfo*)this;
+	const auto& Dispatch = Info->Dispatch;
+	const auto& Field = Dispatch.Fields[FieldIndex];
+	if (Field.RefUid)
+	{
+		OutTypeUid = Field.RefUid;
+		return (Info->Ptr + Field.Offset);
+	}
+	return nullptr;
+}
+
 // }}}
 
 
diff --git a/Engine/Source/Developer/TraceServices/Private/AnalysisService.cpp b/Engine/Source/Developer/TraceServices/Private/AnalysisService.cpp
--- a/Engine/Source/Developer/TraceServices/Private/AnalysisService.cpp
+++ b/Engine/Source/Developer/TraceServices/Private/AnalysisService.cpp
@@ -2,24 +2,24 @@
 
 #include "TraceServices/AnalysisService.h"
 #include "AnalysisServicePrivate.h"
-#include "Trace/Analyzer.h"
-#include "Trace/Analysis.h"
-#include "Trace/DataStream.h"
-#include "HAL/PlatformFile.h"
-#include "Analyzers/MiscTraceAnalysis.h"
-#include "Analyzers/LogTraceAnalysis.h"
-#include "Math/RandomStream.h"
 #include "ModuleServicePrivate.h"
-#include "Model/LogPrivate.h"
-#include "Model/FramesPrivate.h"
+#include "Analyzers/LogTraceAnalysis.h"
+#include "Analyzers/MiscTraceAnalysis.h"
+#include "Analyzers/StringsAnalyzer.h"
+#include "HAL/PlatformFile.h"
 #include "Model/BookmarksPrivate.h"
-#include "Model/ThreadsPrivate.h"
+#include "Model/Channel.h"
 #include "Model/CountersPrivate.h"
-#include "Model/NetProfilerProvider.h"
+#include "Model/DefinitionProvider.h"
+#include "Model/FramesPrivate.h"
+#include "Model/LogPrivate.h"
 #include "Model/MemoryPrivate.h"
-#include "Model/Channel.h"
-#include "Model/DiagnosticsPrivate.h"
+#include "Model/NetProfilerProvider.h"
 #include "Model/ScreenshotProviderPrivate.h"
+#include "Model/ThreadsPrivate.h"
+#include "Trace/Analysis.h"
+#include "Trace/Analyzer.h"
+#include "Trace/DataStream.h"
 
 namespace TraceServices
 {
@@ -247,6 +247,11 @@ TSharedPtr<const IAnalysisSession> FAnalysisService::StartAnalysis(uint32 TraceI
 	Session->AddAnalyzer(new FMiscTraceAnalyzer(*Session, *ThreadProvider, *BookmarkProvider, *LogProvider, *FrameProvider, *ChannelProvider, *ScreenshotProvider));
 	Session->AddAnalyzer(new FLogTraceAnalyzer(*Session, *LogProvider));
 
+	Session->AddAnalyzer(new FStringsAnalyzer(*Session));
+
+	FDefinitionProvider* DefProvider = new FDefinitionProvider(&Session.Get());
+	Session->AddProvider(FDefinitionProvider::ProviderName, DefProvider);
+
 	ModuleService.OnAnalysisBegin(*Session);
 
 	Session->Start();
diff --git a/Engine/Source/Developer/TraceServices/Private/Analyzers/MetadataAnalysis.cpp b/Engine/Source/Developer/TraceServices/Private/Analyzers/MetadataAnalysis.cpp
new file mode 100644
--- /dev/null
+++ b/Engine/Source/Developer/TraceServices/Private/Analyzers/MetadataAnalysis.cpp
@@ -0,0 +1,295 @@
+﻿// Copyright Epic Games, Inc. All Rights Reserved.
+#include "MetadataAnalysis.h"
+
+#include "CoreMinimal.h"
+#include "../Common/Utils.h"
+#include "Common/ProviderLock.h"
+#include "Containers/StringConv.h"
+#include "Model/DefinitionProvider.h"
+#include "Model/MetadataProvider.h"
+#include "Trace/Analyzer.h"
+#include "TraceServices/Model/Strings.h"
+
+namespace TraceServices
+{
+	
+////////////////////////////////////////////////////////////////////////////////////////////////////
+template<typename ArrayType>
+static void WriteData(ArrayType& OutData, const UE::Trace::IAnalyzer::FEventData& EventData, uint16 MetadataTypeId, const FMetadataProvider* MetadataProvider, IAnalysisSession& Session);
+	
+////////////////////////////////////////////////////////////////////////////////////////////////////
+FMetadataAnalysis::FMetadataAnalysis(IAnalysisSession& InSession, FMetadataProvider* InProvider)
+	: Session(InSession)
+	, MetadataProvider(InProvider)
+{
+}
+
+////////////////////////////////////////////////////////////////////////////////////////////////////
+FMetadataAnalysis::~FMetadataAnalysis()
+{
+}
+	
+////////////////////////////////////////////////////////////////////////////////////////////////////
+void FMetadataAnalysis::OnAnalysisBegin(const FOnAnalysisContext& Context)
+{
+	FInterfaceBuilder& Builder = Context.InterfaceBuilder;
+
+	// Some operations with the metadata stack
+	Builder.RouteEvent(RouteId_ClearScope, "MetadataStack", "ClearScope", true);
+	Builder.RouteEvent(RouteId_SaveStack, "MetadataStack", "SaveStack");
+	Builder.RouteEvent(RouteId_RestoreStack, "MetadataStack", "RestoreStack", true);
+
+	// Subscribe to all scopes on the "Metadata" logger. These are the actual meta data scopes.
+	Builder.RouteLoggerEvents(RouteId_Metascope, "Metadata", true);
+}
+
+////////////////////////////////////////////////////////////////////////////////////////////////////
+void FMetadataAnalysis::OnAnalysisEnd()
+{
+}
+
+////////////////////////////////////////////////////////////////////////////////////////////////////
+bool FMetadataAnalysis::OnEvent(uint16 RouteId, EStyle Style, const FOnEventContext& Context)
+{
+	FProviderEditScopeLock _(*MetadataProvider);
+	switch (RouteId)
+	{
+	case RouteId_ClearScope:
+		{
+			//todo: Add method in provider to add a clear scope
+			if (Style == EStyle::EnterScope)
+			{
+			}
+			else
+			{
+			}
+		}
+		break;
+
+	case RouteId_SaveStack:
+		{
+			const uint32 RuntimeId = Context.EventData.GetValue<uint32>("Id");
+			// todo: This should pin into a temporary buffer which can be released on analysis end
+			//auto AnalysisId = MetadataProvider->PinAndGetId(Context.ThreadInfo.GetSystemId());
+			//todo: Store RuntimeId <-> AnalysisId in map
+			//todo: Handle wrapping?
+		}
+		break;
+
+	case RouteId_RestoreStack:
+		{
+			if (Style == EStyle::EnterScope)
+			{
+				const uint32 RuntimeId = Context.EventData.GetValue<uint32>("Id");
+				//todo: Find analysis id in map
+			}
+			else
+			{
+				//MetadataProvider->PopScopedMetadata(Context.ThreadInfo.GetSystemId(), 0xFF /* All */);
+				//todo: At this point the mapping can be released, we can maintain a per thread memory of which saved
+				// stack is active. There can never be two restored stacks on top of each other.
+			}
+		}
+		break;
+
+	case RouteId_Metascope:
+		{
+			const FEventData& EventData = Context.EventData;
+			const FEventTypeInfo& EventInfo = Context.EventData.GetTypeInfo();
+
+			// Check if we have already encountered and registered this meta data type. If not, build a schema
+			// for this event type and register it with the provider.
+			const uint16 MetadataTypeId = GetOrRegisterType(EventInfo);
+
+			// Push the actual data of the scope to the provider or pop if we are leaving the scope.
+			if (Style == EStyle::EnterScope)
+			{
+				TArray<uint8, TInlineAllocator<32>> Data;
+				WriteData(Data, EventData, MetadataTypeId, MetadataProvider, Session);
+
+				MetadataProvider->PushScopedMetadata(Context.ThreadInfo.GetSystemId(), MetadataTypeId, Data.GetData(),
+													Data.Num());
+			}
+			else
+			{
+				//todo: Check that event type info is available on pop scope
+				MetadataProvider->PopScopedMetadata(Context.ThreadInfo.GetSystemId(), MetadataTypeId);
+			}
+		}
+		break;
+
+	default:
+		checkf(false, TEXT("Invalid route."));
+		break;
+	}
+	return true;
+}
+
+////////////////////////////////////////////////////////////////////////////////////////////////////
+uint16 FMetadataAnalysis::GetOrRegisterType(const FEventTypeInfo& EventInfo)
+{
+	if (const uint16* CachedTypeId = EncounteredMetadataTypes.Find(EventInfo.GetId()))
+	{
+		return *CachedTypeId;
+	}
+
+	uint16 MetadataTypeId = 0;
+	FMetadataSchema Schema(EventInfo.GetFieldCount());
+	FMetadataSchema::FBuilder Builder = Schema.Builder();
+	for (uint32 FieldIdx = 0; FieldIdx < EventInfo.GetFieldCount(); ++FieldIdx)
+	{
+		const FEventFieldInfo* FieldInfo = EventInfo.GetFieldInfo(FieldIdx);
+		const TCHAR* Name = Session.StoreString(ANSI_TO_TCHAR(FieldInfo->GetName()));
+		const uint8 Size = FieldInfo->GetSize();
+
+		switch (FieldInfo->GetType())
+		{
+		case FEventFieldInfo::EType::Integer:
+			{
+				if (FieldInfo->IsSigned())
+				{
+					constexpr auto FieldType = FMetadataSchema::EFieldType::SignedInteger;
+					if (Size == 1) Builder.AddField<int8>(Name, FieldType);
+					else if (Size == 2) Builder.AddField<int16>(Name, FieldType);
+					else if (Size == 4) Builder.AddField<int32>(Name, FieldType);
+					else if (Size == 8) Builder.AddField<int64>(Name, FieldType);
+					else { checkNoEntry(); }
+				}
+				else
+				{
+					constexpr auto FieldType = FMetadataSchema::EFieldType::Integer;
+					if (Size == 1) Builder.AddField<uint8>(Name, FieldType);
+					else if (Size == 2) Builder.AddField<uint16>(Name, FieldType);
+					else if (Size == 4) Builder.AddField<uint32>(Name, FieldType);
+					else if (Size == 8) Builder.AddField<uint64>(Name, FieldType);
+					else { checkNoEntry(); }
+				}
+			}
+			break;
+
+		case FEventFieldInfo::EType::Float:
+			{
+				constexpr auto FieldType = FMetadataSchema::EFieldType::FloatingPoint;
+				if (Size == 4) Builder.AddField<float>(Name, FieldType);
+				else if (Size == 8) Builder.AddField<double>(Name, FieldType);
+			}
+			break;
+
+		case FEventFieldInfo::EType::WideString:
+		case FEventFieldInfo::EType::AnsiString:
+			{
+				//Note that we don't have a ansi string store yet, so we store any ansi
+				//string as wide for metadata.
+				constexpr auto FieldType = FMetadataSchema::EFieldType::WideStringPtr;
+				Builder.AddField<TCHAR*>(Name, FieldType);
+			}
+			break;
+
+		case FEventFieldInfo::EType::Reference8:
+		case FEventFieldInfo::EType::Reference16:
+		case FEventFieldInfo::EType::Reference32:
+		case FEventFieldInfo::EType::Reference64:
+			{
+				constexpr auto FieldType = FMetadataSchema::EFieldType::Reference;
+				if (Size == 1) Builder.AddField<UE::Trace::FEventRef8>(Name, FieldType);
+				else if (Size == 2) Builder.AddField<UE::Trace::FEventRef16>(Name, FieldType);
+				else if (Size == 4) Builder.AddField<UE::Trace::FEventRef32>(Name, FieldType);
+				else if (Size == 8) Builder.AddField<UE::Trace::FEventRef64>(Name, FieldType);
+				else { checkNoEntry(); }
+			}
+			break;
+
+		default:
+			break;
+		}
+	}
+	Builder.Finish();
+	MetadataTypeId = MetadataProvider->RegisterMetadataType(ANSI_TO_TCHAR(EventInfo.GetName()), Schema);
+	EncounteredMetadataTypes.Emplace(EventInfo.GetId(), MetadataTypeId);
+	return MetadataTypeId;
+}
+
+////////////////////////////////////////////////////////////////////////////////////////////////////
+template<typename ArrayType>
+static void WriteData(ArrayType& OutData, const UE::Trace::IAnalyzer::FEventData& EventData, uint16 MetadataTypeId, const FMetadataProvider* MetadataProvider, IAnalysisSession& Session)
+{
+	using namespace UE::Trace;
+	const IAnalyzer::FEventTypeInfo& EventInfo = EventData.GetTypeInfo();
+	const FMetadataSchema* Schema = MetadataProvider->GetRegisteredMetadataSchema(MetadataTypeId);
+	check(Schema);
+	FMetadataSchema::FWriter Writer = Schema->Writer();
+
+	for (uint32 FieldIdx = 0; FieldIdx < EventInfo.GetFieldCount(); ++FieldIdx)
+	{
+		const IAnalyzer::FEventFieldInfo* FieldInfo = EventInfo.GetFieldInfo(FieldIdx);
+		switch (FieldInfo->GetType())
+		{
+		case IAnalyzer::FEventFieldInfo::EType::Reference8:
+		case IAnalyzer::FEventFieldInfo::EType::Reference16:
+		case IAnalyzer::FEventFieldInfo::EType::Reference32:
+		case IAnalyzer::FEventFieldInfo::EType::Reference64:
+			{
+				switch (FieldInfo->GetSize())
+				{
+				case 1:
+					{
+						auto RefValue = EventData.GetReferenceValue<uint8>(FieldIdx);
+						Writer.WriteField(FieldIdx, &RefValue, sizeof(UE::Trace::FEventRef8), OutData);
+					}
+					break;
+				case 2:
+					{
+						auto RefValue = EventData.GetReferenceValue<uint16>(FieldIdx);
+						Writer.WriteField(FieldIdx, &RefValue, sizeof(UE::Trace::FEventRef16), OutData);
+					}
+					break;
+				case 4:
+					{
+						auto RefValue = EventData.GetReferenceValue<uint32>(FieldIdx);
+						Writer.WriteField(FieldIdx, &RefValue, sizeof(UE::Trace::FEventRef32), OutData);
+					}
+					break;
+				case 8:
+					{
+						auto RefValue = EventData.GetReferenceValue<uint64>(FieldIdx);
+						Writer.WriteField(FieldIdx, &RefValue, sizeof(UE::Trace::FEventRef64), OutData);
+					}
+					break;
+				default:
+					checkNoEntry();
+				}
+			}
+			break;
+
+		case IAnalyzer::FEventFieldInfo::EType::WideString:
+			{
+				FWideStringView View;
+				EventData.GetString(FieldInfo->GetName(), View);
+				const TCHAR* Str = Session.StoreString(View);
+				Writer.WriteField(FieldIdx, Str, sizeof(TCHAR*), OutData);
+			}
+			break;
+
+		case IAnalyzer::FEventFieldInfo::EType::AnsiString:
+			{
+				FString String;
+				EventData.GetString(FieldInfo->GetName(), String);
+				const TCHAR* Str = Session.StoreString(String);
+				Writer.WriteField(FieldIdx, Str, sizeof(TCHAR*), OutData);
+			}
+			break;
+
+		default:
+			{
+				const IAnalyzer::FEventFieldHandle FieldHandle = EventInfo.GetFieldHandleUnchecked(FieldIdx);
+				const void* Value = EventData.GetValueRaw(FieldHandle);
+				const uint8 Size = FieldInfo->GetSize();
+				Writer.WriteField(FieldIdx, Value, Size, OutData);
+			}
+			break;
+		}
+	}
+}
+	
+////////////////////////////////////////////////////////////////////////////////////////////////////
+} // namespace TraceServices
diff --git a/Engine/Source/Developer/TraceServices/Private/Analyzers/MetadataAnalysis.h b/Engine/Source/Developer/TraceServices/Private/Analyzers/MetadataAnalysis.h
new file mode 100644
--- /dev/null
+++ b/Engine/Source/Developer/TraceServices/Private/Analyzers/MetadataAnalysis.h
@@ -0,0 +1,37 @@
+﻿// Copyright Epic Games, Inc. All Rights Reserved.
+#pragma once
+#include "Trace/Analyzer.h"
+#include "Containers/Map.h"
+
+namespace TraceServices
+{
+
+class IAnalysisSession;
+class FMetadataProvider;	
+
+class FMetadataAnalysis : public UE::Trace::IAnalyzer
+{
+public:
+	FMetadataAnalysis(IAnalysisSession& Session, FMetadataProvider* InProvider);
+	virtual ~FMetadataAnalysis() override;
+	
+	virtual void OnAnalysisBegin(const FOnAnalysisContext& Context) override;
+	virtual void OnAnalysisEnd() override;
+	virtual bool OnEvent(uint16 RouteId, EStyle Style, const FOnEventContext& Context) override;
+
+private:
+	uint16 GetOrRegisterType(const FEventTypeInfo& EventInfo);
+	
+	enum EMetadataRouteIds
+	{
+		 RouteId_ClearScope,
+		 RouteId_SaveStack,
+		 RouteId_RestoreStack,
+		 RouteId_Metascope,
+	};
+	IAnalysisSession& Session;
+	FMetadataProvider* MetadataProvider;
+	TMap<uint32,uint16> EncounteredMetadataTypes;
+};
+
+} // namespace TraceServices
diff --git a/Engine/Source/Developer/TraceServices/Private/Model/DefinitionProvider.cpp b/Engine/Source/Developer/TraceServices/Private/Model/DefinitionProvider.cpp
--- a/Engine/Source/Developer/TraceServices/Private/Model/DefinitionProvider.cpp
+++ b/Engine/Source/Developer/TraceServices/Private/Model/DefinitionProvider.cpp
@@ -5,6 +5,9 @@
 
 namespace TraceServices
 {
+
+//////////////////////////////////////////////////////////////////////////////////////////
+const FName FDefinitionProvider::ProviderName(TEXT("DefinitionProvider"));	
 	
 //////////////////////////////////////////////////////////////////////////////////////////
 FDefinitionProvider::FDefinitionProvider(IAnalysisSession* InSession)
diff --git a/Engine/Source/Developer/TraceServices/Private/Model/DefinitionProvider.h b/Engine/Source/Developer/TraceServices/Private/Model/DefinitionProvider.h
--- a/Engine/Source/Developer/TraceServices/Private/Model/DefinitionProvider.h
+++ b/Engine/Source/Developer/TraceServices/Private/Model/DefinitionProvider.h
@@ -15,6 +15,8 @@ namespace TraceServices
 class FDefinitionProvider : public IDefinitionProvider
 {
 public:
+	static const FName ProviderName;
+	
 	FDefinitionProvider(IAnalysisSession* InSession);
 
 	virtual void BeginEdit() const override { Lock.WriteLock(); };
diff --git a/Engine/Source/Developer/TraceServices/Private/Modules/MemoryModule.cpp b/Engine/Source/Developer/TraceServices/Private/Modules/MemoryModule.cpp
--- a/Engine/Source/Developer/TraceServices/Private/Modules/MemoryModule.cpp
+++ b/Engine/Source/Developer/TraceServices/Private/Modules/MemoryModule.cpp
@@ -4,6 +4,7 @@
 #include "Analyzers/AllocationsAnalysis.h"
 #include "Analyzers/CallstacksAnalysis.h"
 #include "Analyzers/MemoryAnalysis.h"
+#include "Analyzers/MetadataAnalysis.h"
 #include "Analyzers/ModuleAnalysis.h"
 #include "Model/AllocationsProvider.h"
 #include "Model/CallstacksProvider.h"
@@ -39,11 +40,13 @@ void FMemoryModule::OnAnalysisBegin(IAnalysisSession& Session)
 	// Metadata
 	FMetadataProvider* MetadataProvider = new FMetadataProvider(Session);
 	Session.AddProvider(GetMetadataProviderName(), MetadataProvider);
-
+	Session.AddAnalyzer(new FMetadataAnalysis(Session, MetadataProvider));
+	
 	// Allocations
 	FAllocationsProvider* AllocationsProvider = new FAllocationsProvider(Session, *MetadataProvider);
 	Session.AddProvider(GetAllocationsProviderName(), AllocationsProvider);
 	Session.AddAnalyzer(new FAllocationsAnalyzer(Session, *AllocationsProvider, *MetadataProvider));
+
 }
 
 void FMemoryModule::GetLoggers(TArray<const TCHAR*>& OutLoggers)
diff --git a/Engine/Source/Developer/TraceServices/Public/TraceServices/ModuleService.h b/Engine/Source/Developer/TraceServices/Public/TraceServices/ModuleService.h
--- a/Engine/Source/Developer/TraceServices/Public/TraceServices/ModuleService.h
+++ b/Engine/Source/Developer/TraceServices/Public/TraceServices/ModuleService.h
@@ -10,7 +10,6 @@ class FName;
 namespace TraceServices
 {
 
-class IAnalyzer;
 class IAnalysisSession;
 extern TRACESERVICES_API const FName ModuleFeatureName;
 
diff --git a/Engine/Source/Runtime/Core/Private/ProfilingDebugging/AssetMetadataTrace.cpp b/Engine/Source/Runtime/Core/Private/ProfilingDebugging/AssetMetadataTrace.cpp
new file mode 100644
--- /dev/null
+++ b/Engine/Source/Runtime/Core/Private/ProfilingDebugging/AssetMetadataTrace.cpp
@@ -0,0 +1,10 @@
+﻿// Copyright Epic Games, Inc. All Rights Reserved.
+#include "ProfilingDebugging/AssetMetadataTrace.h"
+
+////////////////////////////////////////////////////////////////////////////////////////////////////
+#if UE_TRACE_ASSET_METADATA_ENABLED
+
+UE_TRACE_CHANNEL_DEFINE(AssetMetadataChannel);
+UE_TRACE_METADATA_EVENT_DEFINE(Asset);
+
+#endif // UE_TRACE_ASSET_METADATA_ENABLED
diff --git a/Engine/Source/Runtime/Core/Private/ProfilingDebugging/MetadataTrace.cpp b/Engine/Source/Runtime/Core/Private/ProfilingDebugging/MetadataTrace.cpp
new file mode 100644
--- /dev/null
+++ b/Engine/Source/Runtime/Core/Private/ProfilingDebugging/MetadataTrace.cpp
@@ -0,0 +1,24 @@
+﻿// Copyright Epic Games, Inc. All Rights Reserved.
+#include "ProfilingDebugging/MetadataTrace.h"
+
+////////////////////////////////////////////////////////////////////////////////////////////////////
+#if UE_TRACE_METADATA_ENABLED
+
+UE_TRACE_CHANNEL_DEFINE(MetadataChannel);
+UE_TRACE_EVENT_DEFINE(MetadataStack, ClearScope);
+UE_TRACE_EVENT_DEFINE(MetadataStack, SaveStack);
+UE_TRACE_EVENT_DEFINE(MetadataStack, RestoreStack);
+
+////////////////////////////////////////////////////////////////////////////////////////////////////
+uint32 FMetadataTrace::SaveStack()
+{
+	static std::atomic_uint32_t Counter(1);
+	const uint32 Id = Counter.fetch_add(1u);
+	UE_TRACE_LOG(MetadataStack, SaveStack, MetadataChannel)
+			<< SaveStack.Id(Id);
+	return Id;
+}
+
+#endif // UE_TRACE_METADATA_ENABLED
+
+
diff --git a/Engine/Source/Runtime/Core/Public/Async/TaskGraphInterfaces.h b/Engine/Source/Runtime/Core/Public/Async/TaskGraphInterfaces.h
--- a/Engine/Source/Runtime/Core/Public/Async/TaskGraphInterfaces.h
+++ b/Engine/Source/Runtime/Core/Public/Async/TaskGraphInterfaces.h
@@ -24,6 +24,7 @@
 #include "Experimental/ConcurrentLinearAllocator.h"
 #include "Misc/MemStack.h"
 #include "Templates/Atomic.h"
+#include "ProfilingDebugging/MetadataTrace.h"
 
 #include "Async/Fundamental/Task.h"
 
@@ -815,6 +816,9 @@ protected:
 		checkThreadGraph(LifeStage.Increment() == int32(LS_Contructed));
 #if UE_MEMORY_TAGS_TRACE_ENABLED
 		InheritedTraceTag = MemoryTrace_GetActiveTag();
+#endif
+#if UE_TRACE_METADATA_ENABLED
+		InheritedMetadataId = UE_TRACE_METADATA_SAVE_STACK();
 #endif
 		LLM(InheritedLLMTag = FLowLevelMemTracker::bIsDisabled ? nullptr : FLowLevelMemTracker::Get().GetActiveTagData(ELLMTracker::Default));
 	}
@@ -929,6 +933,7 @@ private:
 	{
 		LLM_SCOPE(InheritedLLMTag);
 		UE_MEMSCOPE(InheritedTraceTag);
+		UE_TRACE_METADATA_RESTORE_STACK(InheritedMetadataId);
 		checkThreadGraph(LifeStage.Increment() == int32(LS_Executing));
 		ExecuteTask(NewTasks, CurrentThread, bDeleteOnCompletion);
 	}
@@ -970,6 +975,9 @@ private:
 #endif
 #if UE_MEMORY_TAGS_TRACE_ENABLED
 	int32 InheritedTraceTag;
+#endif
+#if UE_TRACE_METADATA_ENABLED
+	uint32 InheritedMetadataId;
 #endif
 	LLM(const UE::LLMPrivate::FTagData* InheritedLLMTag);
 
diff --git a/Engine/Source/Runtime/Core/Public/ProfilingDebugging/AssetMetadataTrace.h b/Engine/Source/Runtime/Core/Public/ProfilingDebugging/AssetMetadataTrace.h
new file mode 100644
--- /dev/null
+++ b/Engine/Source/Runtime/Core/Public/ProfilingDebugging/AssetMetadataTrace.h
@@ -0,0 +1,43 @@
+﻿// Copyright Epic Games, Inc. All Rights Reserved.
+#pragma once
+
+#include "MetadataTrace.h"
+
+////////////////////////////////////////////////////////////////////////////////////////////////////
+#if !defined(UE_TRACE_ASSET_METADATA_ENABLED) 
+	#define UE_TRACE_ASSET_METADATA_ENABLED UE_TRACE_METADATA_ENABLED
+#endif
+
+////////////////////////////////////////////////////////////////////////////////////////////////////
+#if UE_TRACE_ASSET_METADATA_ENABLED
+
+/**
+ * Channel that asset meta data is output on
+ */
+CORE_API UE_TRACE_CHANNEL_EXTERN(AssetMetadataChannel);
+
+/**
+ * Metadata scope to instrument operations belonging to a certain asset.
+ */
+CORE_API UE_TRACE_METADATA_EVENT_BEGIN_EXTERN(Asset)
+	UE_TRACE_METADATA_EVENT_REFERENCE_FIELD(Strings, FName, Name)
+	UE_TRACE_METADATA_EVENT_REFERENCE_FIELD(Strings, FName, Class)
+UE_TRACE_METADATA_EVENT_END()
+
+/**
+ * Utility macro to use asset scope
+ * @note When using this macro in an module outside of Core, make sure a dependency to the "TraceLog" module
+ *		 is added.
+ */
+#define UE_TRACE_METADATA_SCOPE_ASSET(Object, ObjClass) \
+	auto MetaNameRef = UE_TRACE_CHANNELEXPR_IS_ENABLED(MetadataChannel) ? FStringTrace::GetNameRef(Object->GetFName()) : UE::Trace::FEventRef32(0,0);\
+	auto ClassNameRef = UE_TRACE_CHANNELEXPR_IS_ENABLED(MetadataChannel) ? FStringTrace::GetNameRef(ObjClass->GetFName()) : UE::Trace::FEventRef32(0,0);\
+	UE_TRACE_METADATA_SCOPE(Asset, AssetMetadataChannel)\
+		<< Asset.Name(MetaNameRef)\
+		<< Asset.Class(ClassNameRef);
+
+#else
+
+#define UE_TRACE_METADATA_SCOPE_ASSET(...)
+
+#endif // UE_TRACE_ASSET_METADATA_ENABLED
diff --git a/Engine/Source/Runtime/Core/Public/ProfilingDebugging/MetadataTrace.h b/Engine/Source/Runtime/Core/Public/ProfilingDebugging/MetadataTrace.h
new file mode 100644
--- /dev/null
+++ b/Engine/Source/Runtime/Core/Public/ProfilingDebugging/MetadataTrace.h
@@ -0,0 +1,94 @@
+﻿// Copyright Epic Games, Inc. All Rights Reserved.
+#pragma once
+
+#include "Trace/Trace.h"
+#include "ProfilingDebugging/StringsTrace.h"
+
+////////////////////////////////////////////////////////////////////////////////////////////////////
+# if !defined(UE_TRACE_METADATA_ENABLED)
+	#if !IS_PROGRAM
+		#define	UE_TRACE_METADATA_ENABLED UE_TRACE_ENABLED
+	#else
+		#define	UE_TRACE_METADATA_ENABLED 0
+	#endif
+#endif //defined(UE_TRACE_METADATA_ENABLED
+
+#if UE_TRACE_METADATA_ENABLED
+
+// Only needed when enabled
+#include "Trace/Trace.inl"
+
+////////////////////////////////////////////////////////////////////////////////////////////////////
+
+/**
+ * Channel that asset meta data is output on
+ */
+CORE_API UE_TRACE_CHANNEL_EXTERN(MetadataChannel);
+
+CORE_API UE_TRACE_EVENT_BEGIN_EXTERN(MetadataStack, ClearScope, NoSync)
+UE_TRACE_EVENT_END()
+
+CORE_API UE_TRACE_EVENT_BEGIN_EXTERN(MetadataStack, SaveStack)
+	UE_TRACE_EVENT_FIELD(uint32, Id)
+UE_TRACE_EVENT_END()
+
+CORE_API UE_TRACE_EVENT_BEGIN_EXTERN(MetadataStack, RestoreStack)
+	UE_TRACE_EVENT_FIELD(uint32, Id)
+UE_TRACE_EVENT_END()
+
+////////////////////////////////////////////////////////////////////////////////////////////////////
+class FMetadataTrace
+{
+public:
+	CORE_API static uint32 SaveStack();
+};
+
+////////////////////////////////////////////////////////////////////////////////////////////////////
+
+// For now just forward to the event definition macros.
+#define UE_TRACE_METADATA_EVENT_BEGIN_EXTERN(ScopeName, ...) UE_TRACE_EVENT_BEGIN_EXTERN(Metadata, ScopeName, ##__VA_ARGS__)
+#define UE_TRACE_METADATA_EVENT_DEFINE(ScopeName) UE_TRACE_EVENT_DEFINE(Metadata, ScopeName)
+#define UE_TRACE_METADATA_EVENT_BEGIN(ScopeName, ...) UE_TRACE_EVENT_BEGIN(Metadata, ScopeName, ##__VA_ARGS__)
+#define UE_TRACE_METADATA_EVENT_FIELD(Type, Name) UE_TRACE_EVENT_FIELD(Type, Name)
+#define UE_TRACE_METADATA_EVENT_REFERENCE_FIELD(RefLogger, RefEventName, Name) UE_TRACE_EVENT_REFERENCE_FIELD(RefLogger, RefEventName, Name)
+#define UE_TRACE_METADATA_EVENT_END() UE_TRACE_EVENT_END()
+#define UE_TRACE_METADATA_SCOPE(ScopeName, ChannelExpr) UE_TRACE_LOG_SCOPED(Metadata, ScopeName, ChannelExpr)
+
+/**
+ * Temporarily clears the current meta data stack and starts a new one. When this scope is exited the previous
+ * stack is restored. This is useful for example in situations where an asset causes allocation in an unrelated
+ * system, like shared buffers, but that allocation should not be ascribed to a specific asset.
+ */
+#define UE_TRACE_METADATA_CLEAR_SCOPE() \
+	UE_TRACE_LOG_SCOPED(MetadataStack, ClearScope, MetadataChannel)
+
+/**
+ * Saves the currents stack, and returns an identifier. The stack can be restored on another thread using
+ * the identifier. This can be used when the context needs to be applied on a task.
+ */
+#define UE_TRACE_METADATA_SAVE_STACK() \
+	FMetadataTrace::SaveStack();
+
+/**
+ * Restore a previously saved stack.
+ */
+#define UE_TRACE_METADATA_RESTORE_STACK(InId) \
+	UE_TRACE_LOG_SCOPED(MetadataStack, RestoreStack, MetadataChannel) \
+		<< RestoreStack.Id(InId);
+
+////////////////////////////////////////////////////////////////////////////////////////////////////
+
+#else
+
+#define UE_TRACE_METADATA_EVENT_BEGIN_EXTERN(ScopeName, ...) 
+#define UE_TRACE_METADATA_EVENT_DEFINE(ScopeName) 
+#define UE_TRACE_METADATA_EVENT_BEGIN(ScopeName, ...) 
+#define UE_TRACE_METADATA_EVENT_FIELD(Type, Name) 
+#define UE_TRACE_METADATA_EVENT_END() 
+#define UE_TRACE_METADATA_SCOPE(ScopeName) 
+
+#define UE_TRACE_METADATA_CLEAR_SCOPE()
+#define UE_TRACE_METADATA_SAVE_STACK() 0
+#define UE_TRACE_METADATA_RESTORE_STACK(...)
+
+#endif
\ No newline at end of file
diff --git a/Engine/Source/Runtime/CoreUObject/Private/Serialization/AsyncLoading.cpp b/Engine/Source/Runtime/CoreUObject/Private/Serialization/AsyncLoading.cpp
--- a/Engine/Source/Runtime/CoreUObject/Private/Serialization/AsyncLoading.cpp
+++ b/Engine/Source/Runtime/CoreUObject/Private/Serialization/AsyncLoading.cpp
@@ -46,6 +46,7 @@
 #include "HAL/LowLevelMemStats.h"
 #include "HAL/IPlatformFileOpenLogWrapper.h"
 #include "ProfilingDebugging/CsvProfiler.h"
+#include "ProfilingDebugging/AssetMetadataTrace.h"
 #include "UObject/GarbageCollectionInternal.h"
 #include "ProfilingDebugging/MiscTrace.h"
 #include "Serialization/LoadTimeTracePrivate.h"
@@ -3134,6 +3135,7 @@ void FAsyncPackage::EventDrivenCreateExport(int32 LocalExportIndex)
 	LLM_SCOPE(ELLMTag::AsyncLoading);
 	LLM_SCOPED_TAG_WITH_OBJECT_IN_SET(GetLinkerRoot(), ELLMTagSet::Assets);
 	LLM_SCOPED_TAG_WITH_OBJECT_IN_SET(CastEventDrivenIndexToObject<UClass>(Export.ClassIndex, false), ELLMTagSet::AssetClasses);
+   	UE_TRACE_METADATA_SCOPE_ASSET(GetLinkerRoot(), CastEventDrivenIndexToObject<UClass>(Export.ClassIndex, false));
 
 	// Check whether we already loaded the object and if not whether the context flags allow loading it.
 	//check(!Export.Object || Export.Object->HasAnyFlags(RF_ClassDefaultObject)); // we should not have this yet, unless it is a CDO
@@ -3419,6 +3421,7 @@ void FAsyncPackage::EventDrivenSerializeExport(int32 LocalExportIndex)
 
 	LLM_SCOPED_TAG_WITH_OBJECT_IN_SET(GetLinkerRoot(), ELLMTagSet::Assets);
 	LLM_SCOPED_TAG_WITH_OBJECT_IN_SET(CastEventDrivenIndexToObject<UClass>(Export.ClassIndex, false), ELLMTagSet::AssetClasses);
+  	UE_TRACE_METADATA_SCOPE_ASSET(GetLinkerRoot(), CastEventDrivenIndexToObject<UClass>(Export.ClassIndex, false));
 
 	UObject* Object = Export.Object;
 	if (Object && Object->HasAnyFlags(RF_NeedLoad))
diff --git a/Engine/Source/Runtime/CoreUObject/Private/Serialization/AsyncLoading2.cpp b/Engine/Source/Runtime/CoreUObject/Private/Serialization/AsyncLoading2.cpp
--- a/Engine/Source/Runtime/CoreUObject/Private/Serialization/AsyncLoading2.cpp
+++ b/Engine/Source/Runtime/CoreUObject/Private/Serialization/AsyncLoading2.cpp
@@ -57,6 +57,7 @@
 #include "UObject/UObjectClusters.h"
 #include "UObject/LinkerInstancingContext.h"
 #include "ProfilingDebugging/CountersTrace.h"
+#include "ProfilingDebugging/AssetMetadataTrace.h"
 #include "Async/Async.h"
 #include "Async/ParallelFor.h"
 #include "HAL/LowLevelMemStats.h"
@@ -4269,13 +4270,10 @@ void FAsyncPackage2::EventDrivenCreateExport(const FAsyncPackageHeaderData& Head
 		return;
 	}
 
-	LLM_SCOPED_TAG_WITH_OBJECT_IN_SET(GetLinkerRoot(), ELLMTagSet::Assets);
-	// LLM_SCOPED_TAG_WITH_OBJECT_IN_SET(CastEventDrivenIndexToObject<UClass>(Export.ClassIndex, false), ELLMTagSet::AssetClasses);
-
 	bool bIsCompleteyLoaded = false;
 	UClass* LoadClass = Export.ClassIndex.IsNull() ? UClass::StaticClass() : CastEventDrivenIndexToObject<UClass>(Header, Exports, Export.ClassIndex, true);
 	UObject* ThisParent = Export.OuterIndex.IsNull() ? LinkerRoot : EventDrivenIndexToObject(Header, Exports, Export.OuterIndex, false);
-
+	
 	if (!LoadClass)
 	{
 		UE_ASYNC_PACKAGE_LOG(Error, Desc, TEXT("CreateExport"), TEXT("Could not find class object for %s"), *ObjectName.ToString());
@@ -4308,6 +4306,10 @@ void FAsyncPackage2::EventDrivenCreateExport(const FAsyncPackageHeaderData& Head
 		ExportObject.bExportLoadFailed = true;
 		return;
 	}
+	
+	LLM_SCOPED_TAG_WITH_OBJECT_IN_SET(GetLinkerRoot(), ELLMTagSet::Assets);
+	LLM_SCOPED_TAG_WITH_OBJECT_IN_SET(LoadClass, ELLMTagSet::AssetClasses);
+    UE_TRACE_METADATA_SCOPE_ASSET(GetLinkerRoot(), LoadClass);
 
 	// Try to find existing object first as we cannot in-place replace objects, could have been created by other export in this package
 	{
@@ -4487,8 +4489,10 @@ bool FAsyncPackage2::EventDrivenSerializeExport(const FAsyncPackageHeaderData& H
 		}
 	}
 
-	LLM_SCOPED_TAG_WITH_OBJECT_IN_SET(GetLinkerRoot(), ELLMTagSet::Assets);
-	// LLM_SCOPED_TAG_WITH_OBJECT_IN_SET(CastEventDrivenIndexToObject<UClass>(Export.ClassIndex, false), ELLMTagSet::AssetClasses);
+	const UClass* LoadClass = Export.ClassIndex.IsNull() ? UClass::StaticClass() : CastEventDrivenIndexToObject<UClass>(Header, Exports, Export.ClassIndex, true);
+	UE_TRACE_METADATA_SCOPE_ASSET(Object, LoadClass);
+	LLM_SCOPED_TAG_WITH_OBJECT_IN_SET(Object, ELLMTagSet::Assets);
+	LLM_SCOPED_TAG_WITH_OBJECT_IN_SET(LoadClass, ELLMTagSet::AssetClasses);
 
 	// cache archetype
 	// prevents GetArchetype from hitting the expensive GetArchetypeFromRequiredInfoImpl
diff --git a/Engine/Source/Runtime/CoreUObject/Private/UObject/Obj.cpp b/Engine/Source/Runtime/CoreUObject/Private/UObject/Obj.cpp
--- a/Engine/Source/Runtime/CoreUObject/Private/UObject/Obj.cpp
+++ b/Engine/Source/Runtime/CoreUObject/Private/UObject/Obj.cpp
@@ -62,12 +62,10 @@
 #include "Misc/ScopeRWLock.h"
 #include "Misc/PackageAccessTracking.h"
 #include "Misc/PackageAccessTrackingOps.h"
+#include "ProfilingDebugging/AssetMetadataTrace.h"
 
 DEFINE_LOG_CATEGORY(LogObj);
 
-/** Stat group for dynamic objects cycle counters*/
-
-
 /*-----------------------------------------------------------------------------
 	Globals.
 -----------------------------------------------------------------------------*/
@@ -1112,6 +1110,7 @@ void UObject::ConditionalPostLoad()
 				UE_SCOPED_COOK_STAT(Package->GetFName(), EPackageEventStatType::LoadPackage);
 				LLM_SCOPED_TAG_WITH_OBJECT_IN_SET(Package, ELLMTagSet::Assets);
 				LLM_SCOPED_TAG_WITH_OBJECT_IN_SET(GetClass(), ELLMTagSet::AssetClasses);
+				UE_TRACE_METADATA_SCOPE_ASSET(Package, GetClass());
 
 				PostLoad();
 
