[Insights] Multiple fixes and improvements to TPagedArray and its iterator:
  - Fixed TPagedArray's end() to point to an item after the last item in array. This fixes the issue with range loops where the last element was ignored.
  - Changed TPagedArrayIterator to construct using an item index (allows values from 0 to TotalItemCount; where a value equal with TotalItemCount is the end() position, past last valid element).
  - Added debug code to check state of iterator and of source array during iteration.
  - Added a new simpler TPagedArrayIterator that is only based on an item index, for debug purposes.
  - Fixed use of TPagedArray iterator in ConextSwitches.

diff --git a/Engine/Source/Developer/TraceServices/Private/Model/ContextSwitches.cpp b/Engine/Source/Developer/TraceServices/Private/Model/ContextSwitches.cpp
--- a/Engine/Source/Developer/TraceServices/Private/Model/ContextSwitches.cpp
+++ b/Engine/Source/Developer/TraceServices/Private/Model/ContextSwitches.cpp
@@ -241,7 +241,9 @@ void FContextSwitchesProvider::EnumerateCpuCoreEventsBackwards(uint32 CoreNumber
 		return;
 	}
 
-	auto Iterator = PageIndex < CpuCoreEvents->NumPages() ? CpuCoreEvents->GetIteratorFromPage(PageIndex) : CpuCoreEvents->end();
+	auto Iterator = PageIndex < CpuCoreEvents->NumPages() ?
+		CpuCoreEvents->GetIteratorFromPage(PageIndex) :
+		CpuCoreEvents->GetIteratorFromItem(CpuCoreEvents->Num() - 1);
 	const FCpuCoreEvent* CurrentCpuCoreEvent = Iterator.GetCurrentItem();
 	while (CurrentCpuCoreEvent && CurrentCpuCoreEvent->End > StartTime)
 	{
diff --git a/Engine/Source/Developer/TraceServices/Public/Common/PagedArray.h b/Engine/Source/Developer/TraceServices/Public/Common/PagedArray.h
--- a/Engine/Source/Developer/TraceServices/Public/Common/PagedArray.h
+++ b/Engine/Source/Developer/TraceServices/Public/Common/PagedArray.h
@@ -5,6 +5,14 @@
 #include "TraceServices/Containers/Allocators.h"
 #include "Containers/Array.h"
 
+#ifndef TRACESERVICES_PAGED_ARRAY_ITERATOR_V2
+#define TRACESERVICES_PAGED_ARRAY_ITERATOR_V2 0 // enables a simpler implementation of TPagedArrayIterator, for debug purposes
+#endif
+
+#ifndef TRACESERVICES_PAGED_ARRAY_ITERATOR_DEBUG_ENABLED
+#define TRACESERVICES_PAGED_ARRAY_ITERATOR_DEBUG_ENABLED 0
+#endif
+
 namespace TraceServices {
 
 template<typename ItemType>
@@ -48,6 +56,8 @@ inline const ItemType* GetLastItem(const TPagedArrayPage<ItemType>& Page)
 template<typename ItemType, typename PageType>
 class TPagedArray;
 
+#if !TRACESERVICES_PAGED_ARRAY_ITERATOR_V2
+
 template<typename ItemType, typename PageType>
 class TPagedArrayIterator
 {
@@ -56,18 +66,16 @@ public:
 	{
 	}
 
-	TPagedArrayIterator(const TPagedArray<ItemType, PageType>& InOuter, uint64 InitialPageIndex, uint64 InitialItemIndex)
+	TPagedArrayIterator(const TPagedArray<ItemType, PageType>& InOuter, uint64 InItemIndex)
 		: Outer(&InOuter)
 	{
-		if (Outer->PagesArray.Num())
-		{
-			check(InitialPageIndex < Outer->PagesArray.Num());
-			CurrentPageIndex = InitialPageIndex;
-			OnCurrentPageChanged();
-			PageType* CurrentPage = Outer->FirstPage + CurrentPageIndex;
-			check(InitialItemIndex < CurrentPage->Count);
-			CurrentItem = CurrentPage->Items + InitialItemIndex;
-		}
+#if TRACESERVICES_PAGED_ARRAY_ITERATOR_DEBUG_ENABLED
+		TotalItemCount = Outer->Num();
+		TotalPageCount = Outer->PagesArray.Num();
+#endif
+
+		SetPositionInternal(InItemIndex);
+		DebugCheckState();
 	}
 
 	//////////////////////////////////////////////////
@@ -83,57 +91,93 @@ public:
 		return Outer->FirstPage + CurrentPageIndex;
 	}
 
+	const PageType* SetCurrentPage(uint64 PageIndex)
+	{
+		DebugCheckState();
+		uint64 ItemIndex = PageIndex * Outer->PageSize;
+		if (ItemIndex >= Outer->Num()) // end()
+		{
+			ItemIndex = Outer->Num();
+			CurrentPageIndex = ItemIndex / Outer->PageSize;
+			CurrentPageFirstItem = nullptr;
+			CurrentPageLastItem = nullptr;
+			CurrentItemIndex = ItemIndex;
+			CurrentItem = nullptr;
+			DebugCheckState();
+			return nullptr;
+		}
+		else
+		{
+			check(PageIndex < Outer->PagesArray.Num());
+			CurrentPageIndex = PageIndex;
+			OnCurrentPageChanged();
+			CurrentItemIndex = ItemIndex;
+			PageType* CurrentPage = Outer->FirstPage + CurrentPageIndex;
+			check(CurrentPage->Count > 0);
+			CurrentItem = CurrentPage->Items;
+			DebugCheckState();
+			return CurrentPage;
+		}
+	}
+
 	const PageType* PrevPage()
 	{
+		DebugCheckState();
 		if (CurrentPageIndex == 0)
 		{
 			CurrentItem = nullptr;
 			CurrentPageFirstItem = nullptr;
 			CurrentPageLastItem = nullptr;
+			CurrentItemIndex = Outer->Num();
+			DebugCheckState();
 			return nullptr;
 		}
 		--CurrentPageIndex;
 		OnCurrentPageChanged();
+		CurrentItemIndex = CurrentPageIndex * Outer->PageSize + (CurrentPageLastItem - CurrentPageFirstItem);
 		CurrentItem = CurrentPageLastItem;
+		DebugCheckState();
 		return GetCurrentPage();
 	}
 
 	const PageType* NextPage()
 	{
+		DebugCheckState();
 		if (CurrentPageIndex == Outer->PagesArray.Num() - 1)
 		{
 			CurrentItem = nullptr;
 			CurrentPageFirstItem = nullptr;
 			CurrentPageLastItem = nullptr;
+			CurrentItemIndex = Outer->Num();
+			DebugCheckState();
 			return nullptr;
 		}
 		++CurrentPageIndex;
 		OnCurrentPageChanged();
+		CurrentItemIndex = CurrentPageIndex * Outer->PageSize;
 		CurrentItem = CurrentPageFirstItem;
+		DebugCheckState();
 		return GetCurrentPage();
 	}
 
 	//////////////////////////////////////////////////
 	// Item Iterator
 
+	uint64 GetCurrentItemIndex() const
+	{
+		return CurrentItemIndex;
+	}
+
 	const ItemType* GetCurrentItem() const
 	{
 		return CurrentItem;
 	}
 
-	const ItemType* SetPosition(uint64 Index)
+	const ItemType* SetPosition(uint64 InItemIndex)
 	{
-		uint64 PageIndex = Index / Outer->PageSize;
-		uint64 ItemIndexInPage = Index % Outer->PageSize;
-		if (PageIndex != CurrentPageIndex)
-		{
-			check(PageIndex < Outer->PagesArray.Num());
-			CurrentPageIndex = PageIndex;
-			OnCurrentPageChanged();
-		}
-		PageType* CurrentPage = Outer->FirstPage + CurrentPageIndex;
-		check(ItemIndexInPage < CurrentPage->Count);
-		CurrentItem = CurrentPage->Items + ItemIndexInPage;
+		DebugCheckState();
+		SetPositionInternal(InItemIndex);
+		DebugCheckState();
 		return CurrentItem;
 	}
 
@@ -150,7 +194,10 @@ public:
 				return CurrentItem;
 			}
 		}
+		DebugCheckState();
+		--CurrentItemIndex;
 		--CurrentItem;
+		DebugCheckState();
 		return CurrentItem;
 	}
 
@@ -167,7 +214,10 @@ public:
 				return CurrentItem;
 			}
 		}
+		DebugCheckState();
+		++CurrentItemIndex;
 		++CurrentItem;
+		DebugCheckState();
 		return CurrentItem;
 	}
 
@@ -230,19 +280,269 @@ private:
 		}
 	}
 
+	void SetPositionInternal(uint64 InItemIndex)
+	{
+		CurrentPageIndex = InItemIndex / Outer->PageSize;
+		CurrentItemIndex = InItemIndex;
+
+		if (InItemIndex == Outer->Num()) // end()
+		{
+			CurrentPageFirstItem = nullptr;
+			CurrentPageLastItem = nullptr;
+			CurrentItem = nullptr;
+		}
+		else
+		{
+			check(InItemIndex < Outer->Num());
+			check(CurrentPageIndex < Outer->PagesArray.Num());
+			OnCurrentPageChanged();
+
+			PageType* CurrentPage = Outer->FirstPage + CurrentPageIndex;
+			uint64 ItemIndexInPage = InItemIndex % Outer->PageSize;
+			check(ItemIndexInPage < CurrentPage->Count);
+			CurrentItem = CurrentPage->Items + ItemIndexInPage;
+		}
+	}
+
 	FORCEINLINE friend bool operator!=(const TPagedArrayIterator& Lhs, const TPagedArrayIterator& Rhs)
 	{
 		checkSlow(Lhs.Outer == Rhs.Outer); // Needs to be iterators of the same array
-		return Lhs.CurrentItem != Rhs.CurrentItem;
+		return Lhs.CurrentItemIndex != Rhs.CurrentItemIndex;
+	}
+
+#if TRACESERVICES_PAGED_ARRAY_ITERATOR_DEBUG_ENABLED
+	void DebugCheckState()
+	{
+		if (Outer)
+		{
+			check(TotalPageCount == Outer->PagesArray.Num());
+			check(TotalItemCount == Outer->Num());
+
+			if (CurrentItemIndex == TotalItemCount) // end()
+			{
+				check(CurrentItem == nullptr);
+
+				check(CurrentPageIndex <= TotalPageCount);
+				check(CurrentPageFirstItem == nullptr);
+				check(CurrentPageLastItem == nullptr);
+			}
+			else
+			{
+				check(CurrentItemIndex < TotalItemCount);
+				check(CurrentItem != nullptr);
+
+				check(CurrentPageIndex < TotalPageCount);
+				PageType* CurrentPage = Outer->FirstPage + CurrentPageIndex;
+				check(CurrentPage->Count > 0);
+
+				uint64 ItemIndexInPage = CurrentItemIndex % Outer->PageSize;
+				check(ItemIndexInPage < CurrentPage->Count);
+
+				check(CurrentPageFirstItem == CurrentPage->Items);
+				check(CurrentPageLastItem == CurrentPage->Items + CurrentPage->Count - 1);
+				check(CurrentItem == CurrentPage->Items + ItemIndexInPage);
+			}
+		}
+		else
+		{
+			check(TotalPageCount == 0);
+			check(TotalItemCount == 0);
+
+			check(CurrentPageIndex == 0);
+			check(CurrentPageFirstItem == nullptr);
+			check(CurrentPageLastItem == nullptr);
+			check(CurrentItemIndex == 0);
+			check(CurrentItem == nullptr);
+		}
 	}
+#else
+	void DebugCheckState()
+	{
+	}
+#endif
 
+private:
 	const TPagedArray<ItemType, PageType>* Outer = nullptr;
-	const ItemType* CurrentItem = nullptr;
+	uint64 CurrentPageIndex = 0;
 	const ItemType* CurrentPageFirstItem = nullptr;
 	const ItemType* CurrentPageLastItem = nullptr;
-	uint64 CurrentPageIndex = 0;
+	uint64 CurrentItemIndex = 0;
+	const ItemType* CurrentItem = nullptr;
+#if TRACESERVICES_PAGED_ARRAY_ITERATOR_DEBUG_ENABLED
+	uint64 TotalPageCount = 0;
+	uint64 TotalItemCount = 0;
+#endif
+};
+
+#else // TRACESERVICES_PAGED_ARRAY_ITERATOR_V2
+
+template<typename ItemType, typename PageType>
+class TPagedArrayIterator
+{
+public:
+	TPagedArrayIterator()
+	{
+	}
+
+	TPagedArrayIterator(const TPagedArray<ItemType, PageType>& InOuter, uint64 InItemIndex)
+		: Outer(&InOuter)
+		, CurrentItemIndex(InItemIndex)
+	{
+	}
+
+	//////////////////////////////////////////////////
+	// Page Iterator
+
+	uint64 GetCurrentPageIndex() const
+	{
+		return CurrentItemIndex / Outer->PageSize;
+	}
+
+	const PageType* GetCurrentPage() const
+	{
+		return Outer->FirstPage + CurrentItemIndex / Outer->PageSize;
+	}
+
+	const PageType* SetCurrentPage(uint64 PageIndex)
+	{
+		CurrentItemIndex = PageIndex * Outer->PageSize;
+		return Outer->FirstPage + PageIndex;
+	}
+
+	const PageType* PrevPage()
+	{
+		uint64 PageIndex = CurrentItemIndex / Outer->PageSize;
+		if (PageIndex > 0)
+		{
+			--PageIndex;
+			CurrentItemIndex = PageIndex * Outer->PageSize;
+			return Outer->FirstPage + PageIndex;
+		}
+		else
+		{
+			CurrentItemIndex = Outer->Num();
+			return nullptr;
+		}
+	}
+
+	const PageType* NextPage()
+	{
+		uint64 PageIndex = CurrentItemIndex / Outer->PageSize + 1;
+		if (PageIndex < Outer->NumPages())
+		{
+			CurrentItemIndex = PageIndex * Outer->PageSize;
+			return Outer->FirstPage + PageIndex;
+		}
+		else
+		{
+			CurrentItemIndex = Outer->Num();
+			return nullptr;
+		}
+	}
+
+	//////////////////////////////////////////////////
+	// Item Iterator
+
+	uint64 GetCurrentItemIndex() const
+	{
+		return CurrentItemIndex;
+	}
+
+	const ItemType* GetCurrentItem() const
+	{
+		return &(*Outer)[CurrentItemIndex];
+	}
+
+	const ItemType* SetPosition(uint64 Index)
+	{
+		CurrentItemIndex = Index;
+		return GetCurrentItem();
+	}
+
+	const ItemType* PrevItem()
+	{
+		if (CurrentItemIndex > 0)
+		{
+			--CurrentItemIndex;
+			return GetCurrentItem();
+		}
+		else
+		{
+			return nullptr;
+		}
+	}
+
+	const ItemType* NextItem()
+	{
+		if (CurrentItemIndex + 1 < Outer->Num())
+		{
+			++CurrentItemIndex;
+			return GetCurrentItem();
+		}
+		else
+		{
+			CurrentItemIndex = Outer->Num();
+			return nullptr;
+		}
+	}
+
+	//////////////////////////////////////////////////
+	// operators
+
+	const ItemType& operator*() const
+	{
+		return (*Outer)[CurrentItemIndex];
+	}
+
+	const ItemType* operator->() const
+	{
+		return &(*Outer)[CurrentItemIndex];
+	}
+
+	explicit operator bool() const
+	{
+		return CurrentItemIndex < Outer->Num();
+	}
+
+	TPagedArrayIterator& operator++()
+	{
+		++CurrentItemIndex;
+		return *this;
+	}
+
+	TPagedArrayIterator operator++(int)
+	{
+		TPagedArrayIterator Tmp(*this);
+		++CurrentItemIndex;
+		return Tmp;
+	}
+
+	TPagedArrayIterator& operator--()
+	{
+		--CurrentItemIndex;
+		return *this;
+	}
+
+	TPagedArrayIterator operator--(int)
+	{
+		TPagedArrayIterator Tmp(*this);
+		--CurrentItemIndex;
+		return Tmp;
+	}
+
+private:
+	FORCEINLINE friend bool operator!=(const TPagedArrayIterator& Lhs, const TPagedArrayIterator& Rhs)
+	{
+		checkSlow(Lhs.Outer == Rhs.Outer); // Needs to be iterators of the same array
+		return Lhs.CurrentItemIndex != Rhs.CurrentItemIndex;
+	}
+
+	const TPagedArray<ItemType, PageType>* Outer = nullptr;
+	uint64 CurrentItemIndex = 0;
 };
 
+#endif // TRACESERVICES_PAGED_ARRAY_ITERATOR_V2
+
 template<typename InItemType, typename InPageType = TPagedArrayPage<InItemType>>
 class TPagedArray
 {
@@ -371,19 +671,17 @@ public:
 
 	TIterator GetIterator() const
 	{
-		return TIterator(*this, 0, 0);
+		return TIterator(*this, 0);
 	}
 
 	TIterator GetIteratorFromPage(uint64 PageIndex) const
 	{
-		return TIterator(*this, PageIndex, 0);
+		return TIterator(*this, PageIndex * PageSize);
 	}
 
 	TIterator GetIteratorFromItem(uint64 ItemIndex) const
 	{
-		uint64 PageIndex = ItemIndex / PageSize;
-		uint64 IndexInPage = ItemIndex % PageSize;
-		return TIterator(*this, PageIndex, IndexInPage);
+		return TIterator(*this, ItemIndex);
 	}
 
 	const PageType* GetPages() const
@@ -429,10 +727,10 @@ public:
 		return *Item;
 	}
 
-	FORCEINLINE TIterator begin() { return TIterator(*this, 0, 0); }
-	FORCEINLINE TIterator begin() const { return TIterator(*this, 0, 0); }
-	FORCEINLINE TIterator end() { return TIterator(*this, NumPages()-1, LastPage->Count - 1); }
-	FORCEINLINE TIterator end() const { return TIterator(*this, NumPages()-1, LastPage->Count - 1); }
+	FORCEINLINE TIterator begin() { return TIterator(*this, 0); }
+	FORCEINLINE TIterator begin() const { return TIterator(*this, 0); }
+	FORCEINLINE TIterator end() { return TIterator(*this, TotalItemCount); }
+	FORCEINLINE TIterator end() const { return TIterator(*this, TotalItemCount); }
 
 private:
 	template<typename ItemType, typename PageType>
