[Insights] Optimizations to MonotonicTimeline implementation (avoids using ScopeEntryIterator.GetCurrentPage() after NextPage/PrevPage or when page index is known).

diff --git a/Engine/Source/Developer/TraceServices/Public/Model/MonotonicTimeline.h b/Engine/Source/Developer/TraceServices/Public/Model/MonotonicTimeline.h
--- a/Engine/Source/Developer/TraceServices/Public/Model/MonotonicTimeline.h
+++ b/Engine/Source/Developer/TraceServices/Public/Model/MonotonicTimeline.h
@@ -216,8 +216,7 @@ public:
 		{
 			--FirstScopePageIndex;
 		}
-		auto ScopeEntryIterator = DetailLevel.ScopeEntries.GetIteratorFromPage(FirstScopePageIndex);
-		const FEventScopeEntryPage* ScopePage = ScopeEntryIterator.GetCurrentPage();
+		const FEventScopeEntryPage* ScopePage = DetailLevel.ScopeEntries.GetPage(FirstScopePageIndex);
 		if (ScopePage->BeginTime > IntervalEnd)
 		{
 			return;
@@ -226,7 +225,7 @@ public:
 		{
 			return;
 		}
-		auto EventsIterator = DetailLevel.Events.GetIteratorFromItem(ScopePage->BeginEventIndex);
+
 		struct FEnumerationStackEntry
 		{
 			double StartTime;
@@ -242,8 +241,12 @@ public:
 			EnumerationStackEntry.Event = DetailLevel.GetEvent(EventStackEntry.EventIndex);
 		}
 
+		auto ScopeEntryIterator = DetailLevel.ScopeEntries.GetIteratorFromPage(FirstScopePageIndex);
 		const FEventScopeEntry* ScopeEntry = ScopeEntryIterator.GetCurrentItem();
+
+		auto EventsIterator = DetailLevel.Events.GetIteratorFromItem(ScopePage->BeginEventIndex);
 		const EventType* Event = EventsIterator.GetCurrentItem();
+
 		while (ScopeEntry && FMath::Abs(ScopeEntry->Time) < IntervalStart)
 		{
 			if (ScopeEntry->Time < 0.0)
@@ -334,11 +337,10 @@ public:
 
 		if (bSearchEndTimeUsingPages)
 		{
+			const FEventScopeEntryPage* CurrentScopePage = ScopeEntryIterator.GetCurrentPage();
 			do
 			{
-				const FEventScopeEntryPage* CurrentScopePage = ScopeEntryIterator.GetCurrentPage();
 				check(CurrentStackDepth <= CurrentScopePage->InitialStackCount);
-
 				while (CurrentStackDepth > 0 && CurrentScopePage->InitialStack[CurrentStackDepth - 1].EndTime > 0)
 				{
 					--CurrentStackDepth;
@@ -348,7 +350,9 @@ public:
 						return;
 					}
 				}
-			} while (ScopeEntryIterator.NextPage());
+				CurrentScopePage = ScopeEntryIterator.NextPage();
+			}
+			while (CurrentScopePage != nullptr);
 		}
 
 		while (CurrentStackDepth > 0)
@@ -400,7 +404,7 @@ public:
 		{
 			// If we have a page we can start from, start enumerating backwards from the begining of that page.
 			ScopeEntryIterator = DetailLevel.ScopeEntries.GetIteratorFromPage(LastScopePageIndex);
-			const FEventScopeEntryPage* ScopePage = ScopeEntryIterator.GetCurrentPage();
+			const FEventScopeEntryPage* ScopePage = DetailLevel.ScopeEntries.GetPage(LastScopePageIndex);
 
 			EventsIterator = DetailLevel.Events.GetIteratorFromItem(ScopePage->BeginEventIndex);
 			CurrentStackDepth = ScopePage->InitialStackCount;
@@ -431,7 +435,6 @@ public:
 			// We start enumerating from the previous page.
 			ScopeEntry = ScopeEntryIterator.PrevItem();
 			EventsIterator.PrevItem();
-
 		}
 		else
 		{
@@ -542,11 +545,10 @@ public:
 
 		if (bSearchStartTimeUsingPages)
 		{
+			const FEventScopeEntryPage* CurrentScopePage = ScopeEntryIterator.GetCurrentPage();
 			do
 			{
-				const FEventScopeEntryPage* CurrentScopePage = ScopeEntryIterator.GetCurrentPage();
 				check(CurrentStackDepth <= CurrentScopePage->InitialStackCount);
-
 				while (CurrentStackDepth > 0)
 				{
 					--CurrentStackDepth;
@@ -557,7 +559,9 @@ public:
 						return;
 					}
 				}
-			} while (ScopeEntryIterator.PrevPage());
+				CurrentScopePage = ScopeEntryIterator.PrevPage();
+			}
+			while (CurrentScopePage != nullptr);
 		}
 	}
 
@@ -801,7 +805,7 @@ public:
 		{
 			--FirstScopePageIndex;
 			ScopeEntryIterator = DetailLevel.ScopeEntries.GetIteratorFromPage(FirstScopePageIndex);
-			ScopePage = ScopeEntryIterator.GetCurrentPage();
+			ScopePage = DetailLevel.ScopeEntries.GetPage(FirstScopePageIndex);
 		}
 
 		auto EventsIterator = DetailLevel.Events.GetIteratorFromItem(ScopePage->BeginEventIndex);
@@ -917,16 +921,14 @@ public:
 		ScopeEntryIterator = DetailLevel.ScopeEntries.GetIteratorFromPage(FirstScopePageIndex);
 		auto EventLastPageIterator = ScopeEntryIterator;
 
-		while (ScopeEntryIterator.NextPage() != nullptr)
+		while (const FEventScopeEntryPage* CurrentScopePage = ScopeEntryIterator.NextPage())
 		{
-			const FEventScopeEntryPage* CurrentScopePage = ScopeEntryIterator.GetCurrentPage();
-
 			if (CurrentScopePage->InitialStackCount <= Depth)
 			{
 				break;
 			}
-			double StartTime = DetailLevel.GetScopeEntryTime(CurrentScopePage->InitialStack[Depth].EnterScopeIndex);
 
+			double StartTime = DetailLevel.GetScopeEntryTime(CurrentScopePage->InitialStack[Depth].EnterScopeIndex);
 			if (StartTime != TargetEntry.StartTime)
 			{
 				break;
@@ -1021,8 +1023,7 @@ public:
 		{
 			--FirstScopePageIndex;
 		}
-		auto ScopeEntryIterator = DetailLevel.ScopeEntries.GetIteratorFromPage(FirstScopePageIndex);
-		const FEventScopeEntryPage* ScopePage = ScopeEntryIterator.GetCurrentPage();
+		const FEventScopeEntryPage* ScopePage = DetailLevel.ScopeEntries.GetPage(FirstScopePageIndex);
 		if (ScopePage->BeginTime > Time)
 		{
 			return 0;
@@ -1030,6 +1031,7 @@ public:
 
 		int32 CurrentStackDepth = ScopePage->InitialStackCount;
 
+		auto ScopeEntryIterator = DetailLevel.ScopeEntries.GetIteratorFromPage(FirstScopePageIndex);
 		const FEventScopeEntry* ScopeEntry = ScopeEntryIterator.GetCurrentItem();
 		while (ScopeEntry && FMath::Abs(ScopeEntry->Time) < Time)
 		{
@@ -1224,11 +1226,11 @@ private:
 			}
 		}
 
-		const FEventScopeEntryPage* NextScopePage = nullptr;
-		if (!bIsInCurrentPageInitStack && ScopeEntryIterator.NextPage() != nullptr)
+		if (!bIsInCurrentPageInitStack)
 		{
-			NextScopePage = ScopeEntryIterator.GetCurrentPage();
-			if (NextScopePage->InitialStackCount > Depth)
+			const FEventScopeEntryPage* NextScopePage = ScopeEntryIterator.NextPage();
+			if (NextScopePage != nullptr &&
+				NextScopePage->InitialStackCount > Depth)
 			{
 				FEventStackEntry& NextPageStackEntry = NextScopePage->InitialStack[Depth];
 				double StartTime = DetailLevel.GetScopeEntryTime(NextPageStackEntry.EnterScopeIndex);
@@ -1245,9 +1247,14 @@ private:
 			return false;
 		}
 
-		while (ScopeEntryIterator.NextPage() != nullptr && CurrentScopePage->InitialStack[Depth].EndTime < 0)
+		while (CurrentScopePage->InitialStack[Depth].EndTime < 0)
 		{
-			CurrentScopePage = ScopeEntryIterator.GetCurrentPage();
+			const FEventScopeEntryPage* NextScopePage = ScopeEntryIterator.NextPage();
+			if (NextScopePage == nullptr)
+			{
+				break;
+			}
+			CurrentScopePage = NextScopePage;
 		}
 
 		OutPageStackEntry = CurrentScopePage->InitialStack[Depth];
