[Insights] Detect cases where it is not possible to recover ucache file

Due to how ucache files are written, table of contents is only updated when cache system is released on application shutdown, if the application crashes blocks can be written to the file but missing from the TOC. This change adds logic to detect the case where a cache file has been created and contains blocks, but the no blocks were read from the TOC.

diff --git a/Engine/Source/Developer/TraceServices/Private/AnalysisCache.cpp b/Engine/Source/Developer/TraceServices/Private/AnalysisCache.cpp
--- a/Engine/Source/Developer/TraceServices/Private/AnalysisCache.cpp
+++ b/Engine/Source/Developer/TraceServices/Private/AnalysisCache.cpp
@@ -39,7 +39,7 @@ FAnalysisCache::FFileContents::FFileContents(const TCHAR* FilePath)
 	// 1. File does not exist, create on first save
 	// 2. File exist, we can read the contents
 	// 3. File exist but we could not open the file for reading. Multiple processes are competing. Put the cache in transient mode
-	if (PlatformFile.FileExists(*CacheFilePath))
+	if (const int64 FileSize = PlatformFile.FileSize(*CacheFilePath) > 0)
 	{
 		if (const TUniquePtr<IFileHandle> File(PlatformFile.OpenRead(*CacheFilePath)); File.IsValid())
 		{
@@ -49,6 +49,16 @@ FAnalysisCache::FFileContents::FFileContents(const TCHAR* FilePath)
 				//todo: Recover by deleting file?
 			}
 
+			// Additional sanity check. A common error scenario is that Insights crashed after writing block but before
+			// committing them to the table of contents. Detect that scenario here.
+			if (FileSize > ReservedSize && Blocks.IsEmpty())
+			{
+				UE_LOG(LogAnalysisCache, Error, TEXT("Cache file has written several blocks but table of contents contains no blocks. This is likely caused by abnormal program termination. Please delete \"%s\" file. Putting cache in transient mode."), *CacheFilePath);
+				IndexEntries.Empty();
+				bTransientMode = true;
+				return;
+			}
+
 			UE_LOG(LogAnalysisCache, VeryVerbose, TEXT("Cache contains %d blocks:"), Blocks.Num());
 			UE_LOG(LogAnalysisCache, VeryVerbose, TEXT("   %10s   %10s   %13s   %13s   %13s"), TEXT("Cache index"), TEXT("Block index"), TEXT("Offset"), TEXT("Uncompressed"), TEXT("Compressed"));
 
@@ -86,7 +96,7 @@ FAnalysisCache::FFileContents::~FFileContents()
 //////////////////////////////////////////////////////////////////////
 FCacheId FAnalysisCache::FFileContents::GetId(const TCHAR* Name, uint16 Flags)
 {
-	FIndexEntry* Entry = Algo::FindByPredicate(IndexEntries, [Name](const FIndexEntry& Entry){ return Entry.Name.Equals(Name); });
+	const FIndexEntry* Entry = Algo::FindByPredicate(IndexEntries, [Name](const FIndexEntry& Entry){ return Entry.Name.Equals(Name); });
 	if (Entry)
 	{
 		return Entry->Id;
@@ -104,7 +114,7 @@ FCacheId FAnalysisCache::FFileContents::GetId(const TCHAR* Name, uint16 Flags)
 //////////////////////////////////////////////////////////////////////
 uint16 FAnalysisCache::FFileContents::GetFlags(FCacheId InId)
 {
-	FIndexEntry* Entry = Algo::FindBy(IndexEntries, InId, [](const FIndexEntry& InEntry){ return InEntry.Id; });
+	const FIndexEntry* Entry = Algo::FindBy(IndexEntries, InId, &FIndexEntry::Id);
 	if (Entry)
 	{
 		return uint16(Entry->Flags & 0xffff);
@@ -115,7 +125,7 @@ uint16 FAnalysisCache::FFileContents::GetFlags(FCacheId InId)
 //////////////////////////////////////////////////////////////////////
 FMutableMemoryView FAnalysisCache::FFileContents::GetUserData(FCacheId InId)
 {
-	FIndexEntry* Entry = Algo::FindBy(IndexEntries, InId, [](const FIndexEntry& InEntry){ return InEntry.Id; });
+	FIndexEntry* Entry = Algo::FindBy(IndexEntries, InId, &FIndexEntry::Id);
 	if (Entry)
 	{
 		return FMutableMemoryView(Entry->UserData, UserDataSize);
@@ -294,8 +304,8 @@ uint64 FAnalysisCache::FFileContents::LoadBlock(FMutableMemoryView Block, BlockK
 		UE_LOG(LogAnalysisCache, Error, TEXT("Trying to load unknown block 0x%x."), BlockKey);
 		return 0;
 	}
-	
-	FBlockEntry& Entry = Blocks[EntryIndex];
+
+	const FBlockEntry& Entry = Blocks[EntryIndex];
 	check(Entry.UncompressedSize <= Block.GetSize());
 	
 	if (!File->Seek(Entry.Offset))
@@ -392,7 +402,7 @@ FAnalysisCache::FAnalysisCache(const TCHAR* Path)
 	Contents = MakeUnique<FFileContents>(*CacheFilePath);
 
 	// Build a dictionary of number of blocks per cache id.
-	for (FFileContents::FBlockEntry& Block : Contents->Blocks)
+	for (const FFileContents::FBlockEntry& Block : Contents->Blocks)
 	{
 		const uint32 CacheId = GetCacheId(Block.BlockKey);
 		IndexBlockCount.FindOrAdd(CacheId)++;
@@ -504,7 +514,7 @@ void FAnalysisCache::ReleaseBlocks(uint8* BlockBuffer, FCacheId CacheId, uint32
 	const uint32 BlockCount = Size / IAnalysisCache::BlockSizeBytes;
 	for (uint32 Block = 0; Block < BlockCount; ++Block)
 	{
-		void* BlockStart = BlockBuffer + (Block * IAnalysisCache::BlockSizeBytes);
+		const void* BlockStart = BlockBuffer + (Block * IAnalysisCache::BlockSizeBytes);
 		const FMemoryView BlockView = FMemoryView(BlockStart, IAnalysisCache::BlockSizeBytes);
 		const BlockKeyType BlockKey = CreateBlockKey(CacheId, BlockIndexStart + Block);
 		const uint64 BytesWritten = Contents->UpdateBlock(BlockView, BlockKey);
