[Insights] Add functionality to undecorate symbols in case of stripped pdbs

When pdbs are stripped of private symbols, some symbols appeared as decorated in memory callstack views. On Windows we can use UnDecorateSymbol to get a human readable string.

diff --git a/Engine/Source/Developer/TraceServices/Private/Model/SymslibResolver.cpp b/Engine/Source/Developer/TraceServices/Private/Model/SymslibResolver.cpp
--- a/Engine/Source/Developer/TraceServices/Private/Model/SymslibResolver.cpp
+++ b/Engine/Source/Developer/TraceServices/Private/Model/SymslibResolver.cpp
@@ -29,6 +29,15 @@
 /////////////////////////////////////////////////////////////////////
 DEFINE_LOG_CATEGORY_STATIC(LogSymslib, Log, All);
 
+#if PLATFORM_WINDOWS
+	#define USE_DBG_HELP_UNDECORATOR 1
+	#if USE_DBG_HELP_UNDECORATOR
+		#include <Microsoft/AllowMicrosoftPlatformTypes.h>
+		#include "DbgHelp.h"
+		#include <Microsoft/HideMicrosoftPlatformTypes.h>
+	#endif
+#endif
+
 /////////////////////////////////////////////////////////////////////
 namespace TraceServices {
 
@@ -1066,6 +1075,38 @@ void FSymslibResolver::ResolveSymbolTracked(uint64 Address, FResolvedSymbol& Tar
 	}
 }
 
+#if USE_DBG_HELP_UNDECORATOR
+FORCEINLINE void UndecorateAndCopySymbolName(ANSICHAR* OutStr, const char* Name, uint32 MaxStringLength)
+{
+	// In case of stripped names we use the UnDecorateSymbolName function from the DbgHelp.
+	if (Name[0] == '?')
+	{
+		// todo: To get full parity with non-stripped pdbs we shouldn't print the argument types. There is a flag
+		//		 UNDNAME_NO_ARGUMENTS however whenever I tried that the symbol was not undecorated at all?
+		constexpr DWORD Flags =
+			UNDNAME_NO_MS_KEYWORDS |
+			UNDNAME_NO_FUNCTION_RETURNS |
+			UNDNAME_NO_ALLOCATION_MODEL |
+			UNDNAME_NO_ALLOCATION_LANGUAGE |
+			UNDNAME_NO_THISTYPE |
+			UNDNAME_NO_ACCESS_SPECIFIERS |
+			UNDNAME_NO_THROW_SIGNATURES |
+			UNDNAME_NO_MEMBER_TYPE |
+			UNDNAME_NO_RETURN_UDT_MODEL;
+
+		// It is unclear from official documentation if UnDecorateSymbolName is thread safe or not, so
+		// we take a local lock just in case.
+		static FCriticalSection Cs;
+		FScopeLock Lock(&Cs);
+		DWORD Length = UnDecorateSymbolName(Name, OutStr, MaxStringLength, Flags);
+	}
+	else
+	{
+		FCStringAnsi::Strncpy(OutStr, Name, MaxStringLength);
+	}
+}
+#endif // PLATFORM_WINDOWS
+
 bool FSymslibResolver::ResolveSymbol(uint64 Address, FResolvedSymbol& Target, FSymbolStringAllocator& StringAllocator, FModuleEntry* Entry) const
 {
 	EModuleStatus Status = Entry->Module->Status.load();
@@ -1168,7 +1209,11 @@ bool FSymslibResolver::ResolveSymbol(uint64 Address, FResolvedSymbol& Target, FS
 	}
 
 	ANSICHAR SymbolName[MaxStringSize];
+#if USE_DBG_HELP_UNDECORATOR
+	UndecorateAndCopySymbolName(SymbolName, SymsSymbol->Name, MaxStringSize);
+#else
 	FCStringAnsi::Strncpy(SymbolName, SymsSymbol->Name, MaxStringSize);
+#endif
 	SymbolName[MaxStringSize - 1] = 0;
 	const TCHAR* SymbolNamePersistent =  StringAllocator.Store(ANSI_TO_TCHAR(SymbolName));
 
