Fix CIS error

diff --git a/Engine/Source/Developer/TraceServices/Private/AnalysisCache.cpp b/Engine/Source/Developer/TraceServices/Private/AnalysisCache.cpp
--- a/Engine/Source/Developer/TraceServices/Private/AnalysisCache.cpp
+++ b/Engine/Source/Developer/TraceServices/Private/AnalysisCache.cpp
@@ -39,7 +39,7 @@ FAnalysisCache::FFileContents::FFileContents(const TCHAR* FilePath)
 	// 1. File does not exist, create on first save
 	// 2. File exist, we can read the contents
 	// 3. File exist but we could not open the file for reading. Multiple processes are competing. Put the cache in transient mode
-	if (const int64 FileSize = PlatformFile.FileSize(*CacheFilePath) > 0)
+	if (const int64 FileSize = PlatformFile.FileSize(*CacheFilePath); FileSize > 0)
 	{
 		if (const TUniquePtr<IFileHandle> File(PlatformFile.OpenRead(*CacheFilePath)); File.IsValid())
 		{
@@ -53,7 +53,7 @@ FAnalysisCache::FFileContents::FFileContents(const TCHAR* FilePath)
 			// committing them to the table of contents. Detect that scenario here.
 			if (FileSize > ReservedSize && Blocks.IsEmpty())
 			{
-				UE_LOG(LogAnalysisCache, Error, TEXT("Cache file has written several blocks but table of contents contains no blocks. This is likely caused by abnormal program termination. Please delete \"%s\" file. Putting cache in transient mode."), *CacheFilePath);
+				UE_LOG(LogAnalysisCache, Error, TEXT("Cache file has written several blocks but table of contents contains no blocks. This is likely caused by abnormal program termination. Please delete \"%s\". Putting cache in transient mode."), *CacheFilePath);
 				IndexEntries.Empty();
 				bTransientMode = true;
 				return;
