[Insights] Fixed access to DefinitonProvider.

diff --git a/Engine/Source/Developer/TraceServices/Private/AnalysisService.cpp b/Engine/Source/Developer/TraceServices/Private/AnalysisService.cpp
--- a/Engine/Source/Developer/TraceServices/Private/AnalysisService.cpp
+++ b/Engine/Source/Developer/TraceServices/Private/AnalysisService.cpp
@@ -255,7 +255,7 @@ TSharedPtr<const IAnalysisSession> FAnalysisService::StartAnalysis(uint32 TraceI
 	Session->AddAnalyzer(new FStringsAnalyzer(*Session));
 
 	TSharedPtr<FDefinitionProvider> DefProvider = MakeShared<FDefinitionProvider>(&Session.Get());
-	Session->AddProvider(FDefinitionProvider::ProviderName, DefProvider);
+	Session->AddProvider(FDefinitionProvider::ProviderName, DefProvider, DefProvider);
 
 	ModuleService.OnAnalysisBegin(*Session);
 
diff --git a/Engine/Source/Developer/TraceServices/Private/Analyzers/StringsAnalyzer.cpp b/Engine/Source/Developer/TraceServices/Private/Analyzers/StringsAnalyzer.cpp
--- a/Engine/Source/Developer/TraceServices/Private/Analyzers/StringsAnalyzer.cpp
+++ b/Engine/Source/Developer/TraceServices/Private/Analyzers/StringsAnalyzer.cpp
@@ -1,4 +1,5 @@
 ﻿// Copyright Epic Games, Inc. All Rights Reserved.
+
 #include "StringsAnalyzer.h"
 #include "../Common/Utils.h"
 #include "TraceServices/Model/Definitions.h"
@@ -6,8 +7,9 @@
 
 namespace TraceServices
 {
-	
-/////////////////////////////////////////////////////////////////////////////////////////
+
+////////////////////////////////////////////////////////////////////////////////////////////////////
+
 void FStringsAnalyzer::OnAnalysisBegin(const FOnAnalysisContext& Context)
 {
 	Context.InterfaceBuilder.RouteEvent(RouteId_StaticString, "Strings", "StaticString");
@@ -16,7 +18,8 @@ void FStringsAnalyzer::OnAnalysisBegin(const FOnAnalysisContext& Context)
 	Context.InterfaceBuilder.RouteEvent(RouteId_FNameNoSync, "Strings", "FNameNoSync");
 }
 
-/////////////////////////////////////////////////////////////////////////////////////////
+////////////////////////////////////////////////////////////////////////////////////////////////////
+
 bool FStringsAnalyzer::OnEvent(uint16 RouteId, EStyle Style, const FOnEventContext& Context)
 {
 	switch (RouteId)
@@ -29,6 +32,7 @@ bool FStringsAnalyzer::OnEvent(uint16 RouteId, EStyle Style, const FOnEventConte
 		{
 			const FEventData& EventData = Context.EventData;
 			IDefinitionProvider* DefinitionProvider = GetDefinitionProvider(Session);
+			check(DefinitionProvider != nullptr);
 
 			FWideStringView DisplayWide;
 			FAnsiStringView DisplayAnsi;
@@ -69,4 +73,6 @@ bool FStringsAnalyzer::OnEvent(uint16 RouteId, EStyle Style, const FOnEventConte
 	return true;
 }
 
-}
+////////////////////////////////////////////////////////////////////////////////////////////////////
+
+} // namespace TraceServices
diff --git a/Engine/Source/Developer/TraceServices/Private/Model/DefinitionProvider.cpp b/Engine/Source/Developer/TraceServices/Private/Model/DefinitionProvider.cpp
--- a/Engine/Source/Developer/TraceServices/Private/Model/DefinitionProvider.cpp
+++ b/Engine/Source/Developer/TraceServices/Private/Model/DefinitionProvider.cpp
@@ -1,4 +1,5 @@
 ﻿// Copyright Epic Games, Inc. All Rights Reserved.
+
 #include "DefinitionProvider.h"
 #include "HAL/UnrealMemory.h"
 #include "UObject/NameTypes.h"
@@ -6,28 +7,33 @@
 namespace TraceServices
 {
 
-//////////////////////////////////////////////////////////////////////////////////////////
-const FName FDefinitionProvider::ProviderName(TEXT("DefinitionProvider"));	
-	
-//////////////////////////////////////////////////////////////////////////////////////////
+////////////////////////////////////////////////////////////////////////////////////////////////////
+
+const FName FDefinitionProvider::ProviderName(TEXT("DefinitionProvider"));
+
+////////////////////////////////////////////////////////////////////////////////////////////////////
+
 FDefinitionProvider::FDefinitionProvider(IAnalysisSession* InSession)
 	: PageRemain(0)
 {
 }
 
-//////////////////////////////////////////////////////////////////////////////////////////
+////////////////////////////////////////////////////////////////////////////////////////////////////
+
 void FDefinitionProvider::AddEntry(uint64 Hash, const void* Ptr)
 {
 	Definitions.Add(Hash, Ptr);
 }
 
-//////////////////////////////////////////////////////////////////////////////////////////
+////////////////////////////////////////////////////////////////////////////////////////////////////
+
 const void* FDefinitionProvider::FindEntry(uint64 Hash) const
 {
 	return Definitions.FindRef(Hash);
 }
 
-//////////////////////////////////////////////////////////////////////////////////////////
+////////////////////////////////////////////////////////////////////////////////////////////////////
+
 void* FDefinitionProvider::Allocate(uint32 Size, uint32 Alignment)
 {
 	check(Size);
@@ -40,7 +46,7 @@ void* FDefinitionProvider::Allocate(uint32 Size, uint32 Alignment)
 	else
 	{
 		auto& NewPage = Pages.Emplace_GetRef((uint8*)FMemory::MallocZeroed(PageSize));
-		
+
 		Dest = NewPage.Get();
 		PageRemain = PageSize;
 	}
@@ -48,15 +54,20 @@ void* FDefinitionProvider::Allocate(uint32 Size, uint32 Alignment)
 	return Dest;
 }
 
-//////////////////////////////////////////////////////////////////////////////////////////
+////////////////////////////////////////////////////////////////////////////////////////////////////
+
 IDefinitionProvider* GetDefinitionProvider(IAnalysisSession& Session)
 {
-	return Session.EditProvider<IDefinitionProvider>("DefinitionProvider");
+	return Session.EditProvider<IDefinitionProvider>(FDefinitionProvider::ProviderName);
 }
-//////////////////////////////////////////////////////////////////////////////////////////
+
+////////////////////////////////////////////////////////////////////////////////////////////////////
+
 const IDefinitionProvider* ReadDefinitionProvider(const IAnalysisSession& Session)
 {
-	return Session.ReadProvider<IDefinitionProvider>("DefinitionProvider");
-}
-	
+	return Session.ReadProvider<IDefinitionProvider>(FDefinitionProvider::ProviderName);
 }
+
+////////////////////////////////////////////////////////////////////////////////////////////////////
+
+} // namespace TraceServices
diff --git a/Engine/Source/Developer/TraceServices/Private/Model/DefinitionProvider.h b/Engine/Source/Developer/TraceServices/Private/Model/DefinitionProvider.h
--- a/Engine/Source/Developer/TraceServices/Private/Model/DefinitionProvider.h
+++ b/Engine/Source/Developer/TraceServices/Private/Model/DefinitionProvider.h
@@ -16,18 +16,20 @@ class FDefinitionProvider : public IDefinitionProvider
 {
 public:
 	static const FName ProviderName;
-	
+
 	FDefinitionProvider(IAnalysisSession* InSession);
 
-	virtual void BeginEdit() const override { Lock.WriteLock(); };
-	virtual void EndEdit() const override { Lock.WriteUnlock(); };
-	virtual void BeginRead() const override { Lock.ReadLock(); };
-	virtual void EndRead() const override { Lock.ReadUnlock(); };
+	virtual void BeginEdit() const override { Lock.WriteLock(); }
+	virtual void EndEdit() const override { Lock.WriteUnlock(); }
+	virtual void BeginRead() const override { Lock.ReadLock(); }
+	virtual void EndRead() const override { Lock.ReadUnlock(); }
+
 private:
 	virtual void AddEntry(uint64 Hash, const void* Ptr) override;
 	virtual const void* FindEntry(uint64 Hash) const override;
 	virtual void* Allocate(uint32 Size, uint32 Alignment) override;
 
+private:
 	TArray<TUniquePtr<uint8>> Pages;
 	static constexpr uint32 PageSize = 1024;
 	uint32 PageRemain;
diff --git a/Engine/Source/Developer/TraceServices/Public/TraceServices/Model/Definitions.h b/Engine/Source/Developer/TraceServices/Public/TraceServices/Model/Definitions.h
--- a/Engine/Source/Developer/TraceServices/Public/TraceServices/Model/Definitions.h
+++ b/Engine/Source/Developer/TraceServices/Public/TraceServices/Model/Definitions.h
@@ -17,21 +17,20 @@ namespace TraceServices
  *
  * Example usage:
  * This code creates a definition, initializes the instance, resolves a reference and finally registers it.
- * \code 
+ * \code
  *		FMyDefinitionClass* Instance = DefinitionProvider->Create<FMyDefinitionClass>();
  *		Instance->SomeValue = EventData.GetValue<uint32>("SomeValue");
  *		const FEventRef Reference = EventData.GetReferenceValue("SomeString");
  *		Instance->SomeString = DefinitionProvider->Get<FStringDefinition>(Reference)->Display;
  *		...
- *		
+ *
  *		const UE::Trace::FEventRef Id(EventData.GetDefinitionId(), EventData.GetTypeInfo().GetId());
  *		DefinitionProvider->Register<FMyDefinitionClass>(Instance, Id);
  * \endcode
  */
-class IDefinitionProvider : public IProvider
+class IDefinitionProvider : public IEditableProvider
 {
 public:
-	
 	struct TRACESERVICES_API FEditScopeLock
 	{
 		FEditScopeLock(const IDefinitionProvider& InMetadataProvider)
@@ -64,6 +63,12 @@ public:
 		const IDefinitionProvider& MetadataProvider;
 	};
 
+public:
+	virtual ~IDefinitionProvider() = default;
+
+	//////////////////////////////////////////////////
+	// Edit operations
+
 	/**
 	 * Allocates memory for the definition. The memory is guaranteed to be valid during the lifetime of the
 	 * analysis session.
@@ -90,6 +95,9 @@ public:
 		AddEntry(Id.GetHash(), Instance);
 	}
 
+	//////////////////////////////////////////////////
+	// Read operations
+
 	/**
 	 * Attempts to retrieve a previously registered instance of a definition using a reference.
 	 * @param Reference Id used to uniquely identify the definition.
@@ -124,6 +132,7 @@ protected:
 	virtual void AddEntry(uint64 Hash, const void* Ptr) = 0;
 	virtual const void* FindEntry(uint64 Hash) const = 0;
 
+protected:
 	using StringifierFn = TFunction<FString (const void*)>;
 	TMap<uint32, StringifierFn> Stringifiers;
 };
@@ -132,5 +141,3 @@ TRACESERVICES_API IDefinitionProvider* GetDefinitionProvider(IAnalysisSession& S
 TRACESERVICES_API const IDefinitionProvider* ReadDefinitionProvider(const IAnalysisSession& Session);
 
 } // namespace TraceServices
-
-
