[Insights] Fixed regression bug where callstacks are not resolved anymore. The Module editable provider was not properly registered.

diff --git a/Engine/Source/Developer/TraceServices/Private/Analyzers/ModuleAnalysis.cpp b/Engine/Source/Developer/TraceServices/Private/Analyzers/ModuleAnalysis.cpp
--- a/Engine/Source/Developer/TraceServices/Private/Analyzers/ModuleAnalysis.cpp
+++ b/Engine/Source/Developer/TraceServices/Private/Analyzers/ModuleAnalysis.cpp
@@ -67,7 +67,7 @@ bool FModuleAnalyzer::OnEvent(uint16 RouteId, EStyle Style, const FOnEventContex
 					TSharedPtr<IModuleAnalysisProvider> ModuleProvider = CreateModuleProvider(Session, SymbolFormat);
 					if (ModuleProvider)
 					{
-						Session.AddProvider(GetModuleProviderName(), ModuleProvider);
+						Session.AddProvider(GetModuleProviderName(), ModuleProvider, ModuleProvider);
 						Provider = ModuleProvider.Get();
 					}
 				}
diff --git a/Engine/Source/Developer/TraceServices/Private/Model/AllocationsProvider.cpp b/Engine/Source/Developer/TraceServices/Private/Model/AllocationsProvider.cpp
--- a/Engine/Source/Developer/TraceServices/Private/Model/AllocationsProvider.cpp
+++ b/Engine/Source/Developer/TraceServices/Private/Model/AllocationsProvider.cpp
@@ -2084,8 +2084,7 @@ void FAllocationsProvider::DebugPrint() const
 
 IAllocationsProvider::FQueryHandle FAllocationsProvider::StartQuery(const IAllocationsProvider::FQueryParams& Params) const
 {
-	const FCallstacksProvider* CallstacksProvider = Session.ReadProvider<FCallstacksProvider>(FName("CallstacksProvider"));
-	auto* Inner = new FAllocationsQuery(*this, *CallstacksProvider, Params);
+	auto* Inner = new FAllocationsQuery(*this, Params);
 	return IAllocationsProvider::FQueryHandle(Inner);
 }
 
diff --git a/Engine/Source/Developer/TraceServices/Private/Model/AllocationsQuery.cpp b/Engine/Source/Developer/TraceServices/Private/Model/AllocationsQuery.cpp
--- a/Engine/Source/Developer/TraceServices/Private/Model/AllocationsQuery.cpp
+++ b/Engine/Source/Developer/TraceServices/Private/Model/AllocationsQuery.cpp
@@ -43,9 +43,8 @@ private:
 // FAllocationsQuery
 ////////////////////////////////////////////////////////////////////////////////////////////////////
 
-FAllocationsQuery::FAllocationsQuery(const FAllocationsProvider& InAllocationProvider, const FCallstacksProvider& InCallstacksProvider, const IAllocationsProvider::FQueryParams& InParams)
+FAllocationsQuery::FAllocationsQuery(const FAllocationsProvider& InAllocationProvider, const IAllocationsProvider::FQueryParams& InParams)
 	: AllocationsProvider(InAllocationProvider)
-	, CallstacksProvider(InCallstacksProvider)
 	, Params(InParams)
 	, IsWorking(true)
 	, IsCanceling(false)
@@ -143,7 +142,7 @@ void FAllocationsQuery::Run()
 			SbTree->Query(Cells, Params);
 		}
 	}
-		
+
 	CellCount += Cells.Num();
 
 	UE_LOG(LogTraceServices, Log, TEXT("[MemAlloc] %d cells to process"), Cells.Num());
@@ -160,7 +159,7 @@ void FAllocationsQuery::Run()
 		UE_LOG(LogTraceServices, Log, TEXT("[MemAlloc] Processing cell %u (%u allocs)..."), CellIndex, Cell->GetAllocCount());
 
 		FAllocationsImpl* CellResult = new FAllocationsImpl();
-			
+
 		Cell->Query(CellResult->Items, Params);
 
 		const uint32 NumAllocs = static_cast<uint32>(CellResult->Items.Num());
diff --git a/Engine/Source/Developer/TraceServices/Private/Model/AllocationsQuery.h b/Engine/Source/Developer/TraceServices/Private/Model/AllocationsQuery.h
--- a/Engine/Source/Developer/TraceServices/Private/Model/AllocationsQuery.h
+++ b/Engine/Source/Developer/TraceServices/Private/Model/AllocationsQuery.h
@@ -26,7 +26,7 @@ class FAllocationsQuery
 	friend class FAllocationsQueryAsyncTask;
 
 public:
-	FAllocationsQuery(const FAllocationsProvider& InAllocationsProvider, const FCallstacksProvider& InCallstacksProvider, const IAllocationsProvider::FQueryParams& InParams);
+	FAllocationsQuery(const FAllocationsProvider& InAllocationsProvider, const IAllocationsProvider::FQueryParams& InParams);
 
 	void Cancel();
 	IAllocationsProvider::FQueryStatus Poll();
@@ -37,7 +37,6 @@ private:
 
 private:
 	const FAllocationsProvider& AllocationsProvider;
-	const FCallstacksProvider& CallstacksProvider;
 
 	IAllocationsProvider::FQueryParams Params;
 
diff --git a/Engine/Source/Developer/TraceServices/Private/Modules/MemoryModule.cpp b/Engine/Source/Developer/TraceServices/Private/Modules/MemoryModule.cpp
--- a/Engine/Source/Developer/TraceServices/Private/Modules/MemoryModule.cpp
+++ b/Engine/Source/Developer/TraceServices/Private/Modules/MemoryModule.cpp
@@ -30,6 +30,7 @@ void FMemoryModule::OnAnalysisBegin(IAnalysisSession& Session)
 	Session.AddAnalyzer(new FMemoryAnalyzer(Session, MemoryProvider.Get()));
 
 	// Module
+	// The Module provider is created and registered by FModuleAnalyzer when the ModuleInit event is detected.
 	Session.AddAnalyzer(new FModuleAnalyzer(Session));
 
 	// Callstack
@@ -41,12 +42,11 @@ void FMemoryModule::OnAnalysisBegin(IAnalysisSession& Session)
 	TSharedPtr<FMetadataProvider> MetadataProvider = MakeShared<FMetadataProvider>(Session);
 	Session.AddProvider(GetMetadataProviderName(), MetadataProvider);
 	Session.AddAnalyzer(new FMetadataAnalysis(Session, MetadataProvider.Get()));
-	
+
 	// Allocations
 	TSharedPtr<FAllocationsProvider> AllocationsProvider = MakeShared<FAllocationsProvider>(Session, *MetadataProvider);
 	Session.AddProvider(GetAllocationsProviderName(), AllocationsProvider);
 	Session.AddAnalyzer(new FAllocationsAnalyzer(Session, *AllocationsProvider, *MetadataProvider));
-
 }
 
 void FMemoryModule::GetLoggers(TArray<const TCHAR*>& OutLoggers)
diff --git a/Engine/Source/Developer/TraceServices/Public/TraceServices/Model/Modules.h b/Engine/Source/Developer/TraceServices/Public/TraceServices/Model/Modules.h
--- a/Engine/Source/Developer/TraceServices/Public/TraceServices/Model/Modules.h
+++ b/Engine/Source/Developer/TraceServices/Public/TraceServices/Model/Modules.h
@@ -154,7 +154,7 @@ struct FModule
 };
 
 ////////////////////////////////////////////////////////////////////////////////
-class IModuleProvider : public IProvider
+class IModuleProvider : public IEditableProvider
 {
 public:
 	struct FStats
@@ -175,7 +175,6 @@ public:
 	 */
 	virtual const FResolvedSymbol* GetSymbol(uint64 Address) = 0;
 
-
 	/**
 	 * Gets the number of discovered modules.
 	 * @return Number of modules that were discovered so far.
