[Insights] Renamed MasterTimerId to OriginalTimerId.

diff --git a/Engine/Source/Developer/TraceServices/Private/Model/TimingProfiler.cpp b/Engine/Source/Developer/TraceServices/Private/Model/TimingProfiler.cpp
--- a/Engine/Source/Developer/TraceServices/Private/Model/TimingProfiler.cpp
+++ b/Engine/Source/Developer/TraceServices/Private/Model/TimingProfiler.cpp
@@ -123,12 +123,12 @@ void FTimingProfilerProvider::SetTimerNameAndLocation(uint32 TimerId, FStringVie
 	Timer.Line = Line;
 }
 
-uint32 FTimingProfilerProvider::AddMetadata(uint32 MasterTimerId, TArray<uint8>&& Metadata)
+uint32 FTimingProfilerProvider::AddMetadata(uint32 OriginalTimerId, TArray<uint8>&& Metadata)
 {
 	Session.WriteAccessCheck();
 
 	uint32 MetadataId = Metadatas.Num();
-	Metadatas.Add({MoveTemp(Metadata), MasterTimerId});
+	Metadatas.Add({ MoveTemp(Metadata), OriginalTimerId });
 
 	return ~MetadataId;
 }
diff --git a/Engine/Source/Developer/TraceServices/Private/Model/TimingProfilerPrivate.h b/Engine/Source/Developer/TraceServices/Private/Model/TimingProfilerPrivate.h
--- a/Engine/Source/Developer/TraceServices/Private/Model/TimingProfilerPrivate.h
+++ b/Engine/Source/Developer/TraceServices/Private/Model/TimingProfilerPrivate.h
@@ -43,7 +43,7 @@ public:
 	// implementing IEditableTimingProfilerProvider
 	virtual uint32 AddCpuTimer(FStringView Name, const TCHAR* File = nullptr, uint32 Line = 0) override;
 	virtual void SetTimerNameAndLocation(uint32 TimerId, FStringView Name, const TCHAR* File, uint32 Line) override;
-	virtual uint32 AddMetadata(uint32 MasterTimerId, TArray<uint8>&& Metadata) override;
+	virtual uint32 AddMetadata(uint32 OriginalTimerId, TArray<uint8>&& Metadata) override;
 	virtual TArrayView<const uint8> GetMetadata(uint32 TimerId) const override;
 	virtual IEditableTimeline<FTimingProfilerEvent>& GetCpuThreadEditableTimeline(uint32 ThreadId) override;
 
diff --git a/Engine/Source/Developer/TraceServices/Public/TraceServices/Model/TimingProfiler.h b/Engine/Source/Developer/TraceServices/Public/TraceServices/Model/TimingProfiler.h
--- a/Engine/Source/Developer/TraceServices/Public/TraceServices/Model/TimingProfiler.h
+++ b/Engine/Source/Developer/TraceServices/Public/TraceServices/Model/TimingProfiler.h
@@ -126,12 +126,12 @@ public:
 	/*
 	* Add metadata to a timer.
 	* 
-	* @param MasterTimerId	The identity of the timer to add metadata to.
-	* @param Metadata		The metadata.
+	* @param OriginalTimerId	The identity of the timer to add metadata to.
+	* @param Metadata			The metadata.
 	* 
 	* @return The identity of the metadata.
 	*/
-	virtual uint32 AddMetadata(uint32 MasterTimerId, TArray<uint8>&& Metadata) = 0;
+	virtual uint32 AddMetadata(uint32 OriginalTimerId, TArray<uint8>&& Metadata) = 0;
 
 	/*
 	* Get metadata for a timer.
