TasksInsight: minor changes required to support ParallelFor tasks:
* made some trace messages optional to reduce traffic and trace file size, as ParallelFor doesn't need them
* optimized out multiple task lookups per message

diff --git a/Engine/Source/Developer/TraceServices/Private/Model/TasksProfiler.cpp b/Engine/Source/Developer/TraceServices/Private/Model/TasksProfiler.cpp
--- a/Engine/Source/Developer/TraceServices/Private/Model/TasksProfiler.cpp
+++ b/Engine/Source/Developer/TraceServices/Private/Model/TasksProfiler.cpp
@@ -160,7 +160,19 @@ namespace TraceServices
 	{
 		InitTaskIdToIndexConversion(TaskId);
 
-		if (!TryRegisterEvent(TEXT("TaskScheduled"), TaskId, &FTaskInfo::ScheduledTimestamp, Timestamp, &FTaskInfo::ScheduledThreadId, ThreadId))
+		FTaskInfo* Task = TryGetOrCreateTask(TaskId);
+		if (Task == nullptr)
+		{
+			UE_LOG(LogTraceServices, Log, TEXT("TaskScheduled(TaskId %d, Timestamp %.6f) skipped"), TaskId, Timestamp);
+			return;
+		}
+
+		TaskScheduled(*Task, Timestamp, ThreadId);
+	}
+
+	void FTasksProvider::TaskScheduled(FTaskInfo& Task, double Timestamp, uint32 ThreadId)
+	{
+		if (!TryRegisterEvent(TEXT("TaskScheduled"), Task, &FTaskInfo::ScheduledTimestamp, Timestamp, &FTaskInfo::ScheduledThreadId, ThreadId))
 		{
 			return;
 		}
@@ -171,7 +183,7 @@ namespace TraceServices
 		}
 
 		WaitingForPrerequisitesTasksCounter->SetValue(Timestamp, --WaitingForPrerequisitesTasksNum);
-		if (IsNamedThread(TryGetTask(TaskId)->ThreadToExecuteOn))
+		if (IsNamedThread(Task.ThreadToExecuteOn))
 		{
 			NamedThreadsScheduledTasksCounter->SetValue(Timestamp, ++ScheduledTasksNum);
 		}
@@ -221,22 +233,26 @@ namespace TraceServices
 	{
 		InitTaskIdToIndexConversion(TaskId);
 
-		if (!TryRegisterEvent(TEXT("TaskStarted"), TaskId, &FTaskInfo::StartedTimestamp, Timestamp, &FTaskInfo::StartedThreadId, ThreadId))
+		FTaskInfo* Task = TryGetOrCreateTask(TaskId);
+		if (Task == nullptr)
 		{
+			UE_LOG(LogTraceServices, Log, TEXT("TaskStarted(TaskId %d, Timestamp %.6f) skipped"), TaskId, Timestamp);
 			return;
 		}
 
-		ExecutionThreads.FindOrAdd(ThreadId).Add(TaskId);
-
-		FTaskInfo* Task = TryGetTask(TaskId);
-		check(Task != nullptr);
+		if (Task->ScheduledTimestamp == FTaskInfo::InvalidTimestamp)
+		{
+			// optimisation for tasks that can't have dependencies, and so are scheduled immediately after launching. in this case "scheduled" trace is not sent
+			TaskScheduled(*Task, Task->LaunchedTimestamp, Task->LaunchedThreadId);
+		}
 
-		if (Task->ScheduledTimestamp == FTaskInfo::InvalidTimestamp) // "inline" tasks or task events are never scheduled
+		if (!TryRegisterEvent(TEXT("TaskStarted"), *Task, &FTaskInfo::StartedTimestamp, Timestamp, &FTaskInfo::StartedThreadId, ThreadId))
 		{
-			Task->ScheduledTimestamp = Timestamp;
-			Task->ScheduledThreadId = ThreadId;
+			return;
 		}
 
+		ExecutionThreads.FindOrAdd(ThreadId).Add(TaskId);
+
 		if (!bCountersCreated)
 		{
 			CreateCounters();
@@ -259,7 +275,18 @@ namespace TraceServices
 	{
 		InitTaskIdToIndexConversion(TaskId);
 
-		if (!TryRegisterEvent(TEXT("TaskFinished"), TaskId, &FTaskInfo::FinishedTimestamp, Timestamp))
+		FTaskInfo* Task = TryGetTask(TaskId);
+		if (Task == nullptr)
+		{
+			UE_LOG(LogTraceServices, Log, TEXT("TaskFinished(TaskId %d, Timestamp %.6f) skipped"), TaskId, Timestamp);
+		}
+
+		TaskFinished(*Task, Timestamp);
+	}
+
+	void FTasksProvider::TaskFinished(FTaskInfo& Task, double Timestamp)
+	{
+		if (!TryRegisterEvent(TEXT("TaskFinished"), Task, &FTaskInfo::FinishedTimestamp, Timestamp))
 		{
 			return;
 		}
@@ -270,10 +297,7 @@ namespace TraceServices
 		}
 
 		RunningTasksCounter->SetValue(Timestamp, --RunningTasksNum);
-	
-		FTaskInfo* Task = TryGetTask(TaskId);
-		check(Task != nullptr);
-		ExecutionTimeCounter->SetValue(Timestamp, Task->FinishedTimestamp - Task->StartedTimestamp);
+		ExecutionTimeCounter->SetValue(Timestamp, Task.FinishedTimestamp - Task.StartedTimestamp);
 	}
 
 	void FTasksProvider::TaskCompleted(TaskTrace::FId TaskId, double Timestamp, uint32 ThreadId)
@@ -294,19 +318,27 @@ namespace TraceServices
 		if (Task->ScheduledTimestamp == FTaskInfo::InvalidTimestamp) // task events are never scheduled or executed
 		{
 			// pretend scheduling and execution happened at the moment of completion
+			Task->LaunchedTimestamp = Timestamp;
+			Task->LaunchedThreadId = ThreadId;
 			Task->ScheduledTimestamp = Timestamp;
 			Task->ScheduledThreadId = ThreadId;
 			Task->StartedTimestamp = Timestamp;
 			Task->StartedThreadId = ThreadId;
 			Task->FinishedTimestamp = Timestamp;
 		}
+		else if (Task->FinishedTimestamp == FTaskInfo::InvalidTimestamp)
+		{
+			// optimisation for ParalellFor tasks that are finished and completed in one go, so no need to send separate trace messages
+			TaskFinished(*Task, Timestamp);
+		}
 
-		TryRegisterEvent(TEXT("TaskCompleted"), TaskId, &FTaskInfo::CompletedTimestamp, Timestamp, &FTaskInfo::CompletedThreadId, ThreadId);
+		TryRegisterEvent(TEXT("TaskCompleted"), *Task, &FTaskInfo::CompletedTimestamp, Timestamp, &FTaskInfo::CompletedThreadId, ThreadId);
 	}
 
 	void FTasksProvider::TaskDestroyed(TaskTrace::FId TaskId, double Timestamp, uint32 ThreadId)
 	{
 		InitTaskIdToIndexConversion(TaskId);
+
 		TryRegisterEvent(TEXT("TaskDestroyed"), TaskId, &FTaskInfo::DestroyedTimestamp, Timestamp, &FTaskInfo::DestroyedThreadId, ThreadId);
 	}
 
@@ -384,8 +416,6 @@ namespace TraceServices
 
 	bool FTasksProvider::TryRegisterEvent(const TCHAR* EventName, TaskTrace::FId TaskId, double FTaskInfo::* TimestampPtr, double TimestampValue, uint32 FTaskInfo::* ThreadIdPtr/* = nullptr*/, uint32 ThreadIdValue/* = 0*/)
 	{
-		UE_LOG(LogTraceServices, Verbose, TEXT("%s(TaskId: %d, Timestamp %.6f)"), EventName, TaskId, TimestampValue);
-
 		FTaskInfo* Task = TryGetTask(TaskId);
 		if (Task == nullptr)
 		{
@@ -393,11 +423,18 @@ namespace TraceServices
 			return false;
 		}
 
-		checkf(Task->*TimestampPtr == FTaskInfo::InvalidTimestamp, TEXT("%s: TaskId %d, old TS %.6f, new TS %.6f"), EventName, TaskId, Task->*TimestampPtr, TimestampValue);
-		Task->*TimestampPtr = TimestampValue;
+		return TryRegisterEvent(EventName, *Task, TimestampPtr, TimestampValue, ThreadIdPtr, ThreadIdValue);
+	}
+
+	bool FTasksProvider::TryRegisterEvent(const TCHAR* EventName, FTaskInfo& Task, double FTaskInfo::* TimestampPtr, double TimestampValue, uint32 FTaskInfo::* ThreadIdPtr/* = nullptr*/, uint32 ThreadIdValue/* = 0*/)
+	{
+		UE_LOG(LogTraceServices, Verbose, TEXT("%s(TaskId: %d, Timestamp %.6f)"), EventName, Task.Id, TimestampValue);
+
+		checkf(Task.*TimestampPtr == FTaskInfo::InvalidTimestamp, TEXT("%s: TaskId %d, old TS %.6f, new TS %.6f"), EventName, Task.Id, Task.*TimestampPtr, TimestampValue);
+		Task.*TimestampPtr = TimestampValue;
 		if (ThreadIdPtr != nullptr)
 		{
-			Task->*ThreadIdPtr = ThreadIdValue;
+			Task.*ThreadIdPtr = ThreadIdValue;
 		}
 
 		return true;
diff --git a/Engine/Source/Developer/TraceServices/Private/Model/TasksProfilerPrivate.h b/Engine/Source/Developer/TraceServices/Private/Model/TasksProfilerPrivate.h
--- a/Engine/Source/Developer/TraceServices/Private/Model/TasksProfilerPrivate.h
+++ b/Engine/Source/Developer/TraceServices/Private/Model/TasksProfilerPrivate.h
@@ -25,9 +25,11 @@ namespace TraceServices
 		void TaskCreated(TaskTrace::FId TaskId, double Timestamp, uint32 ThreadId);
 		void TaskLaunched(TaskTrace::FId TaskId, const TCHAR* DebugName, bool bTracked, int32 ThreadToExecuteOn, double Timestamp, uint32 ThreadId);
 		void TaskScheduled(TaskTrace::FId TaskId, double Timestamp, uint32 ThreadId);
+		void TaskScheduled(FTaskInfo& Task, double Timestamp, uint32 ThreadId);
 		void SubsequentAdded(TaskTrace::FId TaskId, TaskTrace::FId SubsequentId, double Timestamp, uint32 ThreadId);
 		void TaskStarted(TaskTrace::FId TaskId, double Timestamp, uint32 ThreadId);
 		void TaskFinished(TaskTrace::FId TaskId, double Timestamp);
+		void TaskFinished(FTaskInfo& Task, double Timestamp);
 		void TaskCompleted(TaskTrace::FId TaskId, double Timestamp, uint32 ThreadId);
 		void TaskDestroyed(TaskTrace::FId TaskId, double Timestamp, uint32 ThreadId);
 
@@ -56,8 +58,10 @@ namespace TraceServices
 
 		void CreateCounters();
 
-		// registers a minimal event with a name, task id, when and where it happened
+		// registers an event with a name, task id, when and where it happened
 		bool TryRegisterEvent(const TCHAR* EventName, TaskTrace::FId TaskId, double FTaskInfo::* TimestampPtr, double TimestampValue, uint32 FTaskInfo::* ThreadIdPtr = nullptr, uint32 ThreadIdValue = 0);
+		// registers an event for the given task
+		bool TryRegisterEvent(const TCHAR* EventName, FTaskInfo& Task, double FTaskInfo::* TimestampPtr, double TimestampValue, uint32 FTaskInfo::* ThreadIdPtr = nullptr, uint32 ThreadIdValue = 0);
 		// registers a relation with another task, when and where it was established
 		void AddRelative(const TCHAR* RelativeType, TaskTrace::FId TaskId, TArray<FTaskInfo::FRelationInfo> FTaskInfo::* RelationsPtr, TaskTrace::FId RelativeId, double Timestamp, uint32 ThreadId);
 
