fixed a bug reported by static analysis: https://horde.devtools.epicgames.com/job/62cede289b0d2dbd9cec9ecc?step=0cd0&issue=192735

diff --git a/Engine/Source/Developer/TraceServices/Private/Model/TasksProfiler.cpp b/Engine/Source/Developer/TraceServices/Private/Model/TasksProfiler.cpp
--- a/Engine/Source/Developer/TraceServices/Private/Model/TasksProfiler.cpp
+++ b/Engine/Source/Developer/TraceServices/Private/Model/TasksProfiler.cpp
@@ -279,6 +279,7 @@ namespace TraceServices
 		if (Task == nullptr)
 		{
 			UE_LOG(LogTraceServices, Log, TEXT("TaskFinished(TaskId %d, Timestamp %.6f) skipped"), TaskId, Timestamp);
+			return;
 		}
 
 		TaskFinished(*Task, Timestamp);
