Fixed -NoPCH -DisableUnity compile issues.

diff --git a/Engine/Source/Developer/TraceServices/Private/AnalysisServicePrivate.h b/Engine/Source/Developer/TraceServices/Private/AnalysisServicePrivate.h
--- a/Engine/Source/Developer/TraceServices/Private/AnalysisServicePrivate.h
+++ b/Engine/Source/Developer/TraceServices/Private/AnalysisServicePrivate.h
@@ -6,6 +6,7 @@
 #include "TraceServices/Model/AnalysisSession.h"
 #include "TraceServices/AnalysisService.h"
 #include "Containers/HashTable.h"
+#include "Common/SlabAllocator.h"
 #include "Common/StringStore.h"
 #include "Trace/Analysis.h"
 
diff --git a/Engine/Source/Developer/TraceServices/Private/Analyzers/CpuProfilerTraceAnalysis.cpp b/Engine/Source/Developer/TraceServices/Private/Analyzers/CpuProfilerTraceAnalysis.cpp
--- a/Engine/Source/Developer/TraceServices/Private/Analyzers/CpuProfilerTraceAnalysis.cpp
+++ b/Engine/Source/Developer/TraceServices/Private/Analyzers/CpuProfilerTraceAnalysis.cpp
@@ -7,7 +7,10 @@
 #include "Common/Utils.h"
 #include "HAL/LowLevelMemTracker.h"
 #include "Model/ThreadsPrivate.h"
+#include "Model/MonotonicTimeline.h"
 #include "Serialization/MemoryWriter.h"
+#include "TraceServices/Model/TimingProfiler.h"
+#include "TraceServices/Utils.h"
 
 #define CPUPROFILER_DEBUG_LOGF(Format, ...) //{ if (ThreadState.ThreadId == 2) FPlatformMisc::LowLevelOutputDebugStringf(Format, __VA_ARGS__); }
 #define CPUPROFILER_DEBUG_BEGIN_EVENT(Time, Event) { ++TotalScopeCount; }
diff --git a/Engine/Source/Developer/TraceServices/Private/Model/ModuleProvider.cpp b/Engine/Source/Developer/TraceServices/Private/Model/ModuleProvider.cpp
--- a/Engine/Source/Developer/TraceServices/Private/Model/ModuleProvider.cpp
+++ b/Engine/Source/Developer/TraceServices/Private/Model/ModuleProvider.cpp
@@ -16,6 +16,7 @@
 #include "Internationalization/Regex.h"
 #include "Misc/Paths.h"
 #include "Misc/PathViews.h"
+#include "Misc/ScopeRWLock.h"
 #include "TraceServices/Model/AnalysisCache.h"
 #include "TraceServices/Model/AnalysisSession.h"
 
diff --git a/Engine/Source/Editor/ComponentVisualizers/Public/AttenuatedComponentVisualizer.h b/Engine/Source/Editor/ComponentVisualizers/Public/AttenuatedComponentVisualizer.h
--- a/Engine/Source/Editor/ComponentVisualizers/Public/AttenuatedComponentVisualizer.h
+++ b/Engine/Source/Editor/ComponentVisualizers/Public/AttenuatedComponentVisualizer.h
@@ -2,6 +2,7 @@
 
 #pragma once
 
+#include "Engine/Attenuation.h"
 #include "ComponentVisualizer.h"
 #include "ShowFlags.h"
 #include "SceneView.h"
diff --git a/Engine/Source/Editor/CurveEditor/Public/CurveEditorSnapMetrics.h b/Engine/Source/Editor/CurveEditor/Public/CurveEditorSnapMetrics.h
--- a/Engine/Source/Editor/CurveEditor/Public/CurveEditorSnapMetrics.h
+++ b/Engine/Source/Editor/CurveEditor/Public/CurveEditorSnapMetrics.h
@@ -2,6 +2,7 @@
 
 #pragma once
 
+#include "Input/Events.h"
 #include "Misc/FrameRate.h"
 #include "Math/Axis.h"
 #include "Math/Vector2D.h"
diff --git a/Engine/Source/Editor/TurnkeySupport/Private/TurnkeyEditorSupport.cpp b/Engine/Source/Editor/TurnkeySupport/Private/TurnkeyEditorSupport.cpp
--- a/Engine/Source/Editor/TurnkeySupport/Private/TurnkeyEditorSupport.cpp
+++ b/Engine/Source/Editor/TurnkeySupport/Private/TurnkeyEditorSupport.cpp
@@ -1,6 +1,7 @@
 // Copyright Epic Games, Inc. All Rights Reserved.
 
 #include "TurnkeyEditorSupport.h"
+#include "Internationalization/Text.h"
 #include "ITurnkeyIOModule.h"
 #include "Misc/AssertionMacros.h"
 
diff --git a/Engine/Source/Runtime/Slate/Public/Framework/Views/TreeFilterHandler.h b/Engine/Source/Runtime/Slate/Public/Framework/Views/TreeFilterHandler.h
--- a/Engine/Source/Runtime/Slate/Public/Framework/Views/TreeFilterHandler.h
+++ b/Engine/Source/Runtime/Slate/Public/Framework/Views/TreeFilterHandler.h
@@ -3,6 +3,7 @@
 #pragma once
 
 #include "CoreMinimal.h"
+#include "Misc/IFilter.h"
 #include "Widgets/Views/STreeView.h"
 
 /** Handles filtering and expanding a TreeView using an IFilter implementation. */
