[Insights] Mem Insights: Fixed allocations analysis to safely ignore errors re invalid MemoryScopePtr scoped events.

diff --git a/Engine/Source/Developer/TraceServices/Private/Model/AllocationsProvider.cpp b/Engine/Source/Developer/TraceServices/Private/Model/AllocationsProvider.cpp
--- a/Engine/Source/Developer/TraceServices/Private/Model/AllocationsProvider.cpp
+++ b/Engine/Source/Developer/TraceServices/Private/Model/AllocationsProvider.cpp
@@ -1673,13 +1673,11 @@ void FAllocationsProvider::EditPushTagFromPtr(uint32 ThreadId, uint8 Tracker, ui
 
 	// Currently only system root heap is affected by reallocs, so limit search.
 	FLiveAllocCollection* Allocs = LiveAllocs[EMemoryTraceRootHeap::SystemMemory];
-	if (ensure(Allocs))
-	{
-		FAllocationItem* Alloc = Allocs->FindRef(Ptr);
-		const TagIdType Tag = Alloc ? Alloc->Tag : 0; // If ptr is not found use "Untagged"
-		TagTracker.PushTagFromPtr(ThreadId, Tracker, Tag);
-	}
-	else
+	FAllocationItem* Alloc = Allocs ? Allocs->FindRef(Ptr) : nullptr;
+	const TagIdType Tag = Alloc ? Alloc->Tag : 0; // If ptr is not found use "Untagged"
+	TagTracker.PushTagFromPtr(ThreadId, Tracker, Tag);
+
+	if (!ensure(Allocs) || !ensure(Alloc))
 	{
 		++MiscErrors;
 		if (MiscErrors <= MaxLogMessagesPerErrorType)
