[Insights] Removed unused StaticStringNoSync and FNamNoSync trace events.

diff --git a/Engine/Source/Developer/TraceServices/Private/Analyzers/StringsAnalyzer.cpp b/Engine/Source/Developer/TraceServices/Private/Analyzers/StringsAnalyzer.cpp
--- a/Engine/Source/Developer/TraceServices/Private/Analyzers/StringsAnalyzer.cpp
+++ b/Engine/Source/Developer/TraceServices/Private/Analyzers/StringsAnalyzer.cpp
@@ -14,8 +14,6 @@ void FStringsAnalyzer::OnAnalysisBegin(const FOnAnalysisContext& Context)
 {
 	Context.InterfaceBuilder.RouteEvent(RouteId_StaticString, "Strings", "StaticString");
 	Context.InterfaceBuilder.RouteEvent(RouteId_FName, "Strings", "FName");
-	Context.InterfaceBuilder.RouteEvent(RouteId_StaticStringNoSync, "Strings", "StaticStringNoSync");
-	Context.InterfaceBuilder.RouteEvent(RouteId_FNameNoSync, "Strings", "FNameNoSync");
 }
 
 ////////////////////////////////////////////////////////////////////////////////////////////////////
@@ -27,8 +25,6 @@ bool FStringsAnalyzer::OnEvent(uint16 RouteId, EStyle Style, const FOnEventConte
 		// The layout of these two events is identical so they can be handled using the same code.
 	case RouteId_StaticString:
 	case RouteId_FName:
-	case RouteId_StaticStringNoSync:
-	case RouteId_FNameNoSync:
 		{
 			const FEventData& EventData = Context.EventData;
 			IDefinitionProvider* DefinitionProvider = GetDefinitionProvider(Session);
diff --git a/Engine/Source/Developer/TraceServices/Private/Analyzers/StringsAnalyzer.h b/Engine/Source/Developer/TraceServices/Private/Analyzers/StringsAnalyzer.h
--- a/Engine/Source/Developer/TraceServices/Private/Analyzers/StringsAnalyzer.h
+++ b/Engine/Source/Developer/TraceServices/Private/Analyzers/StringsAnalyzer.h
@@ -19,8 +19,6 @@ private:
 	{
 		RouteId_StaticString,
 		RouteId_FName,
-		RouteId_StaticStringNoSync,
-		RouteId_FNameNoSync,
 	};
 	virtual void OnAnalysisBegin(const FOnAnalysisContext& Context) override;
 	virtual bool OnEvent(uint16 RouteId, EStyle Style, const FOnEventContext& Context) override;
diff --git a/Engine/Source/Runtime/Core/Private/ProfilingDebugging/StringTrace.cpp b/Engine/Source/Runtime/Core/Private/ProfilingDebugging/StringTrace.cpp
--- a/Engine/Source/Runtime/Core/Private/ProfilingDebugging/StringTrace.cpp
+++ b/Engine/Source/Runtime/Core/Private/ProfilingDebugging/StringTrace.cpp
@@ -6,8 +6,6 @@
 //////////////////////////////////////////////////////////////////////////////////////////
 UE_TRACE_EVENT_DEFINE(Strings, StaticString);
 UE_TRACE_EVENT_DEFINE(Strings, FName);
-UE_TRACE_EVENT_DEFINE(Strings, StaticStringNoSync);
-UE_TRACE_EVENT_DEFINE(Strings, FNameNoSync);
 
 #if UE_TRACE_ENABLED
 
