[Insights] Improved analysis of mem tag events (fixed handling of tags with invalid tag id or with large id value overlapping the PtrScope bit).

diff --git a/Engine/Source/Developer/TraceServices/Private/Model/AllocationsProvider.cpp b/Engine/Source/Developer/TraceServices/Private/Model/AllocationsProvider.cpp
--- a/Engine/Source/Developer/TraceServices/Private/Model/AllocationsProvider.cpp
+++ b/Engine/Source/Developer/TraceServices/Private/Model/AllocationsProvider.cpp
@@ -137,6 +137,16 @@ FTagTracker::FTagTracker(IAnalysisSession& InSession)
 
 void FTagTracker::AddTagSpec(TagIdType InTag, TagIdType InParentTag, const TCHAR* InDisplay)
 {
+	if (InTag == InvalidTagId)
+	{
+		++NumErrors;
+		if (NumErrors <= MaxLogMessagesPerErrorType)
+		{
+			UE_LOG(LogTraceServices, Error, TEXT("[MemAlloc] Cannot add Tag spec ('%s') with invalid id!"), InDisplay);
+		}
+		return;
+	}
+
 	if (ensure(!TagMap.Contains(InTag)))
 	{
 		FStringView Display(InDisplay);
@@ -148,7 +158,7 @@ void FTagTracker::AddTagSpec(TagIdType InTag, TagIdType InParentTag, const TCHAR
 			FullName = Display;
 			// It is possible to define a child tag in runtime using only a string, even if the parent tag does not yet
 			// exist. We need to find the correct parent or store it to the side until the parent tag is announced.
-			if (InParentTag == ~0)
+			if (InParentTag == InvalidTagId)
 			{
 				const FStringView Parent = FPathViews::GetPathLeaf(FPathViews::GetPath(Display));
 				for (auto EntryPair : TagMap)
@@ -162,7 +172,7 @@ void FTagTracker::AddTagSpec(TagIdType InTag, TagIdType InParentTag, const TCHAR
 					}
 				}
 				// If parent tag is still unknown here, create a temporary entry here
-				if (InParentTag == ~0)
+				if (InParentTag == InvalidTagId)
 				{
 					PendingTags.Emplace(InTag, Parent);
 				}
@@ -197,7 +207,7 @@ void FTagTracker::AddTagSpec(TagIdType InTag, TagIdType InParentTag, const TCHAR
 		}
 		else
 		{
-			UE_LOG(LogTraceServices, Verbose, TEXT("[MemAlloc] Added Tag '%s' ('%s')"), Entry.Display, Entry.FullPath);
+			UE_LOG(LogTraceServices, Verbose, TEXT("[MemAlloc] Added Tag '%s' ('%s') with id %u."), Entry.Display, Entry.FullPath, InTag);
 		}
 	}
 	else
@@ -228,7 +238,15 @@ void FTagTracker::PushTag(uint32 InThreadId, uint8 InTracker, TagIdType InTag)
 {
 	const uint32 TrackerThreadId = GetTrackerThreadId(InThreadId, InTracker);
 	FThreadState& State = TrackerThreadStates.FindOrAdd(TrackerThreadId);
-	State.TagStack.Push(InTag);
+	if (InTag == InvalidTagId)
+	{
+		++NumErrors;
+		if (NumErrors <= MaxLogMessagesPerErrorType)
+		{
+			UE_LOG(LogTraceServices, Error, TEXT("[MemAlloc] Invalid Tag on Thread %u (Tracker=%u)!"), InThreadId, uint32(InTracker));
+		}
+	}
+	State.TagStack.Push({ InTag, ETagStackFlags::None });
 }
 
 ////////////////////////////////////////////////////////////////////////////////////////////////////
@@ -239,7 +257,7 @@ void FTagTracker::PopTag(uint32 InThreadId, uint8 InTracker)
 	FThreadState* State = TrackerThreadStates.Find(TrackerThreadId);
 	if (State && !State->TagStack.IsEmpty())
 	{
-		INSIGHTS_SLOW_CHECK((State->TagStack.Top() & 0x80000000) == 0);
+		INSIGHTS_SLOW_CHECK(!State->TagStack.Top().IsPtrScope());
 		State->TagStack.Pop();
 	}
 	else
@@ -247,7 +265,7 @@ void FTagTracker::PopTag(uint32 InThreadId, uint8 InTracker)
 		++NumErrors;
 		if (NumErrors <= MaxLogMessagesPerErrorType)
 		{
-			UE_LOG(LogTraceServices, Error, TEXT("[MemAlloc] Tag stack on Thread %u (Tracker=%u) is already empty!"), InThreadId, InTracker);
+			UE_LOG(LogTraceServices, Error, TEXT("[MemAlloc] Tag stack on Thread %u (Tracker=%u) is already empty!"), InThreadId, uint32(InTracker));
 		}
 	}
 }
@@ -262,7 +280,7 @@ TagIdType FTagTracker::GetCurrentTag(uint32 InThreadId, uint8 InTracker) const
 	{
 		return 0; // Untagged
 	}
-	return State->TagStack.Top() & ~PtrTagMask;
+	return State->TagStack.Top().Tag;
 }
 
 ////////////////////////////////////////////////////////////////////////////////////////////////////
@@ -291,7 +309,7 @@ void FTagTracker::PushTagFromPtr(uint32 InThreadId, uint8 InTracker, TagIdType I
 {
 	const uint32 TrackerThreadId = GetTrackerThreadId(InThreadId, InTracker);
 	FThreadState& State = TrackerThreadStates.FindOrAdd(TrackerThreadId);
-	State.TagStack.Push(InTag | PtrTagMask);
+	State.TagStack.Push({ InTag, ETagStackFlags::PtrScope });
 }
 
 ////////////////////////////////////////////////////////////////////////////////////////////////////
@@ -302,7 +320,7 @@ void FTagTracker::PopTagFromPtr(uint32 InThreadId, uint8 InTracker)
 	FThreadState* State = TrackerThreadStates.Find(TrackerThreadId);
 	if (State && !State->TagStack.IsEmpty())
 	{
-		INSIGHTS_SLOW_CHECK((State->TagStack.Top() & PtrTagMask) != 0);
+		INSIGHTS_SLOW_CHECK(State->TagStack.Top().IsPtrScope());
 		State->TagStack.Pop();
 	}
 	else
@@ -310,7 +328,7 @@ void FTagTracker::PopTagFromPtr(uint32 InThreadId, uint8 InTracker)
 		++NumErrors;
 		if (NumErrors <= MaxLogMessagesPerErrorType)
 		{
-			UE_LOG(LogTraceServices, Error, TEXT("[MemAlloc] Tag stack on Thread %u (Tracker=%u) is already empty!"), InThreadId, InTracker);
+			UE_LOG(LogTraceServices, Error, TEXT("[MemAlloc] Tag stack on Thread %u (Tracker=%u) is already empty!"), InThreadId, uint32(InTracker));
 		}
 	}
 }
@@ -321,7 +339,7 @@ bool FTagTracker::HasTagFromPtrScope(uint32 InThreadId, uint8 InTracker) const
 {
 	const uint32 TrackerThreadId = GetTrackerThreadId(InThreadId, InTracker);
 	const FThreadState* State = TrackerThreadStates.Find(TrackerThreadId);
-	return State && (State->TagStack.Num() > 0) && ((State->TagStack.Top() & PtrTagMask) != 0);
+	return State && (State->TagStack.Num() > 0) && State->TagStack.Top().IsPtrScope();
 }
 
 ////////////////////////////////////////////////////////////////////////////////////////////////////
@@ -1677,7 +1695,7 @@ void FAllocationsProvider::EditPushTagFromPtr(uint32 ThreadId, uint8 Tracker, ui
 	const TagIdType Tag = Alloc ? Alloc->Tag : 0; // If ptr is not found use "Untagged"
 	TagTracker.PushTagFromPtr(ThreadId, Tracker, Tag);
 
-	if (!ensure(Allocs) || !ensure(Alloc))
+	if (!Alloc)
 	{
 		++MiscErrors;
 		if (MiscErrors <= MaxLogMessagesPerErrorType)
diff --git a/Engine/Source/Developer/TraceServices/Private/Model/AllocationsProvider.h b/Engine/Source/Developer/TraceServices/Private/Model/AllocationsProvider.h
--- a/Engine/Source/Developer/TraceServices/Private/Model/AllocationsProvider.h
+++ b/Engine/Source/Developer/TraceServices/Private/Model/AllocationsProvider.h
@@ -41,11 +41,26 @@ class FTagTracker
 private:
 	static constexpr uint32 TrackerIdShift = 24;
 	static constexpr uint32 TrackerIdMask = 0xFF000000;
-	static constexpr uint32 PtrTagMask = 0x80000000;
+	static constexpr TagIdType UntaggedTagId = 0;
+	static constexpr TagIdType InvalidTagId = ~0;
+
+	enum class ETagStackFlags : uint32
+	{
+		None = 0,
+		PtrScope = 1 << 0,
+	};
+
+	struct FTagStackEntry
+	{
+		TagIdType Tag;
+		ETagStackFlags Flags;
+
+		inline bool IsPtrScope() const { return ((uint32(Flags) & uint32(ETagStackFlags::PtrScope)) != 0); }
+	};
 
 	struct FThreadState
 	{
-		TArray<TagIdType> TagStack;
+		TArray<FTagStackEntry> TagStack;
 	};
 
 	struct FTagEntry
