[Insights] Improved analysis of scoped memory events by avoiding a "strlen" for each scoped event.

diff --git a/Engine/Source/Developer/TraceServices/Private/Analyzers/AllocationsAnalysis.cpp b/Engine/Source/Developer/TraceServices/Private/Analyzers/AllocationsAnalysis.cpp
--- a/Engine/Source/Developer/TraceServices/Private/Analyzers/AllocationsAnalysis.cpp
+++ b/Engine/Source/Developer/TraceServices/Private/Analyzers/AllocationsAnalysis.cpp
@@ -56,6 +56,9 @@ void FAllocationsAnalyzer::OnAnalysisBegin(const FOnAnalysisContext& Context)
 	Builder.RouteEvent(RouteId_HeapMarkAlloc,       "Memory", "HeapMarkAlloc");
 	Builder.RouteEvent(RouteId_HeapUnmarkAlloc,     "Memory", "HeapUnmarkAlloc");
 
+	Builder.RouteEvent(RouteId_NewEvent, "Memory", "MemoryScope", false);
+	Builder.RouteEvent(RouteId_NewEvent, "Memory", "MemoryScopePtr", false);
+
 	Builder.RouteLoggerEvents(RouteId_MemScope, "Memory", true);
 
 #if INSIGHTS_MEM_TRACE_METADATA_TEST
@@ -87,6 +90,29 @@ void FAllocationsAnalyzer::OnAnalysisEnd()
 
 ////////////////////////////////////////////////////////////////////////////////////////////////////
 
+bool FAllocationsAnalyzer::OnNewEvent(uint16 RouteId, const FEventTypeInfo& TypeInfo)
+{
+	const ANSICHAR* LoggerName = TypeInfo.GetLoggerName();
+	const ANSICHAR* EventName = TypeInfo.GetName();
+	const uint32 EventId = TypeInfo.GetId();
+
+	if (FCStringAnsi::Strcmp(LoggerName, "Memory") == 0)
+	{
+		if (FCStringAnsi::Strcmp(EventName, "MemoryScope") == 0)
+		{
+			MemoryScopeEventId = EventId;
+		}
+		else if (FCStringAnsi::Strcmp(EventName, "MemoryScopePtr") == 0)
+		{
+			MemoryScopePtrEventId = EventId;
+		}
+	}
+
+	return true;
+}
+
+////////////////////////////////////////////////////////////////////////////////////////////////////
+
 bool FAllocationsAnalyzer::OnEvent(uint16 RouteId, EStyle Style, const FOnEventContext& Context)
 {
 	LLM_SCOPE_BYNAME(TEXT("Insights/FAllocationsAnalyzer"));
@@ -328,9 +354,11 @@ bool FAllocationsAnalyzer::OnEvent(uint16 RouteId, EStyle Style, const FOnEventC
 			const uint32 ThreadId = Context.ThreadInfo.GetSystemId();
 			const uint8 Tracker = 0; // We only care about the default tracker for now.
 
+			const uint32 EventId = Context.EventData.GetTypeInfo().GetId();
+
 			if (Style == EStyle::EnterScope)
 			{
-				if (FCStringAnsi::Strlen(Context.EventData.GetTypeInfo().GetName()) == 11) // "MemoryScope"
+				if (EventId == MemoryScopeEventId) // "MemoryScope"
 				{
 					const TagIdType Tag = Context.EventData.GetValue<TagIdType>("Tag");
 					{
@@ -344,7 +372,7 @@ bool FAllocationsAnalyzer::OnEvent(uint16 RouteId, EStyle Style, const FOnEventC
 					}
 #endif
 				}
-				else // "MemoryScopePtr"
+				else if (ensure(EventId == MemoryScopePtrEventId)) // "MemoryScopePtr"
 				{
 					const uint64 Ptr = Context.EventData.GetValue<uint64>("Ptr");
 					{
@@ -362,17 +390,16 @@ bool FAllocationsAnalyzer::OnEvent(uint16 RouteId, EStyle Style, const FOnEventC
 			}
 			else // EStyle::LeaveScope
 			{
+				if (EventId == MemoryScopeEventId) // "MemoryScope"
 				{
 					FProviderEditScopeLock _(AllocationsProvider);
-					if (AllocationsProvider.HasTagFromPtrScope(ThreadId, Tracker)) // Is TagFromPtr scope active?
-					{
-						AllocationsProvider.EditPopTagFromPtr(ThreadId, Tracker);
-					}
-					else
-					{
-						check(FCStringAnsi::Strlen(Context.EventData.GetTypeInfo().GetName()) == 11); // "MemoryScope"
-						AllocationsProvider.EditPopTag(ThreadId, Tracker);
-					}
+					AllocationsProvider.EditPopTag(ThreadId, Tracker);
+				}
+				else if (ensure(EventId == MemoryScopePtrEventId)) // "MemoryScopePtr"
+				{
+					FProviderEditScopeLock _(AllocationsProvider);
+					//check(AllocationsProvider.HasTagFromPtrScope(ThreadId, Tracker));
+					AllocationsProvider.EditPopTagFromPtr(ThreadId, Tracker);
 				}
 #if INSIGHTS_MEM_TRACE_METADATA_TEST
 				{
diff --git a/Engine/Source/Developer/TraceServices/Private/Analyzers/AllocationsAnalysis.h b/Engine/Source/Developer/TraceServices/Private/Analyzers/AllocationsAnalysis.h
--- a/Engine/Source/Developer/TraceServices/Private/Analyzers/AllocationsAnalysis.h
+++ b/Engine/Source/Developer/TraceServices/Private/Analyzers/AllocationsAnalysis.h
@@ -17,6 +17,7 @@ class FAllocationsAnalyzer : public UE::Trace::IAnalyzer
 private:
 	enum : uint16
 	{
+		RouteId_NewEvent,
 		RouteId_Init,
 		RouteId_Alloc,
 		RouteId_AllocSystem,
@@ -30,10 +31,10 @@ private:
 		RouteId_ReallocFreeSystem,
 		RouteId_Marker,
 		RouteId_TagSpec,
-		RouteId_MemScope,
 		RouteId_HeapSpec,
 		RouteId_HeapMarkAlloc,
 		RouteId_HeapUnmarkAlloc,
+		RouteId_MemScope,
 	};
 
 public:
@@ -41,6 +42,7 @@ public:
 	~FAllocationsAnalyzer();
 	virtual void OnAnalysisBegin(const FOnAnalysisContext& Context) override;
 	virtual void OnAnalysisEnd() override;
+	virtual bool OnNewEvent(uint16 RouteId, const FEventTypeInfo& TypeInfo) override;
 	virtual bool OnEvent(uint16 RouteId, EStyle Style, const FOnEventContext& Context) override;
 	double GetCurrentTime() const;
 
@@ -54,6 +56,8 @@ private:
 	uint32 MarkerPeriod = 0;
 	uint16 TagIdMetadataType = ~0;
 	uint8 SizeShift = 0;
+	uint32 MemoryScopeEventId = 0;
+	uint32 MemoryScopePtrEventId = 0;
 };
 
 } // namespace TraceServices
