[Insights] Fixed crash in analysis of timing events (analyzing traces with "pending" events (ones with metadata) emitted just before EndThread trace event).

diff --git a/Engine/Source/Developer/TraceServices/Private/Analyzers/CpuProfilerTraceAnalysis.cpp b/Engine/Source/Developer/TraceServices/Private/Analyzers/CpuProfilerTraceAnalysis.cpp
--- a/Engine/Source/Developer/TraceServices/Private/Analyzers/CpuProfilerTraceAnalysis.cpp
+++ b/Engine/Source/Developer/TraceServices/Private/Analyzers/CpuProfilerTraceAnalysis.cpp
@@ -593,7 +593,7 @@ void FCpuProfilerAnalyzer::DispatchPendingEvents(
 		}
 
 #if CPUPROFILER_SAFE_DISPATCH_PENDING_EVENTS
-		if (ensure(PendingCycle >= LastCycle))
+		if (PendingCycle >= LastCycle)
 		{
 			// Update LastCycle in order to verify time (of following pending events) increases monotonically.
 			LastCycle = PendingCycle;
@@ -628,6 +628,8 @@ void FCpuProfilerAnalyzer::DispatchPendingEvents(
 			CPUPROFILER_DEBUG_END_EVENT(PendingTime);
 		}
 	}
+
+	ThreadState.LastCycle = LastCycle;
 }
 
 ////////////////////////////////////////////////////////////////////////////////////////////////////
