[Insights] Fix crash when importing very large files as tables.

diff --git a/Engine/Source/Developer/TraceServices/Private/Model/TableImportTask.cpp b/Engine/Source/Developer/TraceServices/Private/Model/TableImportTask.cpp
--- a/Engine/Source/Developer/TraceServices/Private/Model/TableImportTask.cpp
+++ b/Engine/Source/Developer/TraceServices/Private/Model/TableImportTask.cpp
@@ -23,7 +23,7 @@ FTableImportTask::~FTableImportTask()
 void FTableImportTask::operator()()
 {
 	TArray<FString> Lines;
-	FFileHelper::LoadFileToStringArray(Lines, *FilePath);
+	LoadFileToStringArray(FilePath, Lines);
 
 	if (FilePath.EndsWith(TEXT(".csv")))
 	{
@@ -190,6 +190,18 @@ void FTableImportTask::SplitLineIntoValues(const FString& InLine, TArray<FString
 	OutValues.Add(Value.TrimQuotes());
 }
 
+bool FTableImportTask::LoadFileToStringArray(const FString &InFilePath, TArray<FString>& Lines)
+{
+	// We use this function instead of FFileHelper::LoadFileToArray because that one does not support very large files.
+
+	auto Visitor = [&Lines](FStringView Line)
+	{
+		Lines.Add(FString(Line));
+	};
+
+	return FFileHelper::LoadFileToStringWithLineVisitor(*FilePath, Visitor);
+}
+
 void FTableImportService::ImportTable(const FString& InPath, FName TableId, FTableImportService::TableImportCallback InCallback)
 {
 	UE::Tasks::Launch(UE_SOURCE_LOCATION, FTableImportTask(InPath, TableId, InCallback));
diff --git a/Engine/Source/Developer/TraceServices/Private/Model/TableImportTask.h b/Engine/Source/Developer/TraceServices/Private/Model/TableImportTask.h
--- a/Engine/Source/Developer/TraceServices/Private/Model/TableImportTask.h
+++ b/Engine/Source/Developer/TraceServices/Private/Model/TableImportTask.h
@@ -60,6 +60,7 @@ public:
 
 private:
 	void SplitLineIntoValues(const FString& InLine, TArray<FString>& OutValues);
+	bool LoadFileToStringArray(const FString& InFilePath, TArray<FString>& Lines);
 
 	FTableImportService::TableImportCallback Callback;
 	TSharedPtr<TImportTable<FImportTableRow>> Table;
