[Insights] Added "FVirtualWinApiHooks::" and ".*/ConcurrentLinearAllocator.*" to symbol filters.

diff --git a/Engine/Source/Developer/TraceServices/Private/Model/ModuleProvider.cpp b/Engine/Source/Developer/TraceServices/Private/Model/ModuleProvider.cpp
--- a/Engine/Source/Developer/TraceServices/Private/Model/ModuleProvider.cpp
+++ b/Engine/Source/Developer/TraceServices/Private/Model/ModuleProvider.cpp
@@ -415,6 +415,7 @@ FResolvedSymbolFilter::FResolvedSymbolFilter()
 	IgnoreSymbolsByFunctionName.Add(TEXT("FMemory::"));
 	IgnoreSymbolsByFunctionName.Add(TEXT("FMallocWrapper::"));
 	IgnoreSymbolsByFunctionName.Add(TEXT("FMallocPoisonProxy::"));
+	IgnoreSymbolsByFunctionName.Add(TEXT("FVirtualWinApiHooks::"));
 	IgnoreSymbolsByFunctionName.Add(TEXT("Malloc"));
 	IgnoreSymbolsByFunctionName.Add(TEXT("Realloc"));
 	IgnoreSymbolsByFunctionName.Add(TEXT("MemoryTrace_"));
@@ -422,6 +423,7 @@ FResolvedSymbolFilter::FResolvedSymbolFilter()
 	IgnoreSymbolsByFunctionName.Add(TEXT("std::"));
 
 	IgnoreSymbolsByFilePath.Add(FRegexPattern(FString(TEXT(".*/Containers/.*"))));
+	IgnoreSymbolsByFilePath.Add(FRegexPattern(FString(TEXT(".*/ConcurrentLinearAllocator.*"))));
 }
 
 /////////////////////////////////////////////////////////////////////
