[Insights] Fixed ensure failing in ContextSwitches analyzer for server traces with -nothreading. When -nothreading is used, multiple fake threads (FFakeThread) are created from the game thread. These will have same trace ThreadId == 2 and different SystemThreadId.

diff --git a/Engine/Source/Developer/TraceServices/Private/Model/ContextSwitches.cpp b/Engine/Source/Developer/TraceServices/Private/Model/ContextSwitches.cpp
--- a/Engine/Source/Developer/TraceServices/Private/Model/ContextSwitches.cpp
+++ b/Engine/Source/Developer/TraceServices/Private/Model/ContextSwitches.cpp
@@ -373,6 +373,13 @@ void FContextSwitchesProvider::AddThreadInfo(uint32 ThreadId, uint32 SystemThrea
 	if (OldSystemThreadIdPtr)
 	{
 		uint32 OldSystemThreadId = *OldSystemThreadIdPtr;
+		if (ThreadId == 2 /*GameThread*/ && OldSystemThreadId != SystemThreadId)
+		{
+			// When -nothreading is used, multiple fake threads (FFakeThread) are created from the game thread.
+			// These will have same trace ThreadId == 2 and different SystemThreadId.
+			SystemToTraceThreadIdMap.Add(SystemThreadId, ThreadId);
+			return;
+		}
 		ensure(OldSystemThreadId == SystemThreadId);
 	}
 
