[Insights] Add error reporting to the table import feature.

diff --git a/Engine/Source/Developer/TraceInsights/Private/Insights/ImportTool/TableImportTool.cpp b/Engine/Source/Developer/TraceInsights/Private/Insights/ImportTool/TableImportTool.cpp
--- a/Engine/Source/Developer/TraceInsights/Private/Insights/ImportTool/TableImportTool.cpp
+++ b/Engine/Source/Developer/TraceInsights/Private/Insights/ImportTool/TableImportTool.cpp
@@ -100,9 +100,11 @@ void FTableImportTool::StartImportProcess()
 	FName TableId = GetTableID(Filenames[0]);
 	DisplayTable(TableId);
 
-	FTableImportService::ImportTable(Filenames[0], TableId, [this](FName TableViewID, TSharedPtr<TraceServices::ITable<TraceServices::FImportTableRow>> Table) {
-		this->TableImportServiceCallback(TableViewID, Table);
-	});
+	FTableImportService::ImportTable(Filenames[0],
+		TableId,
+		[this](TSharedPtr<FTableImportCallbackParams> Params) {
+			this->TableImportServiceCallback(Params);
+		});
 }
 
 ////////////////////////////////////////////////////////////////////////////////////////////////////
@@ -133,14 +135,26 @@ void FTableImportTool::DisplayTable(FName TableViewID)
 
 ////////////////////////////////////////////////////////////////////////////////////////////////////
 
-void FTableImportTool::TableImportServiceCallback(FName TableViewID, TSharedPtr<TraceServices::ITable<TraceServices::FImportTableRow>> Table)
+void FTableImportTool::TableImportServiceCallback(TSharedPtr<TraceServices::FTableImportCallbackParams> Params)
 {
-	const FOpenImportedTableTabData* TableViewPtr = OpenTablesMap.Find(TableViewID);
+	ensure(Params.IsValid());
+
+	const FOpenImportedTableTabData* TableViewPtr = OpenTablesMap.Find(Params->TableId);
 	if (TableViewPtr)
 	{
 		TSharedPtr<SUntypedTableTreeView> TableView = TableViewPtr->TableTreeView;
 		TableView->ClearCurrentOperationNameOverride();
-		TableView->UpdateSourceTable(Table);
+
+		if (Params->Result == TraceServices::ETableImportResult::ESuccess)
+		{
+			TableView->UpdateSourceTable(Params->Table);
+		}
+		else
+		{
+			FMessageLog ReportMessageLog(FInsightsManager::Get()->GetLogListingName());
+			ReportMessageLog.AddMessages(Params->Messages);
+			ReportMessageLog.Notify();
+		}
 	}
 }
 
diff --git a/Engine/Source/Developer/TraceInsights/Private/Insights/ImportTool/TableImportTool.h b/Engine/Source/Developer/TraceInsights/Private/Insights/ImportTool/TableImportTool.h
--- a/Engine/Source/Developer/TraceInsights/Private/Insights/ImportTool/TableImportTool.h
+++ b/Engine/Source/Developer/TraceInsights/Private/Insights/ImportTool/TableImportTool.h
@@ -12,6 +12,11 @@
 class SDockTab;
 class FSpawnTabArgs;
 
+namespace TraceServices
+{
+	struct FTableImportCallbackParams;
+}
+
 namespace Insights
 {
 
@@ -56,7 +61,7 @@ public:
 	void StartImportProcess();
 
 	void DisplayTable(FName TableViewID);
-	void TableImportServiceCallback(FName TableViewID, TSharedPtr<TraceServices::ITable<TraceServices::FImportTableRow>> Table);
+	void TableImportServiceCallback(TSharedPtr<TraceServices::FTableImportCallbackParams> Params);
 
 	TSharedRef<SDockTab> SpawnTab_TableImportTreeView(const FSpawnTabArgs& Args, FName TableViewID, FText InDisplayName);
 	void OnTableImportTreeViewTabClosed(TSharedRef<SDockTab> TabBeingClosed);
diff --git a/Engine/Source/Developer/TraceServices/Private/Model/TableImportTask.cpp b/Engine/Source/Developer/TraceServices/Private/Model/TableImportTask.cpp
--- a/Engine/Source/Developer/TraceServices/Private/Model/TableImportTask.cpp
+++ b/Engine/Source/Developer/TraceServices/Private/Model/TableImportTask.cpp
@@ -3,12 +3,15 @@
 #include "TableImportTask.h"
 
 #include "Async/TaskGraphInterfaces.h"
+#include "Logging/TokenizedMessage.h"
 #include "Misc/FileHelper.h"
 #include "Tasks/Task.h"
 
 namespace TraceServices
 {
 
+#define LOCTEXT_NAMESPACE "TableImportTask"
+
 FTableImportTask::FTableImportTask(const FString& InFilePath, FName InTableId, FTableImportService::TableImportCallback InCallback)
 	: Callback(InCallback)
 	, FilePath(InFilePath)
@@ -22,8 +25,42 @@ FTableImportTask::~FTableImportTask()
 
 void FTableImportTask::operator()()
 {
+	ETableImportResult Result = ImportTable();
+
+	TSharedPtr<FTableImportCallbackParams> Params = MakeShared<FTableImportCallbackParams>();
+	Params->TableId = TableId;
+	Params->Result = Result;
+	Params->Table = Table;
+	Params->Messages = Messages;
+
+	FFunctionGraphTask::CreateAndDispatchWhenReady(
+		[Callback = this->Callback, Params]()
+		{
+			Callback(Params);
+		},
+		TStatId(),
+		nullptr,
+		ENamedThreads::GameThread);
+}
+
+ETableImportResult FTableImportTask::ImportTable()
+{
+	Table = MakeShared<TImportTable<FImportTableRow>>();
+
 	TArray<FString> Lines;
-	LoadFileToStringArray(FilePath, Lines);
+	bool Result = LoadFileToStringArray(FilePath, Lines);
+
+	if (!Result)
+	{
+		AddError(FText::Format(LOCTEXT("FailedToReadMsg", "Import failed because file {0} could not be read:"), FText::FromString(FilePath)));
+		return ETableImportResult::EFail;
+	}
+
+	if (Lines.Num() < 2)
+	{
+		AddError(FText::Format(LOCTEXT("NotEnoughtLinesMsg", "Import failed because the files did not contain a minimum of 2 lines."), FText::FromString(FilePath)));
+		return ETableImportResult::EFail;
+	}
 
 	if (FilePath.EndsWith(TEXT(".csv")))
 	{
@@ -34,38 +71,34 @@ void FTableImportTask::operator()()
 		Separator = TEXT("\t");
 	}
 
-	Table = MakeShared<TImportTable<FImportTableRow>>();
-
 	if (!ParseHeader(Lines[0]))
 	{
-		return;
+		return ETableImportResult::EFail;
 	}
 
 	if (!CreateLayout(Lines[1]))
 	{
-		return;
+		return ETableImportResult::EFail;
 	}
 
 	if (!ParseData(Lines))
 	{
-		return;
+		return ETableImportResult::EFail;
 	}
 
-	FFunctionGraphTask::CreateAndDispatchWhenReady(
-		[Callback = this->Callback, Table = this->Table, TableId = this->TableId]()
-		{
-			Callback(TableId, Table);
-		},
-		TStatId(),
-		nullptr,
-		ENamedThreads::GameThread);
+	return ETableImportResult::ESuccess;
 }
 
 bool FTableImportTask::ParseHeader(const FString& HeaderLine)
 {
 	HeaderLine.ParseIntoArray(ColumnNames, *Separator);
+	if (ColumnNames.Num() == 0)
+	{
+		AddError(FText::Format(LOCTEXT("NoColumnsMsg", "Import failed because the file did not contain any columns."), FText::FromString(FilePath)));
+		return false;
+	}
 
-	return ColumnNames.Num() > 0;
+	return true;
 }
 
 bool FTableImportTask::CreateLayout(const FString& Line)
@@ -73,6 +106,12 @@ bool FTableImportTask::CreateLayout(const FString& Line)
 	TArray<FString> Values;
 	SplitLineIntoValues(Line, Values);
 
+	if (Values.Num() != ColumnNames.Num())
+	{
+		AddError(LOCTEXT("ValuesMismatchMsg", "Import failed because the number of values on line 1 does not match the number of columns in the header line."));
+		return false;
+	}
+
 	TTableLayout<FImportTableRow>& Layout = Table->EditLayout();
 	for (int32 Index = 0; Index < Values.Num(); ++Index)
 	{
@@ -111,6 +150,12 @@ bool FTableImportTask::ParseData(TArray<FString>& Lines)
 		TArray<FString> Values;
 		SplitLineIntoValues(Lines[LineIndex], Values);
 
+		if (Values.Num() != ColumnNames.Num())
+		{
+			AddError(FText::Format(LOCTEXT("ValuesMismatchMsg", "Import failed because the number of values on line {0} does not match the number of columns in the header line."), LineIndex + 1));
+			return false;
+		}
+
 		FImportTableRow& NewRow = Table->AddRow();
 		NewRow.SetNumValues(Values.Num());
 		for (int32 ValueIndex = 0; ValueIndex < Values.Num(); ++ValueIndex)
@@ -202,6 +247,11 @@ bool FTableImportTask::LoadFileToStringArray(const FString &InFilePath, TArray<F
 	return FFileHelper::LoadFileToStringWithLineVisitor(*FilePath, Visitor);
 }
 
+void FTableImportTask::AddError(const FText& Msg)
+{
+	Messages.Add(FTokenizedMessage::Create(EMessageSeverity::Error, Msg));
+}
+
 void FTableImportService::ImportTable(const FString& InPath, FName TableId, FTableImportService::TableImportCallback InCallback)
 {
 	UE::Tasks::Launch(UE_SOURCE_LOCATION, FTableImportTask(InPath, TableId, InCallback));
diff --git a/Engine/Source/Developer/TraceServices/Private/Model/TableImportTask.h b/Engine/Source/Developer/TraceServices/Private/Model/TableImportTask.h
--- a/Engine/Source/Developer/TraceServices/Private/Model/TableImportTask.h
+++ b/Engine/Source/Developer/TraceServices/Private/Model/TableImportTask.h
@@ -8,6 +8,8 @@
 #include "Model/Tables.h"
 #include "TraceServices/Model/TableImport.h"
 
+class FTokenizedMessage;
+
 namespace TraceServices
 {
 
@@ -59,12 +61,17 @@ public:
 	bool ParseData(TArray<FString>& Lines);
 
 private:
+	ETableImportResult ImportTable();
+
 	void SplitLineIntoValues(const FString& InLine, TArray<FString>& OutValues);
 	bool LoadFileToStringArray(const FString& InFilePath, TArray<FString>& Lines);
 
+	void AddError(const FText& Msg);
+
 	FTableImportService::TableImportCallback Callback;
 	TSharedPtr<TImportTable<FImportTableRow>> Table;
 
+	TArray<TSharedRef<FTokenizedMessage>> Messages;
 	TArray<FString> ColumnNames;
 	FString FilePath;
 	FName TableId;
diff --git a/Engine/Source/Developer/TraceServices/Public/TraceServices/Model/TableImport.h b/Engine/Source/Developer/TraceServices/Public/TraceServices/Model/TableImport.h
--- a/Engine/Source/Developer/TraceServices/Public/TraceServices/Model/TableImport.h
+++ b/Engine/Source/Developer/TraceServices/Public/TraceServices/Model/TableImport.h
@@ -6,6 +6,8 @@
 #include "CoreMinimal.h"
 #include "TraceServices/Containers/Tables.h"
 
+class FTokenizedMessage;
+
 namespace TraceServices
 {
 class FImportTableRow;
@@ -15,10 +17,24 @@ class ITableImportData
 	TSharedPtr<ITable<FImportTableRow>> GetTable();
 };
 
+enum class ETableImportResult : uint32
+{
+	ESuccess = 0,
+	EFail = 1,
+};
+
+struct FTableImportCallbackParams
+{
+	FName TableId;
+	ETableImportResult Result;
+	TSharedPtr<ITable<FImportTableRow>> Table;
+	TArray<TSharedRef<FTokenizedMessage>> Messages;
+};
+
 class FTableImportService
 {
 public:
-	typedef TFunction<void(FName TableId, TSharedPtr<ITable<FImportTableRow>> Data)> TableImportCallback;
+	typedef TFunction<void(TSharedPtr<FTableImportCallbackParams>)> TableImportCallback;
 
 	TRACESERVICES_API static void ImportTable(const FString& InPath, FName TableId, TableImportCallback InCallback);
 };
