Add missing undef.

diff --git a/Engine/Source/Developer/TraceServices/Private/Model/TableImportTask.cpp b/Engine/Source/Developer/TraceServices/Private/Model/TableImportTask.cpp
--- a/Engine/Source/Developer/TraceServices/Private/Model/TableImportTask.cpp
+++ b/Engine/Source/Developer/TraceServices/Private/Model/TableImportTask.cpp
@@ -257,4 +257,6 @@ void FTableImportService::ImportTable(const FString& InPath, FName TableId, FTab
 	UE::Tasks::Launch(UE_SOURCE_LOCATION, FTableImportTask(InPath, TableId, InCallback));
 }
 
+#undef LOCTEXT_NAMESPACE
+
 } // namespace TraceServices
