[Insights] Fixed unsafe type cast warnings in TraceServices module.

diff --git a/Engine/Source/Developer/TraceServices/Private/AnalysisCache.cpp b/Engine/Source/Developer/TraceServices/Private/AnalysisCache.cpp
--- a/Engine/Source/Developer/TraceServices/Private/AnalysisCache.cpp
+++ b/Engine/Source/Developer/TraceServices/Private/AnalysisCache.cpp
@@ -189,7 +189,7 @@ bool FAnalysisCache::FFileContents::Load()
 	FUniqueBuffer Buffer = FUniqueBuffer::Alloc(ReservedSize);
 	File->Read((uint8*)Buffer.GetData(), Buffer.GetSize());
 	
-	FMemoryReaderView Ar(MakeArrayView<uint8>((uint8*)Buffer.GetData(), Buffer.GetSize()));
+	FMemoryReaderView Ar(MakeArrayView<uint8>((uint8*)Buffer.GetData(), IntCastChecked<int32>(Buffer.GetSize())));
 	
 	FCbPackage Package;
 	if (!Package.TryLoad(Ar))
@@ -206,7 +206,7 @@ bool FAnalysisCache::FFileContents::Load()
 	}
 
 	FCbArrayView IndexArray = Package.GetObject().Find(ANSITEXTVIEW("Index")).AsArrayView();
-	IndexEntries.Reserve(IndexArray.Num());
+	IndexEntries.Reserve(static_cast<int32>(IndexArray.Num()));
 	for (FCbFieldView IndexEntry : IndexArray)
 	{
 		FCbObjectView IndexEntryObj = IndexEntry.AsObjectView();
@@ -224,7 +224,7 @@ bool FAnalysisCache::FFileContents::Load()
 	}
 
 	FCbArrayView BlockArray = Package.GetObject().Find(ANSITEXTVIEW("Blocks")).AsArrayView();
-	Blocks.Reserve(BlockArray.Num());
+	Blocks.Reserve(static_cast<int32>(BlockArray.Num()));
 	
 	for (FCbFieldView BlockEntryView : BlockArray)
 	{
@@ -245,7 +245,7 @@ uint64 FAnalysisCache::FFileContents::UpdateBlock(FMemoryView Block, BlockKeyTyp
 		return 0;
 	}
 
-	const uint64 EntryIndex = Algo::BinarySearchBy(Blocks, BlockKey, [&](const FBlockEntry& InEntry){return InEntry.BlockKey;});
+	const int32 EntryIndex = Algo::BinarySearchBy(Blocks, BlockKey, [&](const FBlockEntry& InEntry) { return InEntry.BlockKey; });
 	const FIoHash CurrentHash = FIoHash::HashBuffer(Block);
 	
 	if (EntryIndex != INDEX_NONE)
@@ -298,7 +298,7 @@ uint64 FAnalysisCache::FFileContents::LoadBlock(FMutableMemoryView Block, BlockK
 		return 0;
 	}
 	
-	const uint64 EntryIndex = Algo::BinarySearchBy(Blocks, BlockKey, [&](const FBlockEntry& InEntry){return InEntry.BlockKey;});
+	const int32 EntryIndex = Algo::BinarySearchBy(Blocks, BlockKey, [&](const FBlockEntry& InEntry) { return InEntry.BlockKey; });
 	if (EntryIndex == INDEX_NONE)
 	{
 		UE_LOG(LogAnalysisCache, Error, TEXT("Trying to load unknown block 0x%x."), BlockKey);
@@ -511,7 +511,7 @@ FSharedBuffer FAnalysisCache::GetBlocks(FCacheId CacheId, uint32 BlockIndexStart
 /////////////////////////////////////////////////////////////////////
 void FAnalysisCache::ReleaseBlocks(uint8* BlockBuffer, FCacheId CacheId, uint32 BlockIndexStart, uint64 Size)
 {
-	const uint32 BlockCount = Size / IAnalysisCache::BlockSizeBytes;
+	const uint32 BlockCount = IntCastChecked<uint32>(Size / IAnalysisCache::BlockSizeBytes);
 	for (uint32 Block = 0; Block < BlockCount; ++Block)
 	{
 		const void* BlockStart = BlockBuffer + (Block * IAnalysisCache::BlockSizeBytes);
diff --git a/Engine/Source/Developer/TraceServices/Private/AnalysisService.cpp b/Engine/Source/Developer/TraceServices/Private/AnalysisService.cpp
--- a/Engine/Source/Developer/TraceServices/Private/AnalysisService.cpp
+++ b/Engine/Source/Developer/TraceServices/Private/AnalysisService.cpp
@@ -192,14 +192,20 @@ TSharedPtr<const IAnalysisSession> FAnalysisService::StartAnalysis(const TCHAR*
 			{
 				return 0;
 			}
-
-			Size = (Size < Remaining) ? Size : Remaining;
+			if (Size > Remaining)
+			{
+				Size = static_cast<uint32>(Remaining);
+			}
 			Remaining -= Size;
-			return Handle->Read((uint8*)Data, Size) ? Size : 0;
+			if (!Handle->Read((uint8*)Data, Size))
+			{
+				return 0;
+			}
+			return Size;
 		}
 
 		TUniquePtr<IFileHandle> Handle;
-		int64 Remaining;
+		uint64 Remaining;
 	};
 
 	IPlatformFile& FileSystem = IPlatformFile::GetPlatformPhysical();
diff --git a/Engine/Source/Developer/TraceServices/Private/Analyzers/AllocationsAnalysis.cpp b/Engine/Source/Developer/TraceServices/Private/Analyzers/AllocationsAnalysis.cpp
--- a/Engine/Source/Developer/TraceServices/Private/Analyzers/AllocationsAnalysis.cpp
+++ b/Engine/Source/Developer/TraceServices/Private/Analyzers/AllocationsAnalysis.cpp
@@ -172,7 +172,7 @@ bool FAllocationsAnalyzer::OnEvent(uint16 RouteId, EStyle Style, const FOnEventC
 
 			uint64 Address = EventData.GetValue<uint64>("Address");
 
-			RootHeap = EventData.GetValue<uint8>("RootHeap", RootHeap);
+			RootHeap = EventData.GetValue<uint8>("RootHeap", static_cast<uint8>(RootHeap));
 
 			uint64 SizeUpper = EventData.GetValue<uint32>("Size");
 			const uint8 SizeLowerMask = ((1 << SizeShift) - 1);
@@ -236,7 +236,7 @@ bool FAllocationsAnalyzer::OnEvent(uint16 RouteId, EStyle Style, const FOnEventC
 			}
 #endif // INSIGHTS_MEM_TRACE_LEGACY_FORMAT
 
-			RootHeap = EventData.GetValue<uint8>("RootHeap", RootHeap);
+			RootHeap = EventData.GetValue<uint8>("RootHeap", static_cast<uint8>(RootHeap));
 
 			const uint32 TraceThreadId = Context.ThreadInfo.GetId();
 			const uint32 SystemThreadId = Context.ThreadInfo.GetSystemId();
diff --git a/Engine/Source/Developer/TraceServices/Private/Analyzers/CpuProfilerTraceAnalysis.cpp b/Engine/Source/Developer/TraceServices/Private/Analyzers/CpuProfilerTraceAnalysis.cpp
--- a/Engine/Source/Developer/TraceServices/Private/Analyzers/CpuProfilerTraceAnalysis.cpp
+++ b/Engine/Source/Developer/TraceServices/Private/Analyzers/CpuProfilerTraceAnalysis.cpp
@@ -302,7 +302,7 @@ uint64 FCpuProfilerAnalyzer::ProcessBuffer(const FEventTime& EventTime, FThreadS
 
 		if (DecodedCycle & 1ull)
 		{
-			uint32 SpecId = FTraceAnalyzerUtils::Decode7bit(BufferPtr);
+			uint32 SpecId = IntCastChecked<uint32>(FTraceAnalyzerUtils::Decode7bit(BufferPtr));
 			uint32 TimerId = GetTimerId(SpecId);
 
 			FEventScopeState& ScopeState = ThreadState.ScopeStack.AddDefaulted_GetRef();
@@ -405,7 +405,7 @@ uint64 FCpuProfilerAnalyzer::ProcessBufferV2(const FEventTime& EventTime, FThrea
 			if (DecodedCycle & 1ull)
 			{
 				uint64 CoroutineId = FTraceAnalyzerUtils::Decode7bit(BufferPtr);
-				uint32 TimerScopeDepth = FTraceAnalyzerUtils::Decode7bit(BufferPtr);
+				uint32 TimerScopeDepth = IntCastChecked<uint32>(FTraceAnalyzerUtils::Decode7bit(BufferPtr));
 
 				// Begins a "CoroTask" scoped timer.
 				{
@@ -462,7 +462,7 @@ uint64 FCpuProfilerAnalyzer::ProcessBufferV2(const FEventTime& EventTime, FThrea
 			}
 			else
 			{
-				uint32 TimerScopeDepth = FTraceAnalyzerUtils::Decode7bit(BufferPtr);
+				uint32 TimerScopeDepth = IntCastChecked<uint32>(FTraceAnalyzerUtils::Decode7bit(BufferPtr));
 
 				if (TimerScopeDepth != 0)
 				{
@@ -510,7 +510,7 @@ uint64 FCpuProfilerAnalyzer::ProcessBufferV2(const FEventTime& EventTime, FThrea
 		{
 			if (DecodedCycle & 1ull)
 			{
-				uint32 SpecId = FTraceAnalyzerUtils::Decode7bit(BufferPtr);
+				uint32 SpecId = IntCastChecked<uint32>(FTraceAnalyzerUtils::Decode7bit(BufferPtr));
 				uint32 TimerId = GetTimerId(SpecId);
 
 				FEventScopeState& ScopeState = ThreadState.ScopeStack.AddDefaulted_GetRef();
diff --git a/Engine/Source/Developer/TraceServices/Private/Analyzers/CsvProfilerTraceAnalysis.h b/Engine/Source/Developer/TraceServices/Private/Analyzers/CsvProfilerTraceAnalysis.h
--- a/Engine/Source/Developer/TraceServices/Private/Analyzers/CsvProfilerTraceAnalysis.h
+++ b/Engine/Source/Developer/TraceServices/Private/Analyzers/CsvProfilerTraceAnalysis.h
@@ -87,7 +87,7 @@ private:
 		uint64 ProviderHandle = uint64(-1);
 		uint64 ProviderCountHandle = uint64(-1);
 		IEditableCounter* Counter = nullptr;
-		int64 CurrentFrame = -1;
+		uint32 CurrentFrame = 0;
 		FStatSeriesValue CurrentValue;
 		int64 CurrentCount = 0;
 		ECsvStatSeriesType Type = CsvStatSeriesType_CustomStatInt;
diff --git a/Engine/Source/Developer/TraceServices/Private/Analyzers/MetadataAnalysis.cpp b/Engine/Source/Developer/TraceServices/Private/Analyzers/MetadataAnalysis.cpp
--- a/Engine/Source/Developer/TraceServices/Private/Analyzers/MetadataAnalysis.cpp
+++ b/Engine/Source/Developer/TraceServices/Private/Analyzers/MetadataAnalysis.cpp
@@ -134,7 +134,7 @@ uint16 FMetadataAnalysis::GetOrRegisterType(const FEventTypeInfo& EventInfo)
 	}
 
 	uint16 MetadataTypeId = 0;
-	FMetadataSchema Schema(EventInfo.GetFieldCount());
+	FMetadataSchema Schema(static_cast<uint8>(EventInfo.GetFieldCount()));
 	FMetadataSchema::FBuilder Builder = Schema.Builder();
 	for (uint32 FieldIdx = 0; FieldIdx < EventInfo.GetFieldCount(); ++FieldIdx)
 	{
@@ -234,25 +234,25 @@ static void WriteData(ArrayType& OutData, const UE::Trace::IAnalyzer::FEventData
 				case 1:
 					{
 						auto RefValue = EventData.GetReferenceValue<uint8>(FieldIdx);
-						Writer.WriteField(FieldIdx, &RefValue, sizeof(UE::Trace::FEventRef8), OutData);
+						Writer.WriteField(static_cast<uint8>(FieldIdx), &RefValue, sizeof(UE::Trace::FEventRef8), OutData);
 					}
 					break;
 				case 2:
 					{
 						auto RefValue = EventData.GetReferenceValue<uint16>(FieldIdx);
-						Writer.WriteField(FieldIdx, &RefValue, sizeof(UE::Trace::FEventRef16), OutData);
+						Writer.WriteField(static_cast<uint8>(FieldIdx), &RefValue, sizeof(UE::Trace::FEventRef16), OutData);
 					}
 					break;
 				case 4:
 					{
 						auto RefValue = EventData.GetReferenceValue<uint32>(FieldIdx);
-						Writer.WriteField(FieldIdx, &RefValue, sizeof(UE::Trace::FEventRef32), OutData);
+						Writer.WriteField(static_cast<uint8>(FieldIdx), &RefValue, sizeof(UE::Trace::FEventRef32), OutData);
 					}
 					break;
 				case 8:
 					{
 						auto RefValue = EventData.GetReferenceValue<uint64>(FieldIdx);
-						Writer.WriteField(FieldIdx, &RefValue, sizeof(UE::Trace::FEventRef64), OutData);
+						Writer.WriteField(static_cast<uint8>(FieldIdx), &RefValue, sizeof(UE::Trace::FEventRef64), OutData);
 					}
 					break;
 				default:
@@ -266,7 +266,7 @@ static void WriteData(ArrayType& OutData, const UE::Trace::IAnalyzer::FEventData
 				FWideStringView View;
 				EventData.GetString(FieldInfo->GetName(), View);
 				const TCHAR* Str = Session.StoreString(View);
-				Writer.WriteField(FieldIdx, Str, sizeof(TCHAR*), OutData);
+				Writer.WriteField(static_cast<uint8>(FieldIdx), Str, sizeof(TCHAR*), OutData);
 			}
 			break;
 
@@ -275,7 +275,7 @@ static void WriteData(ArrayType& OutData, const UE::Trace::IAnalyzer::FEventData
 				FString String;
 				EventData.GetString(FieldInfo->GetName(), String);
 				const TCHAR* Str = Session.StoreString(String);
-				Writer.WriteField(FieldIdx, Str, sizeof(TCHAR*), OutData);
+				Writer.WriteField(static_cast<uint8>(FieldIdx), Str, sizeof(TCHAR*), OutData);
 			}
 			break;
 
@@ -284,7 +284,7 @@ static void WriteData(ArrayType& OutData, const UE::Trace::IAnalyzer::FEventData
 				const IAnalyzer::FEventFieldHandle FieldHandle = EventInfo.GetFieldHandleUnchecked(FieldIdx);
 				const void* Value = EventData.GetValueRaw(FieldHandle);
 				const uint8 Size = FieldInfo->GetSize();
-				Writer.WriteField(FieldIdx, Value, Size, OutData);
+				Writer.WriteField(static_cast<uint8>(FieldIdx), Value, Size, OutData);
 			}
 			break;
 		}
diff --git a/Engine/Source/Developer/TraceServices/Private/Analyzers/ModuleAnalysis.cpp b/Engine/Source/Developer/TraceServices/Private/Analyzers/ModuleAnalysis.cpp
--- a/Engine/Source/Developer/TraceServices/Private/Analyzers/ModuleAnalysis.cpp
+++ b/Engine/Source/Developer/TraceServices/Private/Analyzers/ModuleAnalysis.cpp
@@ -59,7 +59,7 @@ bool FModuleAnalyzer::OnEvent(uint16 RouteId, EStyle Style, const FOnEventContex
 	{
 		case RouteId_ModuleInit:
 			{
-				ModuleBaseShift = Context.EventData.GetValue("ModuleBaseShift", 16);
+				ModuleBaseShift = IntCastChecked<uint8>(Context.EventData.GetValue("ModuleBaseShift", 16));
 				FAnsiStringView SymbolFormat;
 				if (Context.EventData.GetString("SymbolFormat", SymbolFormat))
 				{
diff --git a/Engine/Source/Developer/TraceServices/Private/Analyzers/NetTraceAnalyzer.cpp b/Engine/Source/Developer/TraceServices/Private/Analyzers/NetTraceAnalyzer.cpp
--- a/Engine/Source/Developer/TraceServices/Private/Analyzers/NetTraceAnalyzer.cpp
+++ b/Engine/Source/Developer/TraceServices/Private/Analyzers/NetTraceAnalyzer.cpp
@@ -106,7 +106,7 @@ bool FNetTraceAnalyzer::OnEvent(uint16 RouteId, EStyle Style, const FOnEventCont
 			// we always trace the version so that we make sure that we are backwards compatible with older trace stream
 			NetTraceVersion = EventData.GetValue<uint32>("NetTraceVersion");
 			NetTraceReporterVersion = EventData.GetValue<uint32>("NetTraceReporterVersion");
-			
+
 			NetProfilerProvider.SetNetTraceVersion(NetTraceVersion);
 		}
 		break;
@@ -199,7 +199,7 @@ bool FNetTraceAnalyzer::OnEvent(uint16 RouteId, EStyle Style, const FOnEventCont
 			HandleObjectDestroyedEvent(Context, EventData);
 		}
 		break;
-		
+
 		case RouteId_InstanceUpdatedEvent:
 		{
 			HandleGameInstanceUpdatedEvent(Context, EventData);
@@ -279,27 +279,27 @@ void FNetTraceAnalyzer::HandlePacketContentEvent(const FOnEventContext& Context,
 				}
 				else if (DecodedEventType == EContentEventType::NameId)
 				{
-					if (const uint32* NetProfilerNameIndex = TracedNameIdToNetProfilerNameIdMap.Find(DecodedNameOrObjectId))
+					if (const uint32* NetProfilerNameIndex = TracedNameIdToNetProfilerNameIdMap.Find(IntCastChecked<uint16>(DecodedNameOrObjectId)))
 					{
 						Event.NameIndex = *NetProfilerNameIndex;
 					}
 					else
 					{
-						UE_LOG(LogNetTrace, Warning, TEXT("PacketContentEvent GameInstanceId: %u, ConnectionId: %u %s, Missing NameIndex: %llu"), (uint32)GameInstanceId, (uint32)ConnectionId, ConnectionMode ? TEXT("Incoming") : TEXT("Outgoing"), DecodedNameOrObjectId);	
+						UE_LOG(LogNetTrace, Warning, TEXT("PacketContentEvent GameInstanceId: %u, ConnectionId: %u %s, Missing NameIndex: %llu"), (uint32)GameInstanceId, (uint32)ConnectionId, ConnectionMode ? TEXT("Incoming") : TEXT("Outgoing"), DecodedNameOrObjectId);
 					}
 				}
 
 				// EventTypeIndex does not match NameIndex as we might see the same name on different levels
-				Event.EventTypeIndex = GetTracedEventTypeIndex(Event.NameIndex, Event.Level);
+				Event.EventTypeIndex = GetTracedEventTypeIndex(static_cast<uint16>(Event.NameIndex), static_cast<uint8>(Event.Level));
 			}
 			break;
 
 			case EContentEventType::BunchEvent:
 			{
-				const uint64 DecodedNameId = FTraceAnalyzerUtils::Decode7bit(BufferPtr);
-				const uint64 DecodedEventStartPos = FTraceAnalyzerUtils::Decode7bit(BufferPtr);
-				const uint64 DecodedEventEndPos = FTraceAnalyzerUtils::Decode7bit(BufferPtr) + DecodedEventStartPos;
-				
+				const uint16 DecodedNameId = IntCastChecked<uint16>(FTraceAnalyzerUtils::Decode7bit(BufferPtr));
+				const uint32 DecodedEventStartPos = IntCastChecked<uint32>(FTraceAnalyzerUtils::Decode7bit(BufferPtr));
+				const uint32 DecodedEventEndPos = IntCastChecked<uint32>(FTraceAnalyzerUtils::Decode7bit(BufferPtr)) + DecodedEventStartPos;
+
 				const uint32* NetProfilerNameIndex = DecodedNameId ? TracedNameIdToNetProfilerNameIdMap.Find(DecodedNameId) : nullptr;
 
 				FBunchInfo BunchInfo;
@@ -308,7 +308,7 @@ void FNetTraceAnalyzer::HandlePacketContentEvent(const FOnEventContext& Context,
 				BunchInfo.HeaderBits = 0U;
 				BunchInfo.BunchBits = DecodedEventEndPos;
 				BunchInfo.FirstBunchEventIndex = Events.Num();
-				BunchInfo.NameIndex = NetProfilerNameIndex ? *NetProfilerNameIndex : 0U;
+				BunchInfo.NameIndex = NetProfilerNameIndex ? IntCastChecked<uint16>(*NetProfilerNameIndex) : 0U;
 
 				BunchInfos.Add(BunchInfo);
 
@@ -319,8 +319,8 @@ void FNetTraceAnalyzer::HandlePacketContentEvent(const FOnEventContext& Context,
 
 			case EContentEventType::BunchHeaderEvent:
 			{
-				const uint64 DecodedEventCount = FTraceAnalyzerUtils::Decode7bit(BufferPtr);
-				const uint64 DecodedHeaderBits = FTraceAnalyzerUtils::Decode7bit(BufferPtr);
+				const uint32 DecodedEventCount = IntCastChecked<uint32>(FTraceAnalyzerUtils::Decode7bit(BufferPtr));
+				const uint32 DecodedHeaderBits = IntCastChecked<uint32>(FTraceAnalyzerUtils::Decode7bit(BufferPtr));
 
 				FBunchInfo& BunchInfo = BunchInfos.Last();
 
@@ -349,8 +349,8 @@ void FNetTraceAnalyzer::HandlePacketContentEvent(const FOnEventContext& Context,
 						}
 						else
 						{
-							const uint32* ExistingChannelNameIndex = GameInstanceState->ChannelNames.Find(BunchInfo.BunchInfo.ChannelIndex);						
-							BunchInfo.NameIndex = ExistingChannelNameIndex ? *ExistingChannelNameIndex : 0U;
+							const uint32* ExistingChannelNameIndex = GameInstanceState->ChannelNames.Find(BunchInfo.BunchInfo.ChannelIndex);
+							BunchInfo.NameIndex = ExistingChannelNameIndex ? IntCastChecked<uint16>(*ExistingChannelNameIndex) : 0U;
 						}
 
 						BunchInfo.BunchInfo.bIsValid = 1U;
@@ -375,13 +375,13 @@ void FNetTraceAnalyzer::HandlePacketContentEvent(const FOnEventContext& Context,
 void FNetTraceAnalyzer::AddEvent(TPagedArray<FNetProfilerContentEvent>& Events, const FNetProfilerContentEvent& InEvent, uint32 Offset, uint32 LevelOffset)
 {
 	FNetProfilerContentEvent& Event = Events.PushBack();
-	
-	Event.EventTypeIndex = GetTracedEventTypeIndex(InEvent.NameIndex, InEvent.Level + LevelOffset);
-	Event.NameIndex =  InEvent.NameIndex;
+
+	Event.EventTypeIndex = GetTracedEventTypeIndex(IntCastChecked<uint16>(InEvent.NameIndex), IntCastChecked<uint8>(InEvent.Level + LevelOffset));
+	Event.NameIndex = InEvent.NameIndex;
 	Event.ObjectInstanceIndex = InEvent.ObjectInstanceIndex;
 	Event.StartPos = InEvent.StartPos + Offset;
 	Event.EndPos = InEvent.EndPos + Offset;
-	Event.Level = InEvent.Level + LevelOffset;	
+	Event.Level = InEvent.Level + LevelOffset;
 	Event.BunchInfo = InEvent.BunchInfo;
 }
 
@@ -389,11 +389,11 @@ void FNetTraceAnalyzer::AddEvent(TPagedArray<FNetProfilerContentEvent>& Events,
 {
 	FNetProfilerContentEvent& Event = Events.PushBack();
 
-	Event.EventTypeIndex = GetTracedEventTypeIndex(NameIndex, Level);
-	Event.NameIndex = NameIndex; 
+	Event.EventTypeIndex = GetTracedEventTypeIndex(IntCastChecked<uint16>(NameIndex), IntCastChecked<uint8>(Level));
+	Event.NameIndex = NameIndex;
 	Event.ObjectInstanceIndex = 0;
 	Event.StartPos = StartPos;
-	Event.EndPos = EndPos;			
+	Event.EndPos = EndPos;
 	Event.Level = Level;
 	Event.BunchInfo = BunchInfo;
 }
@@ -419,8 +419,8 @@ void FNetTraceAnalyzer::FlushPacketEvents(FNetTraceConnectionState& ConnectionSt
 		const FNetProfilerContentEvent& BunchEvent = BunchEvents[CurrentBunchEventIndex];
 
 		AddEvent(Events, BunchEvent, 0U, 0U);
-	
-		NextBunchOffset = FMath::Max<uint64>(BunchEvent.EndPos, NextBunchOffset);
+
+		NextBunchOffset = FMath::Max(static_cast<uint32>(BunchEvent.EndPos), NextBunchOffset);
 		++CurrentBunchEventIndex;
 		++ConnectionData.ContentEventChangeCount;
 	}
@@ -439,7 +439,7 @@ void FNetTraceAnalyzer::FlushPacketEvents(FNetTraceConnectionState& ConnectionSt
 
 			// Bunch header event
 			AddEvent(Events, NextBunchOffset, NextBunchOffset + Bunch.HeaderBits, 1, BunchHeaderNameIndex, FNetProfilerBunchInfo::MakeBunchInfo(0));
-	
+
 			// Add events belonging to bunch, including the ones from merged bunches
 			for (uint32 EventIt = 0; EventIt < EventsToAdd; ++EventIt)
 			{
@@ -456,7 +456,7 @@ void FNetTraceAnalyzer::FlushPacketEvents(FNetTraceConnectionState& ConnectionSt
 			// Reset event count
 			EventsToAdd = 0U;
 		}
-		
+
 		++ConnectionData.ContentEventChangeCount;
 	}
 
@@ -480,7 +480,7 @@ void FNetTraceAnalyzer::HandlePacketEvent(const FOnEventContext& Context, const
 
 	// Get the NetProfilerFrameIndex for the current engine frame/timestamp
 	const uint32 NetProfilerFrameIndex = GetCurrentNetProfilerFrameIndexAndFlushFrameStatsCountersIfNeeded(
-		GameInstanceId, 
+		GameInstanceId,
 		FrameProvider.GetFrameNumberForTimestamp(ETraceFrameType::TraceFrameType_Game, LastTimeStamp)
 	);
 
@@ -502,7 +502,7 @@ void FNetTraceAnalyzer::HandlePacketEvent(const FOnEventContext& Context, const
 
 	// Fill in packet data a packet must have at least 1 event?
 	Packet.StartEventIndex = ConnectionState->CurrentPacketStartIndex[ConnectionMode];
-	Packet.EventCount = ConnectionData.ContentEvents.Num() - Packet.StartEventIndex;
+	Packet.EventCount = static_cast<uint32>(ConnectionData.ContentEvents.Num()) - Packet.StartEventIndex;
 	Packet.TimeStamp = GetLastTimestamp();
 	Packet.SequenceNumber = SequenceNumber;
 	Packet.DeliveryStatus = ENetProfilerDeliveryStatus::Unknown;
@@ -513,7 +513,7 @@ void FNetTraceAnalyzer::HandlePacketEvent(const FOnEventContext& Context, const
 	Packet.DeliveryStatus = ENetProfilerDeliveryStatus::Delivered;
 
 	// Flush PacketStats
-	Packet.StartStatsIndex = ConnectionData.PacketStats.Num();
+	Packet.StartStatsIndex = static_cast<uint32>(ConnectionData.PacketStats.Num());
 	Packet.StatsCount = ConnectionState->PacketStats.Num();
 	for (const FNetProfilerStats& Stat : ConnectionState->PacketStats)
 	{
@@ -523,11 +523,11 @@ void FNetTraceAnalyzer::HandlePacketEvent(const FOnEventContext& Context, const
 	ConnectionState->PacketStats.Reset();
 	++ConnectionData.PacketStatsChangeCount;
 
-	
+
 	// Mark the beginning of a new packet
-	ConnectionState->CurrentPacketStartIndex[ConnectionMode] = ConnectionData.ContentEvents.Num();
+	ConnectionState->CurrentPacketStartIndex[ConnectionMode] = static_cast<uint32>(ConnectionData.ContentEvents.Num());
 	ConnectionState->CurrentPacketBitOffset[ConnectionMode] = 0U;
-	ConnectionState->CurrentPacketStatsStartIndex[ConnectionMode] = ConnectionData.PacketStats.Num();
+	ConnectionState->CurrentPacketStatsStartIndex[ConnectionMode] = static_cast<uint32>(ConnectionData.PacketStats.Num());
 
 	//UE_LOG(LogNetTrace, Log, TEXT("PacketEvent GameInstanceId: %u, ConnectionId: %u, %s, Seq: %u PacketBits: %u"), (uint32)GameInstanceId, (uint32)ConnectionId, ConnectionMode ? TEXT("Incoming") : TEXT("Outgoing"), SequenceNumber, Packet.ContentSizeInBits);
 }
@@ -542,7 +542,7 @@ void FNetTraceAnalyzer::FlushFrameStatsCounters(FNetTraceAnalyzer::FNetTraceGame
 		{
 			FNetProfilerFrame& Frame = GameInstance->Frames->EmplaceBack();
 			Frame.EngineFrameNumber = GameInstanceState.CurrentEngineFrameIndex;
-			Frame.StartStatsIndex = GameInstance->FrameStats->Num();
+			Frame.StartStatsIndex = static_cast<uint32>(GameInstance->FrameStats->Num());
 			Frame.StatsCount = GameInstanceState.FrameStatsCounters.Num();
 			Frame.TimeStamp = LastTimeStamp;
 		}
@@ -558,7 +558,7 @@ void FNetTraceAnalyzer::FlushFrameStatsCounters(FNetTraceAnalyzer::FNetTraceGame
 		}
 
 		GameInstanceState.FrameStatsCounters.Reset();
-		GameInstanceState.CurrentNetProfilerFrameIndex = GameInstance->Frames->Num();
+		GameInstanceState.CurrentNetProfilerFrameIndex = static_cast<uint32>(GameInstance->Frames->Num());
 
 		// Mark frames dirty
 		++GameInstance->FramesChangeCount;
@@ -570,7 +570,7 @@ uint32 FNetTraceAnalyzer::GetCurrentNetProfilerFrameIndexAndFlushFrameStatsCount
 	TSharedRef<FNetTraceAnalyzer::FNetTraceGameInstanceState> GameInstanceStateRef = GetOrCreateActiveGameInstanceState(GameInstanceId);
 	FNetTraceAnalyzer::FNetTraceGameInstanceState& GameInstanceState = GameInstanceStateRef.Get();
 	if (EngineFrameIndex > GameInstanceState.CurrentEngineFrameIndex)
-	{		
+	{
 		FlushFrameStatsCounters(GameInstanceState);
 		GameInstanceState.CurrentEngineFrameIndex = EngineFrameIndex;
 	}
@@ -630,18 +630,18 @@ void FNetTraceAnalyzer::HandleConnectionCreatedEvent(const FOnEventContext& Cont
 
 uint32 FNetTraceAnalyzer::GetOrCreateNetProfilerStatsCounterTypeIndex(uint32 NameId, ENetProfilerStatsCounterType StatsType)
 {
-	if (const uint32* ExistingNetProfilerStatsCounterTypeIndex = TraceNetStatsCounterIdToNetProfilerStatsCounterTypeIndexMap.Find(NameId))
+	if (const uint32* ExistingNetProfilerStatsCounterTypeIndex = TraceNetStatsCounterIdToNetProfilerStatsCounterTypeIndexMap.Find(static_cast<uint16>(NameId)))
 	{
 		return *ExistingNetProfilerStatsCounterTypeIndex;
 	}
 	else
 	{
 		// Add new counter type
-		const uint32* NetProfilerNameIndex = TracedNameIdToNetProfilerNameIdMap.Find(NameId);
+		const uint32* NetProfilerNameIndex = TracedNameIdToNetProfilerNameIdMap.Find(static_cast<uint16>(NameId));
 		const uint32 NameIndex = NetProfilerNameIndex ? *NetProfilerNameIndex : 0u;
 
 		uint32 NetProfilerStatsCounterTypeIndex = NetProfilerProvider.AddNetProfilerStatsCounterType(NameIndex, StatsType);
-		TraceNetStatsCounterIdToNetProfilerStatsCounterTypeIndexMap.Add(NameId, NetProfilerStatsCounterTypeIndex);
+		TraceNetStatsCounterIdToNetProfilerStatsCounterTypeIndexMap.Add(static_cast<uint16>(NameId), NetProfilerStatsCounterTypeIndex);
 
 		return NetProfilerStatsCounterTypeIndex;
 	}
@@ -653,7 +653,7 @@ void FNetTraceAnalyzer::HandlePacketStatsCounterEvent(const FOnEventContext& Con
 	const uint8 GameInstanceId = EventData.GetValue<uint8>("GameInstanceId");
 	const uint16 ConnectionId = EventData.GetValue<uint16>("ConnectionId");
 	const uint16 NameId = EventData.GetValue<uint16>("NameId");
-	
+
 	FNetProfilerStats Stats;
 	Stats.StatsCounterTypeIndex = GetOrCreateNetProfilerStatsCounterTypeIndex(NameId, ENetProfilerStatsCounterType::Packet);
 	Stats.StatsValue = StatsCounterValue;
@@ -662,7 +662,7 @@ void FNetTraceAnalyzer::HandlePacketStatsCounterEvent(const FOnEventContext& Con
 	TSharedRef<FNetTraceGameInstanceState> GameInstanceState = GetOrCreateActiveGameInstanceState(GameInstanceId);
 	if (TSharedRef<FNetTraceConnectionState>* ConnectionState = GameInstanceState->ActiveConnections.Find(ConnectionId))
 	{
-		(*ConnectionState)->PacketStats.Emplace(Stats);			
+		(*ConnectionState)->PacketStats.Emplace(Stats);
 	}
 }
 
@@ -715,7 +715,7 @@ void FNetTraceAnalyzer::HandleConnectionUpdatedEvent(const FOnEventContext& Cont
 	EventData.GetString("Address", AddressString);
 
 	TSharedRef<FNetTraceGameInstanceState> GameInstanceState = GetOrCreateActiveGameInstanceState(GameInstanceId);
-	
+
 	if (TSharedRef<FNetTraceConnectionState>* ConnectionState = GameInstanceState->ActiveConnections.Find(ConnectionId))
 	{
 		if (FNetProfilerConnectionInternal* Connection = NetProfilerProvider.EditConnection((*ConnectionState)->ConnectionIndex))
@@ -766,7 +766,7 @@ void FNetTraceAnalyzer::HandleObjectCreatedEvent(const FOnEventContext& Context,
 	TSharedRef<FNetTraceGameInstanceState> GameInstanceState = GetOrCreateActiveGameInstanceState(GameInstanceId);
 	const uint32* NetProfilerNameIndex = TracedNameIdToNetProfilerNameIdMap.Find(NameId);
 	const uint32 NameIndex = NetProfilerNameIndex ? *NetProfilerNameIndex : 0u;
-	
+
 	if (GameInstanceState->ActiveObjects.Contains(ObjectId))
 	{
 		if (FNetProfilerObjectInstance* ExistingInstance = NetProfilerProvider.EditObject(GameInstanceState->GameInstanceIndex, GameInstanceState->ActiveObjects[ObjectId].ObjectIndex))
@@ -791,8 +791,8 @@ void FNetTraceAnalyzer::HandleObjectCreatedEvent(const FOnEventContext& Context,
 	FNetProfilerObjectInstance& ObjectInstance = NetProfilerProvider.CreateObject(GameInstanceState->GameInstanceIndex);
 
 	// Fill in object data
-	ObjectInstance.LifeTime.Begin =  GetLastTimestamp();
-	ObjectInstance.NameIndex = NameIndex;
+	ObjectInstance.LifeTime.Begin = GetLastTimestamp();
+	ObjectInstance.NameIndex = IntCastChecked<uint16>(NameIndex);
 	ObjectInstance.NetId = ObjectId;
 	ObjectInstance.TypeId = TypeId;
 
@@ -837,7 +837,8 @@ TSharedRef<FNetTraceAnalyzer::FNetTraceGameInstanceState> FNetTraceAnalyzer::Get
 		ActiveGameInstances.Add(GameInstanceId, GameInstanceState);
 		GameInstanceState->GameInstanceIndex = GameInstance.Instance.GameInstanceIndex;
 
-		GameInstanceState->CurrentEngineFrameIndex = FMath::Min<uint32>(FrameProvider.GetFrameCount(TraceFrameType_Game) - 1U, 0U);
+		const uint32 FrameCount = static_cast<uint32>(FrameProvider.GetFrameCount(TraceFrameType_Game));
+		GameInstanceState->CurrentEngineFrameIndex = (FrameCount > 0) ? FrameCount - 1 : 0;
 
 		return GameInstanceState;
 	}
@@ -866,7 +867,7 @@ FNetTraceAnalyzer::FNetTraceConnectionState* FNetTraceAnalyzer::GetActiveConnect
 			return &(*ConnectionState).Get();
 		}
 	}
-		
+
 	return nullptr;
 }
 
diff --git a/Engine/Source/Developer/TraceServices/Private/Analyzers/NetTraceAnalyzer.h b/Engine/Source/Developer/TraceServices/Private/Analyzers/NetTraceAnalyzer.h
--- a/Engine/Source/Developer/TraceServices/Private/Analyzers/NetTraceAnalyzer.h
+++ b/Engine/Source/Developer/TraceServices/Private/Analyzers/NetTraceAnalyzer.h
@@ -97,7 +97,7 @@ private:
 		TArray<FNetProfilerContentEvent> BunchEvents[ENetProfilerConnectionMode::Count];
 
 		TArray<FBunchInfo> BunchInfos[ENetProfilerConnectionMode::Count];
-		
+
 		ENetProfilerConnectionState ConnectionState;
 	};
 
@@ -114,7 +114,7 @@ private:
 		TMap<uint64, FNetTraceActiveObjectState> ActiveObjects;
 		TMap<int32, uint32> ChannelNames;
 
-		// StatsCounters reported during the current fram
+		// StatsCounters reported during the current frame
 		TArray<FNetProfilerStats> FrameStatsCounters;
 		uint32 CurrentEngineFrameIndex = 0;
 		uint32 CurrentNetProfilerFrameIndex = 0;
@@ -138,7 +138,7 @@ private:
 	void FlushFrameStatsCounters(FNetTraceGameInstanceState& GameInstanceState);
 	uint32 GetCurrentNetProfilerFrameIndexAndFlushFrameStatsCountersIfNeeded(uint32 GameInstanceId, uint32 EngineFrameIndex);
 
-	void HandlePacketEvent(const FOnEventContext& Context, const FEventData& EventData);	
+	void HandlePacketEvent(const FOnEventContext& Context, const FEventData& EventData);
 	void HandlePacketContentEvent(const FOnEventContext& Context, const FEventData& EventData);
 	void HandlePacketDroppedEvent(const FOnEventContext& Context, const FEventData& EventData);
 	void HandleConnectionStateUpdatedEvent(const FOnEventContext& Context, const FEventData& EventData);
diff --git a/Engine/Source/Developer/TraceServices/Private/Analyzers/StatsTraceAnalysis.cpp b/Engine/Source/Developer/TraceServices/Private/Analyzers/StatsTraceAnalysis.cpp
--- a/Engine/Source/Developer/TraceServices/Private/Analyzers/StatsTraceAnalysis.cpp
+++ b/Engine/Source/Developer/TraceServices/Private/Analyzers/StatsTraceAnalysis.cpp
@@ -137,7 +137,7 @@ bool FStatsAnalyzer::OnEvent(uint16 RouteId, EStyle Style, const FOnEventContext
 				};
 
 				uint64 DecodedIdAndOp = FTraceAnalyzerUtils::Decode7bit(BufferPtr);
-				uint32 StatId = DecodedIdAndOp >> 3;
+				uint32 StatId = static_cast<uint32>(DecodedIdAndOp >> 3);
 				IEditableCounter* EditableCounter = EditableCountersMap.FindRef(StatId);
 				if (!EditableCounter)
 				{
diff --git a/Engine/Source/Developer/TraceServices/Private/Common/FormatArgs.cpp b/Engine/Source/Developer/TraceServices/Private/Common/FormatArgs.cpp
--- a/Engine/Source/Developer/TraceServices/Private/Common/FormatArgs.cpp
+++ b/Engine/Source/Developer/TraceServices/Private/Common/FormatArgs.cpp
@@ -148,15 +148,15 @@ const TCHAR* FFormatArgsHelper::ExtractNextFormatArg(const TCHAR* FormatString,
 			}
 
 			++Src;
-			uint64 FormatSpecifierLength = Src + 1 - FormatSpecifierStart;
+			int32 FormatSpecifierLength = static_cast<int32>(Src + 1 - FormatSpecifierStart);
 			check(FormatSpecifierLength < 255);
 			FCString::Strncpy(Spec.FormatString, FormatSpecifierStart, FormatSpecifierLength);
 			Spec.Valid = true;
-			Spec.PassthroughLength = FormatSpecifierStart - FormatString;
+			Spec.PassthroughLength = static_cast<int32>(FormatSpecifierStart - FormatString);
 			return Src;
 		}
 	}
-	Spec.PassthroughLength = Src - FormatString;
+	Spec.PassthroughLength = static_cast<int32>(Src - FormatString);
 	return Src;
 }
 
@@ -249,7 +249,7 @@ double FFormatArgsHelper::ExtractFloatingPointArgument(FFormatArgsStreamContext&
 	return Result;
 }
 
-const TCHAR* FFormatArgsHelper::ExtractStringArgument(FFormatArgsStreamContext& ArgStream, TCHAR* Temp, uint64 MaxTemp)
+const TCHAR* FFormatArgsHelper::ExtractStringArgument(FFormatArgsStreamContext& ArgStream, TCHAR* Temp, int32 MaxTemp)
 {
 	static TCHAR Empty[] = TEXT("");
 	const TCHAR* Result = Empty;
@@ -276,7 +276,7 @@ const TCHAR* FFormatArgsHelper::ExtractStringArgument(FFormatArgsStreamContext&
 	return Result;
 }
 
-int32 FFormatArgsHelper::FormatArgument(TCHAR* Out, uint64 MaxOut, TCHAR* Temp, uint64 MaxTemp, const FFormatArgSpec& ArgSpec, FFormatArgsStreamContext& ArgStream)
+int32 FFormatArgsHelper::FormatArgument(TCHAR* Out, int32 MaxOut, TCHAR* Temp, int32 MaxTemp, const FFormatArgSpec& ArgSpec, FFormatArgsStreamContext& ArgStream)
 {
 	check(ArgSpec.AdditionalIntegerArgumentCount <= 2);
 	switch (ArgSpec.ExpectedTypeCategory)
@@ -357,7 +357,7 @@ void FFormatArgsHelper::Format(TCHAR* Out, uint64 MaxOut, TCHAR* Temp, uint64 Ma
 	{
 		FFormatArgSpec Spec;
 		const TCHAR* NextSrc = ExtractNextFormatArg(Src, Spec);
-		uint64 PassthroughCopyLength = FMath::Min<uint64>(Spec.PassthroughLength, DstEnd - Dst);
+		int32 PassthroughCopyLength = FMath::Min(Spec.PassthroughLength, static_cast<int32>(DstEnd - Dst));
 		if (PassthroughCopyLength)
 		{
 			FCString::Strncpy(Dst, Src, PassthroughCopyLength + 1);
@@ -365,7 +365,7 @@ void FFormatArgsHelper::Format(TCHAR* Out, uint64 MaxOut, TCHAR* Temp, uint64 Ma
 		}
 		if (Spec.Valid)
 		{
-			int32 Length = FormatArgument(Dst, DstEnd - Dst, Temp, MaxTemp, Spec, ArgumentStream);
+			int32 Length = FormatArgument(Dst, static_cast<int32>(DstEnd - Dst), Temp, static_cast<int32>(MaxTemp), Spec, ArgumentStream);
 			if (Length < 0)
 			{
 				break;
diff --git a/Engine/Source/Developer/TraceServices/Private/Common/FormatArgs.h b/Engine/Source/Developer/TraceServices/Private/Common/FormatArgs.h
--- a/Engine/Source/Developer/TraceServices/Private/Common/FormatArgs.h
+++ b/Engine/Source/Developer/TraceServices/Private/Common/FormatArgs.h
@@ -14,7 +14,7 @@ struct FFormatArgsHelper
 private:
 	struct FFormatArgSpec
 	{
-		uint64 PassthroughLength;
+		int32 PassthroughLength;
 		TCHAR FormatString[255];
 		uint8 ExpectedTypeCategory;
 		uint8 AdditionalIntegerArgumentCount;
@@ -36,8 +36,8 @@ private:
 	static bool AdvanceArgumentStream(FFormatArgsStreamContext& Context);
 	static uint64 ExtractIntegerArgument(FFormatArgsStreamContext& ArgStream);
 	static double ExtractFloatingPointArgument(FFormatArgsStreamContext& ArgStream);
-	static const TCHAR* ExtractStringArgument(FFormatArgsStreamContext& ArgStream, TCHAR* Temp, uint64 MaxTemp);
-	static int32 FormatArgument(TCHAR* Out, uint64 MaxOut, TCHAR* Temp, uint64 MaxTemp, const FFormatArgSpec& ArgSpec, FFormatArgsStreamContext& ArgStream);
+	static const TCHAR* ExtractStringArgument(FFormatArgsStreamContext& ArgStream, TCHAR* Temp, int32 MaxTemp);
+	static int32 FormatArgument(TCHAR* Out, int32 MaxOut, TCHAR* Temp, int32 MaxTemp, const FFormatArgSpec& ArgSpec, FFormatArgsStreamContext& ArgStream);
 };
 
 } // namespace TraceServices
diff --git a/Engine/Source/Developer/TraceServices/Private/Common/TimelineStatistics.h b/Engine/Source/Developer/TraceServices/Private/Common/TimelineStatistics.h
--- a/Engine/Source/Developer/TraceServices/Private/Common/TimelineStatistics.h
+++ b/Engine/Source/Developer/TraceServices/Private/Common/TimelineStatistics.h
@@ -239,7 +239,7 @@ private:
 
 		if (ComputeMedian)
 		{
-			const int32 HalfCount = Stats.InstanceCount / 2;
+			const int32 HalfCount = static_cast<int32>(Stats.InstanceCount / 2);
 
 			// Compute median inclusive time.
 			int32 InclCount = 0;
diff --git a/Engine/Source/Developer/TraceServices/Private/Common/Utils.cpp b/Engine/Source/Developer/TraceServices/Private/Common/Utils.cpp
--- a/Engine/Source/Developer/TraceServices/Private/Common/Utils.cpp
+++ b/Engine/Source/Developer/TraceServices/Private/Common/Utils.cpp
@@ -18,8 +18,8 @@ void Table2Csv(const IUntypedTable& Table, const TCHAR* Filename)
 	check(OutputFile);
 	FString Header;
 	const ITableLayout& Layout = Table.GetLayout();
-	int32 ColumnCount = Layout.GetColumnCount();
-	for (int32 ColumnIndex = 0; ColumnIndex < ColumnCount; ++ColumnIndex)
+	uint32 ColumnCount = static_cast<uint32>(Layout.GetColumnCount());
+	for (uint32 ColumnIndex = 0; ColumnIndex < ColumnCount; ++ColumnIndex)
 	{
 		Header += Layout.GetColumnName(ColumnIndex);
 		if (ColumnIndex < ColumnCount - 1)
@@ -37,7 +37,7 @@ void Table2Csv(const IUntypedTable& Table, const TCHAR* Filename)
 	for (; TableReader->IsValid(); TableReader->NextRow())
 	{
 		FString Line;
-		for (int32 ColumnIndex = 0; ColumnIndex < ColumnCount; ++ColumnIndex)
+		for (uint32 ColumnIndex = 0; ColumnIndex < ColumnCount; ++ColumnIndex)
 		{
 			switch (Layout.GetColumnType(ColumnIndex))
 			{
diff --git a/Engine/Source/Developer/TraceServices/Private/Model/AllocationsProvider.cpp b/Engine/Source/Developer/TraceServices/Private/Model/AllocationsProvider.cpp
--- a/Engine/Source/Developer/TraceServices/Private/Model/AllocationsProvider.cpp
+++ b/Engine/Source/Developer/TraceServices/Private/Model/AllocationsProvider.cpp
@@ -1337,7 +1337,8 @@ void FAllocationsProvider::EditAlloc(double Time, uint32 CallstackId, uint64 Add
 		}
 
 		INSIGHTS_SLOW_CHECK(Allocation.Address == Address);
-		Allocation.SizeAndAlignment = FAllocationItem::PackSizeAndAlignment(InSize, InAlignment);
+		INSIGHTS_SLOW_CHECK(InAlignment < 256);
+		Allocation.SizeAndAlignment = FAllocationItem::PackSizeAndAlignment(InSize, static_cast<uint8>(InAlignment));
 		Allocation.StartEventIndex = EventIndex[RootHeap];
 		Allocation.EndEventIndex = (uint32)-1;
 		Allocation.StartTime = Time;
@@ -1347,7 +1348,7 @@ void FAllocationsProvider::EditAlloc(double Time, uint32 CallstackId, uint64 Add
 		Allocation.FreeCallstackId = 0; // no callstack yet
 		Allocation.MetadataId = MetadataId;
 		Allocation.Tag = Tag;
-		Allocation.RootHeap = RootHeap;
+		Allocation.RootHeap = static_cast<uint8>(RootHeap);
 		Allocation.Flags = EMemoryTraceHeapAllocationFlags::None;
 
 		UpdateHistogramByAllocSize(InSize);
@@ -1554,7 +1555,7 @@ void FAllocationsProvider::EditMarkAllocationAsHeap(double Time, uint64 Address,
 		// Mark allocation as a "heap" allocation.
 		check(EnumHasAnyFlags(Flags, EMemoryTraceHeapAllocationFlags::Heap));
 		Alloc->Flags = Flags;
-		Alloc->RootHeap = RootHeap;
+		Alloc->RootHeap = static_cast<uint8>(RootHeap);
 
 		// Re-add it to the Live allocs as a heap allocation.
 		LiveAllocs[RootHeap]->AddHeap(Alloc);
@@ -1593,7 +1594,7 @@ void FAllocationsProvider::UpdateHistogramByAllocSize(uint64 Size)
 	// i : [2^(i-1) .. 2^i-1], i > 0
 	// ...
 	// 64 : [2^63 .. 2^64-1]
-	uint32 HistogramIndexPow2 = 64 - FMath::CountLeadingZeros64(Size);
+	uint32 HistogramIndexPow2 = 64 - static_cast<uint32>(FMath::CountLeadingZeros64(Size));
 	++AllocSizeHistogramPow2[HistogramIndexPow2];
 }
 
@@ -1884,7 +1885,7 @@ void FAllocationsProvider::GetTimelineIndexRange(double StartTime, double EndTim
 
 void FAllocationsProvider::EnumerateMinTotalAllocatedMemoryTimeline(int32 StartIndex, int32 EndIndex, TFunctionRef<void(double Time, double Duration, uint64 Value)> Callback) const
 {
-	const int32 NumPoints = Timeline.Num();
+	const int32 NumPoints = static_cast<int32>(Timeline.Num());
 	StartIndex = FMath::Max(StartIndex, 0);
 	EndIndex = FMath::Min(EndIndex + 1, NumPoints); // make it exclusive
 	if (StartIndex < EndIndex)
@@ -1918,7 +1919,7 @@ void FAllocationsProvider::EnumerateMinTotalAllocatedMemoryTimeline(int32 StartI
 
 void FAllocationsProvider::EnumerateMaxTotalAllocatedMemoryTimeline(int32 StartIndex, int32 EndIndex, TFunctionRef<void(double Time, double Duration, uint64 Value)> Callback) const
 {
-	const int32 NumPoints = Timeline.Num();
+	const int32 NumPoints = static_cast<int32>(Timeline.Num());
 	StartIndex = FMath::Max(StartIndex, 0);
 	EndIndex = FMath::Min(EndIndex + 1, NumPoints); // make it exclusive
 	if (StartIndex < EndIndex)
@@ -1952,7 +1953,7 @@ void FAllocationsProvider::EnumerateMaxTotalAllocatedMemoryTimeline(int32 StartI
 
 void FAllocationsProvider::EnumerateMinLiveAllocationsTimeline(int32 StartIndex, int32 EndIndex, TFunctionRef<void(double Time, double Duration, uint32 Value)> Callback) const
 {
-	const int32 NumPoints = Timeline.Num();
+	const int32 NumPoints = static_cast<int32>(Timeline.Num());
 	StartIndex = FMath::Max(StartIndex, 0);
 	EndIndex = FMath::Min(EndIndex + 1, NumPoints); // make it exclusive
 	if (StartIndex < EndIndex)
@@ -1986,7 +1987,7 @@ void FAllocationsProvider::EnumerateMinLiveAllocationsTimeline(int32 StartIndex,
 
 void FAllocationsProvider::EnumerateMaxLiveAllocationsTimeline(int32 StartIndex, int32 EndIndex, TFunctionRef<void(double Time, double Duration, uint32 Value)> Callback) const
 {
-	const int32 NumPoints = Timeline.Num();
+	const int32 NumPoints = static_cast<int32>(Timeline.Num());
 	StartIndex = FMath::Max(StartIndex, 0);
 	EndIndex = FMath::Min(EndIndex + 1, NumPoints); // make it exclusive
 	if (StartIndex < EndIndex)
@@ -2020,7 +2021,7 @@ void FAllocationsProvider::EnumerateMaxLiveAllocationsTimeline(int32 StartIndex,
 
 void FAllocationsProvider::EnumerateAllocEventsTimeline(int32 StartIndex, int32 EndIndex, TFunctionRef<void(double Time, double Duration, uint32 Value)> Callback) const
 {
-	const int32 NumPoints = Timeline.Num();
+	const int32 NumPoints = static_cast<int32>(Timeline.Num());
 	StartIndex = FMath::Max(StartIndex, 0);
 	EndIndex = FMath::Min(EndIndex + 1, NumPoints); // make it exclusive
 	if (StartIndex < EndIndex)
@@ -2054,7 +2055,7 @@ void FAllocationsProvider::EnumerateAllocEventsTimeline(int32 StartIndex, int32
 
 void FAllocationsProvider::EnumerateFreeEventsTimeline(int32 StartIndex, int32 EndIndex, TFunctionRef<void(double Time, double Duration, uint32 Value)> Callback) const
 {
-	const int32 NumPoints = Timeline.Num();
+	const int32 NumPoints = static_cast<int32>(Timeline.Num());
 	StartIndex = FMath::Max(StartIndex, 0);
 	EndIndex = FMath::Min(EndIndex + 1, NumPoints); // make it exclusive
 	if (StartIndex < EndIndex)
diff --git a/Engine/Source/Developer/TraceServices/Private/Model/AllocationsProvider.h b/Engine/Source/Developer/TraceServices/Private/Model/AllocationsProvider.h
--- a/Engine/Source/Developer/TraceServices/Private/Model/AllocationsProvider.h
+++ b/Engine/Source/Developer/TraceServices/Private/Model/AllocationsProvider.h
@@ -343,7 +343,7 @@ public:
 
 	virtual bool IsInitialized() const override { ReadAccessCheck(); return bInitialized; }
 
-	virtual uint32 GetTimelineNumPoints() const override { ReadAccessCheck(); return Timeline.Num(); }
+	virtual int32 GetTimelineNumPoints() const override { ReadAccessCheck(); return static_cast<int32>(Timeline.Num()); }
 	virtual void EnumerateRootHeaps(TFunctionRef<void(HeapId Id, const FHeapSpec&)> Callback) const override;
 	virtual void GetTimelineIndexRange(double StartTime, double EndTime, int32& StartIndex, int32& EndIndex) const override;
 	virtual void EnumerateMinTotalAllocatedMemoryTimeline(int32 StartIndex, int32 EndIndex, TFunctionRef<void(double Time, double Duration, uint64 Value)> Callback) const override;
diff --git a/Engine/Source/Developer/TraceServices/Private/Model/Bookmarks.cpp b/Engine/Source/Developer/TraceServices/Private/Model/Bookmarks.cpp
--- a/Engine/Source/Developer/TraceServices/Private/Model/Bookmarks.cpp
+++ b/Engine/Source/Developer/TraceServices/Private/Model/Bookmarks.cpp
@@ -70,16 +70,16 @@ void FBookmarkProvider::EnumerateBookmarks(double IntervalStart, double Interval
 	{
 		return;
 	}
-	uint64 FirstBookmarkIndex = Algo::LowerBoundBy(Bookmarks, IntervalStart, [](const TSharedRef<FBookmarkInternal>& B)
+	int32 FirstBookmarkIndex = Algo::LowerBoundBy(Bookmarks, IntervalStart, [](const TSharedRef<FBookmarkInternal>& B)
 	{
 		return B->Time;
 	});
-	uint64 BookmarkCount = Bookmarks.Num();
+	int32 BookmarkCount = Bookmarks.Num();
 	if (FirstBookmarkIndex >= BookmarkCount)
 	{
 		return;
 	}
-	uint64 LastBookmarkIndex = Algo::UpperBoundBy(Bookmarks, IntervalEnd, [](const TSharedRef<FBookmarkInternal>& B)
+	int32 LastBookmarkIndex = Algo::UpperBoundBy(Bookmarks, IntervalEnd, [](const TSharedRef<FBookmarkInternal>& B)
 	{
 		return B->Time;
 	});
@@ -88,7 +88,7 @@ void FBookmarkProvider::EnumerateBookmarks(double IntervalStart, double Interval
 		return;
 	}
 	--LastBookmarkIndex;
-	for (uint64 Index = FirstBookmarkIndex; Index <= LastBookmarkIndex; ++Index)
+	for (int32 Index = FirstBookmarkIndex; Index <= LastBookmarkIndex; ++Index)
 	{
 		const FBookmarkInternal& InternalBookmark = Bookmarks[Index].Get();
 		FBookmark Bookmark;
diff --git a/Engine/Source/Developer/TraceServices/Private/Model/CallstacksProvider.cpp b/Engine/Source/Developer/TraceServices/Private/Model/CallstacksProvider.cpp
--- a/Engine/Source/Developer/TraceServices/Private/Model/CallstacksProvider.cpp
+++ b/Engine/Source/Developer/TraceServices/Private/Model/CallstacksProvider.cpp
@@ -125,7 +125,7 @@ uint32 FCallstacksProvider::AddCallstackWithHash(uint64 InCallstackHash, const u
 	uint32 CallstackId;
 	{
 		FRWScopeLock WriteLock(EntriesLock, SLT_Write);
-		CallstackId = Callstacks.Num();
+		CallstackId = static_cast<uint32>(Callstacks.Num());
 		CallstackMap.Add(InCallstackHash, CallstackId);
 	}
 
diff --git a/Engine/Source/Developer/TraceServices/Private/Model/CookProfilerProvider.cpp b/Engine/Source/Developer/TraceServices/Private/Model/CookProfilerProvider.cpp
--- a/Engine/Source/Developer/TraceServices/Private/Model/CookProfilerProvider.cpp
+++ b/Engine/Source/Developer/TraceServices/Private/Model/CookProfilerProvider.cpp
@@ -88,9 +88,9 @@ uint32 FCookProfilerProvider::FindOrAddPackage(uint64 Id)
 	}
 	else
 	{
-		PackageIdToIndexMap.Add(Id, Packages.Num());
+		PackageIdToIndexMap.Add(Id, static_cast<uint32>(Packages.Num()));
 		Packages.Emplace(Id);
-		return Packages.Num() - 1;
+		return static_cast<uint32>(Packages.Num() - 1);
 	}
 }
 
diff --git a/Engine/Source/Developer/TraceServices/Private/Model/Counters.cpp b/Engine/Source/Developer/TraceServices/Private/Model/Counters.cpp
--- a/Engine/Source/Developer/TraceServices/Private/Model/Counters.cpp
+++ b/Engine/Source/Developer/TraceServices/Private/Model/Counters.cpp
@@ -41,7 +41,7 @@ static void EnumerateCounterValuesInternal(const TCounterData<CounterType>& Coun
 			{
 				if (!bFirstValue && bIncludeExternalBounds)
 				{
-					Callback(LastTime, LastValue);
+					Callback(LastTime, static_cast<EnumerationType>(LastValue));
 				}
 				bFirstEnumeratedValue = false;
 			}
@@ -49,11 +49,11 @@ static void EnumerateCounterValuesInternal(const TCounterData<CounterType>& Coun
 			{
 				if (bIncludeExternalBounds)
 				{
-					Callback(Time, CurrentValue);
+					Callback(Time, static_cast<EnumerationType>(CurrentValue));
 				}
 				break;
 			}
-			Callback(Time, CurrentValue);
+			Callback(Time, static_cast<EnumerationType>(CurrentValue));
 		}
 		LastTime = Time;
 		LastValue = CurrentValue;
diff --git a/Engine/Source/Developer/TraceServices/Private/Model/CsvProfilerProvider.cpp b/Engine/Source/Developer/TraceServices/Private/Model/CsvProfilerProvider.cpp
--- a/Engine/Source/Developer/TraceServices/Private/Model/CsvProfilerProvider.cpp
+++ b/Engine/Source/Developer/TraceServices/Private/Model/CsvProfilerProvider.cpp
@@ -19,7 +19,7 @@ const TCHAR* FCsvProfilerProvider::FTableLayout::GetColumnName(uint64 ColumnInde
 	}
 	else
 	{
-		return StatSeries[ColumnIndex - 1]->Name;
+		return StatSeries[static_cast<int32>(ColumnIndex - 1)]->Name;
 	}
 }
 
@@ -31,7 +31,7 @@ ETableColumnType FCsvProfilerProvider::FTableLayout::GetColumnType(uint64 Column
 	}
 	else
 	{
-		ECsvStatSeriesType SeriesType = StatSeries[ColumnIndex - 1]->Type;
+		ECsvStatSeriesType SeriesType = StatSeries[static_cast<int32>(ColumnIndex - 1)]->Type;
 		if (SeriesType == CsvStatSeriesType_CustomStatInt)
 		{
 			return TableColumnType_Int;
@@ -55,7 +55,7 @@ void FCsvProfilerProvider::FTableReader::NextRow()
 
 void FCsvProfilerProvider::FTableReader::SetRowIndex(uint64 RowIndex)
 {
-	CurrentRowIndex = RowIndex;
+	CurrentRowIndex = IntCastChecked<uint32>(RowIndex);
 }
 
 bool FCsvProfilerProvider::FTableReader::GetValueBool(uint64 ColumnIndex) const
@@ -69,7 +69,7 @@ const FCsvProfilerProvider::FStatSeriesValue* FCsvProfilerProvider::FTableReader
 	{
 		return nullptr;
 	}
-	const FStatSeriesValue* FindIt = Capture.StatSeries[ColumnIndex - 1]->Values.Find(Capture.StartFrame + CurrentRowIndex);
+	const FStatSeriesValue* FindIt = Capture.StatSeries[static_cast<int32>(ColumnIndex - 1)]->Values.Find(Capture.StartFrame + CurrentRowIndex);
 	return FindIt;
 }
 
@@ -152,7 +152,7 @@ void FCsvProfilerProvider::EnumerateCaptures(TFunctionRef<void(const FCaptureInf
 	}
 }
 
-void FCsvProfilerProvider::StartCapture(const TCHAR* Filename, int32 FrameNumber)
+void FCsvProfilerProvider::StartCapture(const TCHAR* Filename, uint32 FrameNumber)
 {
 	Session.WriteAccessCheck();
 
@@ -163,7 +163,7 @@ void FCsvProfilerProvider::StartCapture(const TCHAR* Filename, int32 FrameNumber
 	CurrentCapture->Filename = Filename;
 }
 
-void FCsvProfilerProvider::EndCapture(int32 FrameNumber)
+void FCsvProfilerProvider::EndCapture(uint32 FrameNumber)
 {
 	Session.WriteAccessCheck();
 
@@ -192,7 +192,7 @@ uint64 FCsvProfilerProvider::AddSeries(const TCHAR* Name, ECsvStatSeriesType Typ
 
 FCsvProfilerProvider::FStatSeriesValue& FCsvProfilerProvider::GetValueRef(uint64 SeriesHandle, uint32 FrameNumber)
 {
-	FStatSeries* Series = StatSeries[SeriesHandle];
+	FStatSeries* Series = StatSeries[static_cast<int32>(SeriesHandle)];
 	if (CurrentCapture && !Series->Captures.Contains(Captures.Num()))
 	{
 		CurrentCapture->StatSeries.Add(Series);
diff --git a/Engine/Source/Developer/TraceServices/Private/Model/CsvProfilerProviderPrivate.h b/Engine/Source/Developer/TraceServices/Private/Model/CsvProfilerProviderPrivate.h
--- a/Engine/Source/Developer/TraceServices/Private/Model/CsvProfilerProviderPrivate.h
+++ b/Engine/Source/Developer/TraceServices/Private/Model/CsvProfilerProviderPrivate.h
@@ -24,8 +24,8 @@ public:
 	virtual ~FCsvProfilerProvider();
 	virtual const IUntypedTable& GetTable(uint32 CaptureIndex) const override { return Captures[CaptureIndex]->Table; }
 	virtual void EnumerateCaptures(TFunctionRef<void(const FCaptureInfo&)> Callback) const override;
-	void StartCapture(const TCHAR* Filename, int32 FrameNumber);
-	void EndCapture(int32 FrameNumber);
+	void StartCapture(const TCHAR* Filename, uint32 FrameNumber);
+	void EndCapture(uint32 FrameNumber);
 	uint64 AddSeries(const TCHAR* Name, ECsvStatSeriesType Type);
 	void SetValue(uint64 SeriesHandle, uint32 FrameNumber, double Value);
 	void SetValue(uint64 SeriesHandle, uint32 FrameNumber, int64 Value);
@@ -45,7 +45,7 @@ private:
 
 	struct FStatSeries
 	{
-		TMap<uint64, FStatSeriesValue> Values;
+		TMap<uint32, FStatSeriesValue> Values;
 		TSet<uint32> Captures;
 		const TCHAR* Name = nullptr;
 		ECsvStatSeriesType Type;
@@ -106,7 +106,7 @@ private:
 
 		const FCapture& Capture;
 		const TMap<uint32, FEvents*>& Events;
-		uint64 CurrentRowIndex;
+		uint32 CurrentRowIndex;
 	};
 
 	class FTable
diff --git a/Engine/Source/Developer/TraceServices/Private/Model/DbgHelpResolver.cpp b/Engine/Source/Developer/TraceServices/Private/Model/DbgHelpResolver.cpp
--- a/Engine/Source/Developer/TraceServices/Private/Model/DbgHelpResolver.cpp
+++ b/Engine/Source/Developer/TraceServices/Private/Model/DbgHelpResolver.cpp
@@ -349,7 +349,7 @@ void FDbgHelpResolver::ResolveSymbol(uint64 Address, FResolvedSymbol& Target)
 		Entry->Name,
 		SymbolNameStr,
 		SymbolFileStr,
-		Line.LineNumber);
+		static_cast<uint16>(Line.LineNumber));
 	SymbolFilter.Update(Target);
 }
 
diff --git a/Engine/Source/Developer/TraceServices/Private/Model/FileActivity.cpp b/Engine/Source/Developer/TraceServices/Private/Model/FileActivity.cpp
--- a/Engine/Source/Developer/TraceServices/Private/Model/FileActivity.cpp
+++ b/Engine/Source/Developer/TraceServices/Private/Model/FileActivity.cpp
@@ -53,9 +53,9 @@ const ITable<FFileActivity>& FFileActivityProvider::GetFileActivityTable() const
 
 uint32 FFileActivityProvider::GetFileIndex(const TCHAR* Path)
 {
-	uint32 FileIndex = Files.Num();
+	uint32 FileIndex = static_cast<uint32>(Files.Num());
 	FFileInfoInternal& FileInfo = Files.PushBack();
-	FileInfo.FileInfo.Id = Files.Num() - 1;
+	FileInfo.FileInfo.Id = FileIndex;
 	FileInfo.FileInfo.Path = Session.StoreString(Path);
 	FileInfo.ActivityTimeline = MakeShared<TimelineInternal>(Session.GetLinearAllocator());
 	return FileIndex;
diff --git a/Engine/Source/Developer/TraceServices/Private/Model/Frames.cpp b/Engine/Source/Developer/TraceServices/Private/Model/Frames.cpp
--- a/Engine/Source/Developer/TraceServices/Private/Model/Frames.cpp
+++ b/Engine/Source/Developer/TraceServices/Private/Model/Frames.cpp
@@ -105,7 +105,7 @@ uint32 FFrameProvider::GetFrameNumberForTimestamp(ETraceFrameType FrameType, dou
 	}
 	else if (Timestamp >= FrameStartTimes[FrameType].Last())
 	{
-		return FrameStartTimes[FrameType].Num();
+		return static_cast<uint32>(FrameStartTimes[FrameType].Num());
 	}
 	else
 	{
diff --git a/Engine/Source/Developer/TraceServices/Private/Model/LoadTimeProfiler.cpp b/Engine/Source/Developer/TraceServices/Private/Model/LoadTimeProfiler.cpp
--- a/Engine/Source/Developer/TraceServices/Private/Model/LoadTimeProfiler.cpp
+++ b/Engine/Source/Developer/TraceServices/Private/Model/LoadTimeProfiler.cpp
@@ -399,7 +399,7 @@ FPackageInfo& FLoadTimeProfilerProvider::CreatePackage()
 {
 	Session.WriteAccessCheck();
 
-	uint32 PackageId = Packages.Num();
+	uint32 PackageId = static_cast<uint32>(Packages.Num());
 	FPackageInfo& Package = Packages.PushBack();
 	Package.Id = PackageId;
 	return Package;
@@ -523,7 +523,7 @@ FPackageExportInfo& FLoadTimeProfilerProvider::CreateExport()
 {
 	Session.WriteAccessCheck();
 
-	uint32 ExportId = Exports.Num();
+	uint32 ExportId = static_cast<uint32>(Exports.Num());
 	FPackageExportInfo& Export = Exports.PushBack();
 	Export.Id = ExportId;
 	return Export;
diff --git a/Engine/Source/Developer/TraceServices/Private/Model/Memory.cpp b/Engine/Source/Developer/TraceServices/Private/Model/Memory.cpp
--- a/Engine/Source/Developer/TraceServices/Private/Model/Memory.cpp
+++ b/Engine/Source/Developer/TraceServices/Private/Model/Memory.cpp
@@ -200,23 +200,23 @@ void FMemoryProvider::EnumerateTagSamples(FMemoryTrackerId TrackerId, FMemoryTag
 	const TArray<double>& SampleTimes = Trackers[TrackerId].SampleTimes;
 	const TPagedArray<FMemoryTagSample>& SampleValues = TagSamples->Values;
 
-	int64 IndexStart = Algo::LowerBound(SampleTimes, StartTime);
-	int64 IndexEnd = Algo::UpperBound(SampleTimes, EndTime);
+	int32 IndexStart = Algo::LowerBound(SampleTimes, StartTime);
+	int32 IndexEnd = Algo::UpperBound(SampleTimes, EndTime);
 
 	if (bIncludeRangeNeighbours)
 	{
-		IndexStart = FMath::Max(IndexStart - 1, (int64)0);
-		IndexEnd = FMath::Min(IndexEnd + 1, (int64)SampleValues.Num());
+		IndexStart = FMath::Max(IndexStart - 1, 0);
+		IndexEnd = FMath::Min(IndexEnd + 1, static_cast<int32>(SampleValues.Num()));
 	}
 
-	const bool bEnumerateLastSample = (IndexEnd == (int64)SampleValues.Num());
+	const bool bEnumerateLastSample = (IndexEnd == static_cast<uint32>(SampleValues.Num()));
 	if (bEnumerateLastSample)
 	{
 		// Skip last sample from regular enumeration.
 		--IndexEnd;
 	}
 
-	for (int64 SampleIndex = IndexStart; SampleIndex < IndexEnd; ++SampleIndex)
+	for (int32 SampleIndex = IndexStart; SampleIndex < IndexEnd; ++SampleIndex)
 	{
 		const double Time = SampleTimes[SampleIndex];
 		const double Duration = SampleTimes[SampleIndex + 1] - Time;
diff --git a/Engine/Source/Developer/TraceServices/Private/Model/MetadataProvider.cpp b/Engine/Source/Developer/TraceServices/Private/Model/MetadataProvider.cpp
--- a/Engine/Source/Developer/TraceServices/Private/Model/MetadataProvider.cpp
+++ b/Engine/Source/Developer/TraceServices/Private/Model/MetadataProvider.cpp
@@ -129,7 +129,7 @@ uint16 FMetadataProvider::RegisterMetadataType(const TCHAR* Name, const FMetadat
 	Lock.WriteAccessCheck();
 
 	check(RegisteredTypes.Num() <= MaxMetadataTypeId);
-	const uint16 Type = RegisteredTypes.Num();
+	const uint16 Type = static_cast<uint16>(RegisteredTypes.Num());
 	RegisteredTypes.EmplaceBack(Schema);
 	RegisteredTypesMap.Add(Name, Type);
 
@@ -152,8 +152,7 @@ FName FMetadataProvider::GetRegisteredMetadataName(uint16 Type) const
 {
 	Lock.ReadAccessCheck();
 
-	const int32 Index = (int32)Type;
-	const FName* Name = RegisteredTypesMap.FindKey(Index);
+	const FName* Name = RegisteredTypesMap.FindKey(Type);
 	return Name ? *Name : FName();
 }
 
@@ -241,7 +240,7 @@ void FMetadataProvider::PushScopedMetadata(uint32 ThreadId, uint16 Type, void* D
 	{
 		FMemory::Memcpy(StackEntry.StoreEntry.Value, Data, Size);
 	}
-	StackEntry.StoreEntry.Size = Size;
+	StackEntry.StoreEntry.Size = static_cast<uint16>(Size);
 	StackEntry.StoreEntry.Type = Type;
 	StackEntry.StoreIndex = InvalidMetadataStoreIndex;
 	StackEntry.PinnedId = InvalidMetadataId;
@@ -319,7 +318,7 @@ uint32 FMetadataProvider::PinAndGetId(uint32 ThreadId)
 		}
 		if (StackEntry.StoreIndex == InvalidMetadataStoreIndex)
 		{
-			StackEntry.StoreIndex = MetadataStore.Num();
+			StackEntry.StoreIndex = static_cast<uint32>(MetadataStore.Num());
 			MetadataStore.EmplaceBack(StackEntry.StoreEntry);
 		}
 	}
@@ -344,11 +343,11 @@ uint32 FMetadataProvider::PinAndGetId(uint32 ThreadId)
 			for (int32 StackIndex = FirstUnpinnedStackIndex; StackIndex < StackSize; ++StackIndex)
 			{
 				FMetadataStackEntry& StackEntry = CurrentStack[StackIndex];
-				StackEntry.PinnedId = Metadata.Num();
+				StackEntry.PinnedId = static_cast<uint32>(Metadata.Num());
 				FMetadataEntry& Entry = Metadata.PushBack();
 				Entry.StoreIndex = StackEntry.StoreIndex;
 				Entry.Type = StackEntry.StoreEntry.Type;
-				Entry.StackSize = StackIndex + 1;
+				Entry.StackSize = static_cast<uint16>(StackIndex + 1);
 			}
 
 			return CurrentStack.Top().PinnedId;
@@ -359,11 +358,11 @@ uint32 FMetadataProvider::PinAndGetId(uint32 ThreadId)
 	for (int32 StackIndex = StackSize - 1; StackIndex >= 0; --StackIndex)
 	{
 		FMetadataStackEntry& StackEntry = CurrentStack[StackIndex];
-		StackEntry.PinnedId = Metadata.Num();
+		StackEntry.PinnedId = static_cast<uint32>(Metadata.Num());
 		FMetadataEntry& Entry = Metadata.PushBack();
 		Entry.StoreIndex = StackEntry.StoreIndex;
 		Entry.Type = StackEntry.StoreEntry.Type;
-		Entry.StackSize = StackIndex + 1;
+		Entry.StackSize = static_cast<uint16>(StackIndex + 1);
 	}
 
 	return CurrentStack.Top().PinnedId;
diff --git a/Engine/Source/Developer/TraceServices/Private/Model/ModuleProvider.cpp b/Engine/Source/Developer/TraceServices/Private/Model/ModuleProvider.cpp
--- a/Engine/Source/Developer/TraceServices/Private/Model/ModuleProvider.cpp
+++ b/Engine/Source/Developer/TraceServices/Private/Model/ModuleProvider.cpp
@@ -160,7 +160,7 @@ const FResolvedSymbol* TModuleProvider<SymbolResolverType>::GetSymbol(uint64 Add
 		}
 
 		// Add a pending entry to our cache.
-		ResolvedSymbol = &SymbolCache.EmplaceBack(ESymbolQueryResult::Pending, nullptr, nullptr, nullptr, 0, EResolvedSymbolFilterStatus::Unknown);
+		ResolvedSymbol = &SymbolCache.EmplaceBack(ESymbolQueryResult::Pending, nullptr, nullptr, nullptr, (uint16)0, EResolvedSymbolFilterStatus::Unknown);
 		SymbolCacheLookup.Add(Address, ResolvedSymbol);
 		++SymbolsDiscovered;
 	}
@@ -177,7 +177,7 @@ template <typename SymbolResolverType>
 uint32 TModuleProvider<SymbolResolverType>::GetNumModules() const
 {
 	FReadScopeLock _(ModulesLock);
-	return Modules.Num();
+	return static_cast<uint32>(Modules.Num());
 }
 
 /////////////////////////////////////////////////////////////////////
@@ -339,9 +339,10 @@ void TModuleProvider<SymbolResolverType>::SaveSymbolsToCache(IAnalysisCache& Cac
 
 	// Save new symbols
 	TCachedPagedArray<FSavedSymbol, 1024> SavedSymbols(TEXT("ModuleProvider.Symbols"), Cache);
-	const uint32 NumPreviouslySavedSymbols = SavedSymbols.Num();
+	const uint32 NumPreviouslySavedSymbols = static_cast<uint32>(SavedSymbols.Num());
+	const uint32 NumSymbols = static_cast<uint32>(SymbolCache.Num());
 	uint32 NumSavedSymbols = 0;
-	for (uint32 SymbolIndex = SavedSymbols.Num(); SymbolIndex < SymbolCache.Num(); ++SymbolIndex)
+	for (uint32 SymbolIndex = NumPreviouslySavedSymbols; SymbolIndex < NumSymbols; ++SymbolIndex)
 	{
 		const FResolvedSymbol& Symbol = SymbolCache[SymbolIndex];
 		if (Symbol.GetResult() != ESymbolQueryResult::OK)
@@ -349,9 +350,9 @@ void TModuleProvider<SymbolResolverType>::SaveSymbolsToCache(IAnalysisCache& Cac
 			continue;
 		}
 		const uint64* Address = SymbolReverseLookup.Find(&Symbol);
-		const uint32 ModuleOffset = Strings.Store_GetOffset(Symbol.Module);
-		const uint32 NameOffset = Strings.Store_GetOffset(Symbol.Name);
-		const uint32 FileOffset = Strings.Store_GetOffset(Symbol.File);
+		const uint32 ModuleOffset = static_cast<uint32>(Strings.Store_GetOffset(Symbol.Module));
+		const uint32 NameOffset = static_cast<uint32>(Strings.Store_GetOffset(Symbol.Name));
+		const uint32 FileOffset = static_cast<uint32>(Strings.Store_GetOffset(Symbol.File));
 		SavedSymbols.EmplaceBack(FSavedSymbol{*Address, ModuleOffset, NameOffset, FileOffset, Symbol.Line});
 		++NumSavedSymbols;
 	}
@@ -375,13 +376,13 @@ void TModuleProvider<SymbolResolverType>::LoadSymbolsFromCache(IAnalysisCache& C
 			UE_LOG(LogTraceServices, Warning, TEXT("Found cached symbol (adress %llx) which referenced unknown string."), Symbol.Address);
 			continue;
 		}
-		FResolvedSymbol& Resolved = SymbolCache.EmplaceBack(ESymbolQueryResult::OK, Module, Name, File, Symbol.Line, EResolvedSymbolFilterStatus::Unknown);
+		FResolvedSymbol& Resolved = SymbolCache.EmplaceBack(ESymbolQueryResult::OK, Module, Name, File, static_cast<uint16>(Symbol.Line), EResolvedSymbolFilterStatus::Unknown);
 		SymbolCacheLookup.Add(Symbol.Address, &Resolved);
 	}
 	NumCachedSymbols = SymbolCacheLookup.Num();
 
 	// Update filter for all cached symbols.
-	ParallelFor(SymbolCache.Num(), [this](uint32 Index)
+	ParallelFor(static_cast<int32>(SymbolCache.Num()), [this](int32 Index)
 		{
 			SymbolFilter.Update(SymbolCache[Index]);
 		},
diff --git a/Engine/Source/Developer/TraceServices/Private/Model/NetProfilerProvider.cpp b/Engine/Source/Developer/TraceServices/Private/Model/NetProfilerProvider.cpp
--- a/Engine/Source/Developer/TraceServices/Private/Model/NetProfilerProvider.cpp
+++ b/Engine/Source/Developer/TraceServices/Private/Model/NetProfilerProvider.cpp
@@ -190,7 +190,7 @@ FNetProfilerConnectionInternal& FNetProfilerProvider::CreateConnection(uint32 Ga
 	check(GameInstance);
 
 	// Create new connection
-	uint32 ConnectionIndex = Connections.Num();
+	uint32 ConnectionIndex = static_cast<uint32>(Connections.Num());
 	FNetProfilerConnectionInternal& Connection = Connections.PushBack();
 
 	Connection.Connection.ConnectionIndex = ConnectionIndex;
@@ -220,7 +220,7 @@ FNetProfilerObjectInstance& FNetProfilerProvider::CreateObject(uint32 GameInstan
 	check(GameInstance);
 
 	FNetProfilerObjectInstance& Object = GameInstance->Objects->PushBack();
-	Object.ObjectIndex = GameInstance->Objects->Num() - 1;
+	Object.ObjectIndex = static_cast<uint32>(GameInstance->Objects->Num()) - 1;
 	++GameInstance->ObjectsChangeCount;
 
 	return Object;
@@ -267,8 +267,8 @@ void FNetProfilerProvider::EditPacketDeliveryStatus(uint32 ConnectionIndex, ENet
 	FNetProfilerConnectionData& Data = *Connections[ConnectionIndex].Data[Mode];
 
 	// try to locate packet
-	uint32 PacketCount = Data.Packets.Num();
-	for (uint32 It=0; It < PacketCount; ++It)
+	uint32 PacketCount = static_cast<uint32>(Data.Packets.Num());
+	for (uint32 It = 0; It < PacketCount; ++It)
 	{
 		const uint32 PacketIndex = PacketCount - It - 1u;
 		if (Data.Packets[PacketIndex].SequenceNumber == SequenceNumber)
@@ -397,7 +397,7 @@ uint32 FNetProfilerProvider::GetObjectCount(uint32 GameInstanceIndex) const
 
 	const FNetProfilerGameInstanceInternal& GameInstance = GameInstances[GameInstanceIndex];
 
-	return GameInstance.Objects->Num();
+	return static_cast<uint32>(GameInstance.Objects->Num());
 }
 
 void FNetProfilerProvider::ReadObjects(uint32 GameInstanceIndex, TFunctionRef<void(const FNetProfilerObjectInstance&)> Callback) const
@@ -409,7 +409,8 @@ void FNetProfilerProvider::ReadObjects(uint32 GameInstanceIndex, TFunctionRef<vo
 	const FNetProfilerGameInstanceInternal& GameInstance = GameInstances[GameInstanceIndex];
 	const auto& Objects = *GameInstance.Objects;
 
-	for (uint32 ObjectsIt = 0, ObjectsEndIt = GameInstance.Objects->Num(); ObjectsIt < ObjectsEndIt; ++ObjectsIt)
+	const uint32 ObjectsEndIt = static_cast<uint32>(GameInstance.Objects->Num());
+	for (uint32 ObjectsIt = 0; ObjectsIt < ObjectsEndIt; ++ObjectsIt)
 	{
 		Callback(Objects[ObjectsIt]);
 	}
@@ -483,7 +484,7 @@ uint32 FNetProfilerProvider::GetPacketCount(uint32 ConnectionIndex, ENetProfiler
 
 	const auto& Packets = Connections[ConnectionIndex].Data[Mode]->Packets;
 
-	return Packets.Num();
+	return static_cast<uint32>(Packets.Num());
 }
 
 void FNetProfilerProvider::EnumeratePackets(uint32 ConnectionIndex, ENetProfilerConnectionMode Mode, uint32 PacketIndexIntervalStart, uint32 PacketIndexIntervalEnd, TFunctionRef<void(const FNetProfilerPacket&)> Callback) const
@@ -494,7 +495,7 @@ void FNetProfilerProvider::EnumeratePackets(uint32 ConnectionIndex, ENetProfiler
 
 	const auto& Packets = Connections[ConnectionIndex].Data[Mode]->Packets;
 
-	const uint32 PacketCount = Packets.Num();
+	const uint32 PacketCount = static_cast<uint32>(Packets.Num());
 
 	// [PacketIndexIntervalStart, PacketIndexIntervalEnd] is an inclusive interval.
 	if (PacketCount == 0 || PacketIndexIntervalStart > PacketIndexIntervalEnd)
@@ -517,7 +518,7 @@ void FNetProfilerProvider::EnumeratePacketContentEventsByIndex(uint32 Connection
 
 	const auto& ContentEvents = Connections[ConnectionIndex].Data[Mode]->ContentEvents;
 
-	const uint32 EventCount = ContentEvents.Num();
+	const uint32 EventCount = static_cast<uint32>(ContentEvents.Num());
 
 	// [StartEventIndex, EndEventIndex] is an inclusive interval.
 	if (EventCount == 0 || StartEventIndex > EndEventIndex)
@@ -601,7 +602,7 @@ ITable<FNetProfilerAggregatedStats>* FNetProfilerProvider::CreateAggregation(uin
 	}
 
 	const auto& Packets = Connections[ConnectionIndex].Data[Mode]->Packets;
-	const uint32 PacketCount = Packets.Num();
+	const uint32 PacketCount = static_cast<uint32>(Packets.Num());
 
 	if (!ensure(PacketCount > 0))
 	{
@@ -638,7 +639,7 @@ ITable<FNetProfilerAggregatedStats>* FNetProfilerProvider::CreateAggregation(uin
 		}
 
 		// Fill in basics.
-		const uint32 InclusiveSize = ContentEvent.EndPos - ContentEvent.StartPos;
+		const uint32 InclusiveSize = static_cast<uint32>(ContentEvent.EndPos - ContentEvent.StartPos);
 		++StatsEntry->InstanceCount;
 		StatsEntry->TotalInclusive += InclusiveSize;
 		StatsEntry->MaxInclusive = FMath::Max(InclusiveSize, StatsEntry->MaxInclusive);
@@ -646,7 +647,8 @@ ITable<FNetProfilerAggregatedStats>* FNetProfilerProvider::CreateAggregation(uin
 		// Pops events from the stack. Keeps only the parent hierarchy of the current event.
 		while (Helper.StackSize > ContentEvent.Level)
 		{
-			FStackEntry& StackEntry = Helper.Stack[--Helper.StackSize]; // pop
+			--Helper.StackSize;
+			FStackEntry& StackEntry = Helper.Stack[static_cast<uint32>(Helper.StackSize)];
 
 			// Finalize exclusive for each poped event (all its children were already processed).
 			FNetProfilerAggregatedStats& Stats = Helper.AggregatedStats.FindChecked(StackEntry.EventTypeIndex);
@@ -670,7 +672,7 @@ ITable<FNetProfilerAggregatedStats>* FNetProfilerProvider::CreateAggregation(uin
 		if (ContentEvent.Level > 0U)
 		{
 			// Update parent event with the contribution from current event.
-			FStackEntry& ParentStackEntry = Helper.Stack[ContentEvent.Level - 1U];
+			FStackEntry& ParentStackEntry = Helper.Stack[static_cast<uint32>(ContentEvent.Level) - 1];
 			ParentStackEntry.ExclusiveAccumulator += InclusiveSize;
 		}
 	};
@@ -683,7 +685,8 @@ ITable<FNetProfilerAggregatedStats>* FNetProfilerProvider::CreateAggregation(uin
 		// Pops the remaining events from the stack.
 		while (Helper.StackSize > 0)
 		{
-			FStackEntry& StackEntry = Helper.Stack[--Helper.StackSize]; // pop
+			--Helper.StackSize;
+			FStackEntry& StackEntry = Helper.Stack[static_cast<uint32>(Helper.StackSize)];
 
 			// Finalize exclusive for each poped event (all its children were already processed).
 			FNetProfilerAggregatedStats& Stats = Helper.AggregatedStats.FindChecked(StackEntry.EventTypeIndex);
@@ -698,7 +701,7 @@ ITable<FNetProfilerAggregatedStats>* FNetProfilerProvider::CreateAggregation(uin
 		for (uint32 PacketIt = PacketIndexIntervalStart, PacketEndIt = FMath::Min(PacketIndexIntervalEnd, PacketCount - 1u); PacketIt <= PacketEndIt; ++PacketIt)
 		{
 			const auto& ContentEvents = Connections[ConnectionIndex].Data[Mode]->ContentEvents;
-			const uint32 EventCount = ContentEvents.Num();
+			const uint32 EventCount = static_cast<uint32>(ContentEvents.Num());
 
 			const FNetProfilerPacket& Packet = Packets[PacketIt];
 			if (Packet.EventCount > 0U)
@@ -714,7 +717,8 @@ ITable<FNetProfilerAggregatedStats>* FNetProfilerProvider::CreateAggregation(uin
 				// Pops the remaining events from the stack, for each packet.
 				while (Helper.StackSize > 0)
 				{
-					FStackEntry& StackEntry = Helper.Stack[--Helper.StackSize]; // pop
+					--Helper.StackSize;
+					FStackEntry& StackEntry = Helper.Stack[static_cast<uint32>(Helper.StackSize)];
 
 					// Finalize exclusive for each poped event (all its children were already processed).
 					FNetProfilerAggregatedStats& Stats = Helper.AggregatedStats.FindChecked(StackEntry.EventTypeIndex);
@@ -757,7 +761,7 @@ ITable<FNetProfilerAggregatedStatsCounterStats>* FNetProfilerProvider::CreateSta
 	const FNetProfilerConnectionData* ConnectionData = Connections[ConnectionIndex].Data[Mode];
 	const auto& Packets = ConnectionData->Packets;
 
-	const uint32 PacketCount = Packets.Num();
+	const uint32 PacketCount = static_cast<uint32>(Packets.Num());
 
 	if (!ensure(PacketCount > 0))
 	{
@@ -790,7 +794,7 @@ ITable<FNetProfilerAggregatedStatsCounterStats>* FNetProfilerProvider::CreateSta
 	for (uint32 PacketIt = PacketIndexIntervalStart, PacketEndIt = FMath::Min(PacketIndexIntervalEnd, PacketCount - 1u); PacketIt <= PacketEndIt; ++PacketIt)
 	{
 		const FNetProfilerPacket& Packet = Packets[PacketIt];
-		const uint32 StatsCounterCount = PacketStatsCounters.Num();
+		const uint32 StatsCounterCount = static_cast<uint32>(PacketStatsCounters.Num());
 
 		// Iterate over all StatsCounters stored for the Packet
 		for (uint32 StatsCounterIt = 0; StatsCounterIt < Packet.StatsCount; ++StatsCounterIt)
diff --git a/Engine/Source/Developer/TraceServices/Private/Model/SymslibResolver.cpp b/Engine/Source/Developer/TraceServices/Private/Model/SymslibResolver.cpp
--- a/Engine/Source/Developer/TraceServices/Private/Model/SymslibResolver.cpp
+++ b/Engine/Source/Developer/TraceServices/Private/Model/SymslibResolver.cpp
@@ -919,7 +919,7 @@ EModuleStatus FSymslibResolver::LoadModule(FModuleEntry* Entry, FStringView Sear
 
 	// unit storage
 	SYMS_U64 UnitCount = syms_group_unit_count(Group);
-	Instance->Units.SetNum(UnitCount);
+	Instance->Units.SetNum(static_cast<int32>(UnitCount));
 
 	// per-thread arena storage (at least one)
 	int32 WorkerThreadCount = FMath::Max(1, FTaskGraphInterface::Get().GetNumWorkerThreads());
@@ -934,7 +934,7 @@ EModuleStatus FSymslibResolver::LoadModule(FModuleEntry* Entry, FStringView Sear
 
 		std::atomic<uint32> LaneCount = 0;
 		syms_group_begin_multilane(Group, WorkerThreadCount);
-		ParallelFor(UnitCount, [Instance, Group, LaneSlot, &LaneCount, &SymbolCount](uint32 Index)
+		ParallelFor(static_cast<int32>(UnitCount), [Instance, Group, LaneSlot, &LaneCount, &SymbolCount](int32 Index)
 		{
 			SYMS_Arena* Arena;
 			uint32 LaneValue = uint32(reinterpret_cast<intptr_t>(FPlatformTLS::GetTlsValue(LaneSlot)));
@@ -1000,7 +1000,7 @@ EModuleStatus FSymslibResolver::LoadModule(FModuleEntry* Entry, FStringView Sear
 
 			syms_release_scratch(Scratch);
 
-			SymbolCount += ProcCount;
+			SymbolCount += static_cast<uint32>(ProcCount);
 		});
 		syms_group_end_multilane(Group);
 
@@ -1156,7 +1156,7 @@ bool FSymslibResolver::ResolveSymbol(uint64 Address, FResolvedSymbol& Target, FS
 	SYMS_UnitID UnitID = syms_spatial_map_1d_value_from_point(&Instance->UnitMap, VirtualOffset);
 	if (UnitID)
 	{
-		FSymsUnit* Unit = &Instance->Units[UnitID - 1];
+		FSymsUnit* Unit = &Instance->Units[static_cast<uint32>(UnitID) - 1];
 
 		SYMS_U64 Value = syms_spatial_map_1d_value_from_point(&Unit->ProcMap, VirtualOffset);
 		if (Value)
@@ -1223,7 +1223,7 @@ bool FSymslibResolver::ResolveSymbol(uint64 Address, FResolvedSymbol& Target, FS
 		Entry->Module->Name,
 		SymbolNamePersistent,
 		SourceFilePersistent,
-		SourceFileLine);
+		static_cast<uint16>(SourceFileLine));
 	SymbolFilter.Update(Target);
 
 	return true;
diff --git a/Engine/Source/Developer/TraceServices/Private/Model/Tables.h b/Engine/Source/Developer/TraceServices/Private/Model/Tables.h
--- a/Engine/Source/Developer/TraceServices/Private/Model/Tables.h
+++ b/Engine/Source/Developer/TraceServices/Private/Model/Tables.h
@@ -215,27 +215,27 @@ public:
 
 	const TCHAR* GetColumnName(uint64 ColumnIndex) const override
 	{
-		return *Columns[ColumnIndex].Name;
+		return *Columns[(int32)ColumnIndex].Name;
 	}
 
 	ETableColumnType GetColumnType(uint64 ColumnIndex) const override
 	{
-		return Columns[ColumnIndex].Type;
+		return Columns[(int32)ColumnIndex].Type;
 	}
 
 	void SetColumnType(uint64 ColumnIndex, ETableColumnType ColumnType)
 	{
-		Columns[ColumnIndex].Type = ColumnType;
+		Columns[(int32)ColumnIndex].Type = ColumnType;
 	}
 
 	uint32 GetColumnDisplayHintFlags(uint64 ColumnIndex) const override
 	{
-		return Columns[ColumnIndex].DisplayHintFlags;
+		return Columns[(int32)ColumnIndex].DisplayHintFlags;
 	}
 
 	FColumnValueContainer GetColumnValue(const RowType& Row, uint64 ColumnIndex) const
 	{
-		return Columns[ColumnIndex].Projector(Row);
+		return Columns[(int32)ColumnIndex].Projector(Row);
 	}
 
 private:
@@ -340,7 +340,7 @@ public:
 		case TableColumnType_Float:
 			return Layout.GetColumnValue(*CurrentRow, ColumnIndex).FloatValue;
 		case TableColumnType_Double:
-			return static_cast<double>(Layout.GetColumnValue(*CurrentRow, ColumnIndex).DoubleValue);
+			return static_cast<float>(Layout.GetColumnValue(*CurrentRow, ColumnIndex).DoubleValue);
 		}
 		return 0.0;
 	}
diff --git a/Engine/Source/Developer/TraceServices/Public/Common/CachedPagedArray.h b/Engine/Source/Developer/TraceServices/Public/Common/CachedPagedArray.h
--- a/Engine/Source/Developer/TraceServices/Public/Common/CachedPagedArray.h
+++ b/Engine/Source/Developer/TraceServices/Public/Common/CachedPagedArray.h
@@ -72,8 +72,8 @@ namespace TraceServices
 		ItemType& operator[](uint64 Index)
 		{
 			check(Index < *ElementCount);
-			const uint32 PageIndex = Index / PageSize;
-			const uint32 IndexInPage = Index % PageSize;
+			const uint32 PageIndex = static_cast<uint32>(Index / PageSize);
+			const uint32 IndexInPage = static_cast<uint32>(Index % PageSize);
 			ItemType* Page = Pages[PageIndex];
 			if (!Page)
 			{
@@ -84,7 +84,7 @@ namespace TraceServices
 
 		uint32 NumPages() const
 		{
-			return (*ElementCount + PageSize - 1) / PageSize;
+			return static_cast<uint32>((*ElementCount + PageSize - 1) / PageSize);
 		}
 
 		uint64 Num() const
diff --git a/Engine/Source/Developer/TraceServices/Public/Common/CachedStringStore.h b/Engine/Source/Developer/TraceServices/Public/Common/CachedStringStore.h
--- a/Engine/Source/Developer/TraceServices/Public/Common/CachedStringStore.h
+++ b/Engine/Source/Developer/TraceServices/Public/Common/CachedStringStore.h
@@ -85,8 +85,8 @@ public:
 	 */
 	const CharType* GetStringAtOffset(uint64 OffsetBytes) const
 	{
-		const int32 BlockIndex = OffsetBytes / BlockSize;
-		const uint32 OffsetInBlockBytes = OffsetBytes % BlockSize;
+		const int32 BlockIndex = static_cast<int32>(OffsetBytes / BlockSize);
+		const uint32 OffsetInBlockBytes = static_cast<uint32>(OffsetBytes % BlockSize);
 		if (BlockIndex >= Blocks.Num())
 		{
 			return nullptr;
diff --git a/Engine/Source/Developer/TraceServices/Public/Common/VariablePagedArray.h b/Engine/Source/Developer/TraceServices/Public/Common/VariablePagedArray.h
--- a/Engine/Source/Developer/TraceServices/Public/Common/VariablePagedArray.h
+++ b/Engine/Source/Developer/TraceServices/Public/Common/VariablePagedArray.h
@@ -502,7 +502,7 @@ public:
 		if (ItemIndex == 0)
 		{
 			// No need to split the page. Just add a new one in front of it.
-			LeftPage = &Pages.InsertDefaulted_GetRef(LeftPageIndex);
+			LeftPage = &Pages.InsertDefaulted_GetRef(static_cast<uint32>(LeftPageIndex));
 			RightPage = LeftPage + 1;
 
 			LeftPage->Items = reinterpret_cast<ItemType*>(Allocator.Allocate(PageSize * sizeof(ItemType)));
@@ -514,7 +514,7 @@ public:
 		else
 		{
 			// Add another page after current one.
-			RightPage = &Pages.InsertDefaulted_GetRef(LeftPageIndex + 1);
+			RightPage = &Pages.InsertDefaulted_GetRef(static_cast<uint32>(LeftPageIndex) + 1);
 			LeftPage = RightPage - 1;
 
 			RightPage->Items = reinterpret_cast<ItemType*>(Allocator.Allocate(PageSize * sizeof(ItemType)));
@@ -532,7 +532,7 @@ public:
 		OnInsertedItem(PageGroupIndex);
 
 		// Split the page group.
-		PageGroupType* RightPageGroup = &PageGroups.InsertDefaulted_GetRef(PageGroupIndex + 1);
+		PageGroupType* RightPageGroup = &PageGroups.InsertDefaulted_GetRef(static_cast<uint32>(PageGroupIndex) + 1);
 		FirstPageGroup = PageGroups.GetData();
 		LastPageGroup = PageGroups.GetData() + PageGroups.Num() - 1;
 		PageGroupType* LeftPageGroup = RightPageGroup - 1;
@@ -656,15 +656,15 @@ private:
 
 	ItemType& PushBackItemInPageGroup(const uint64 PageGroupIndex)
 	{
-		PageGroupType& PageGroup = PageGroups[PageGroupIndex];
+		PageGroupType& PageGroup = PageGroups[static_cast<uint32>(PageGroupIndex)];
 		uint64 PageIndex = PageGroup.FirstPageIndex + PageGroup.PageCount - 1;
-		PageType* Page = &Pages[PageIndex];
+		PageType* Page = &Pages[static_cast<uint32>(PageIndex)];
 
 		if (Page->ItemOffset + Page->ItemCount == PageSize) // if page is full
 		{
 			// Add a new page.
 			++PageIndex;
-			Page = &Pages.InsertDefaulted_GetRef(PageIndex);
+			Page = &Pages.InsertDefaulted_GetRef(static_cast<uint32>(PageIndex));
 			Page->Items = reinterpret_cast<ItemType*>(Allocator.Allocate(PageSize * sizeof(ItemType)));
 			OnInsertedPage(PageGroupIndex);
 		}
diff --git a/Engine/Source/Developer/TraceServices/Public/Model/MonotonicTimeline.h b/Engine/Source/Developer/TraceServices/Public/Model/MonotonicTimeline.h
--- a/Engine/Source/Developer/TraceServices/Public/Model/MonotonicTimeline.h
+++ b/Engine/Source/Developer/TraceServices/Public/Model/MonotonicTimeline.h
@@ -662,10 +662,11 @@ public:
 			});
 
 		TArray<TSharedRef<FAsyncTask<FEnumarateAsyncTask>>> WorkerTasks;
-		uint64 NumPages = FMath::Max(LastScopePageIndex - FirstScopePageIndex, 1ULL);
 
+		check(LastScopePageIndex >= FirstScopePageIndex);
+		uint32 NumPages = FMath::Max(static_cast<uint32>(LastScopePageIndex - FirstScopePageIndex), 1u);
 		uint32 NumThreads = GThreadPool->GetNumThreads();
-		uint32 NumTasks = FMath::Min((uint64)NumThreads, NumPages);
+		uint32 NumTasks = FMath::Min(NumThreads, NumPages);
 		uint32 PagesPerTask = NumPages / NumTasks;
 		uint32 RemainingPages = NumPages % NumTasks; // The remaining pages will be split between the first tasks
 
diff --git a/Engine/Source/Developer/TraceServices/Public/TraceServices/Model/AllocationsProvider.h b/Engine/Source/Developer/TraceServices/Public/TraceServices/Model/AllocationsProvider.h
--- a/Engine/Source/Developer/TraceServices/Public/TraceServices/Model/AllocationsProvider.h
+++ b/Engine/Source/Developer/TraceServices/Public/TraceServices/Model/AllocationsProvider.h
@@ -124,7 +124,7 @@ public:
 	virtual bool IsInitialized() const = 0;
 
 	// Returns the number of points in each timeline (Min/Max Total Allocated Memory, Min/Max Live Allocations, Total Alloc Events, Total Free Events).
-	virtual uint32 GetTimelineNumPoints() const = 0;
+	virtual int32 GetTimelineNumPoints() const = 0;
 
 	// Returns the inclusive index range [StartIndex, EndIndex] for a time range [StartTime, EndTime].
 	// Index values are in range { -1, 0, .. , N-1, N }, where N = GetTimelineNumPoints().
diff --git a/Engine/Source/Developer/TraceServices/Public/TraceServices/Model/MetadataProvider.h b/Engine/Source/Developer/TraceServices/Public/TraceServices/Model/MetadataProvider.h
--- a/Engine/Source/Developer/TraceServices/Public/TraceServices/Model/MetadataProvider.h
+++ b/Engine/Source/Developer/TraceServices/Public/TraceServices/Model/MetadataProvider.h
@@ -23,7 +23,7 @@ public:
 		WideStringPtr,
 		Reference,
 	};
-	
+
 	struct FField
 	{
 		FName			Name;
@@ -58,7 +58,7 @@ public:
 		void Finish()
 		{
 		}
-		
+
 	private:
 		FField* Buffer;
 		uint8 MaxFieldCount;
@@ -87,7 +87,7 @@ public:
 			void* Dest = Dst.GetData() + Field.Offset;
 			FMemory::Memcpy(Dest, Src, Size);
 		}
-		
+
 	private:
 		const FField* Buffer;
 		uint8 FieldCount;
@@ -109,18 +109,18 @@ public:
 		FString GetString(uint8* MetaData, uint8 Index) const
 		{
 			const FField& Field = Buffer[Index];
-			switch(Field.Type)
+			switch (Field.Type)
 			{
 			case EFieldType::Integer:
-				 {
-					switch(Field.Size)
+				{
+					switch (Field.Size)
 					{
 					case 1: return FString::FromInt(*GetValue<uint8>(MetaData, Index));
 					case 2: return FString::FromInt(*GetValue<uint16>(MetaData, Index));
 					case 4: return FString::FromInt(*GetValue<uint32>(MetaData, Index));
-					default: return FString::FromInt(*GetValue<uint64>(MetaData, Index));
+					default: return FString::Printf(TEXT("%llu"), *GetValue<uint64>(MetaData, Index));
 					}
-				 }
+				}
 			case EFieldType::SignedInteger:
 				{
 					switch (Field.Size)
@@ -128,7 +128,7 @@ public:
 					case 1: return FString::FromInt(*GetValue<int8>(MetaData, Index));
 					case 2: return FString::FromInt(*GetValue<int16>(MetaData, Index));
 					case 4: return FString::FromInt(*GetValue<int32>(MetaData, Index));
-					default: return FString::FromInt(*GetValue<int64>(MetaData, Index));
+					default: return FString::Printf(TEXT("%lli"), *GetValue<int64>(MetaData, Index));
 					}
 				}
 			case EFieldType::FloatingPoint:
@@ -146,7 +146,7 @@ public:
 				}
 			case EFieldType::Reference:
 				{
-					switch(Field.Size)
+					switch (Field.Size)
 					{
 					case 1: { const auto* Ref = GetValueAs<UE::Trace::FEventRef8>(MetaData, Index); return FString::Format(TEXT("Ref({0},{1})"), {Ref->RefTypeId, Ref->Id});}
 					case 2: { const auto* Ref = GetValueAs<UE::Trace::FEventRef16>(MetaData, Index); return FString::Format(TEXT("Ref({0},{1})"), {Ref->RefTypeId, Ref->Id});}
@@ -162,22 +162,22 @@ public:
 		FText GetText(uint8* MetaData, uint8 Index) const
 		{
 			const FField& Field = Buffer[Index];
-			switch(Field.Type)
+			switch (Field.Type)
 			{
 			case EFieldType::Integer:
-				 {
-					switch(Field.Size)
+				{
+					switch (Field.Size)
 					{
 					case 1: return FText::AsNumber(*GetValue<uint8>(MetaData, Index));
 					case 2: return FText::AsNumber(*GetValue<uint16>(MetaData, Index));
 					case 4: return FText::AsNumber(*GetValue<uint32>(MetaData, Index));
 					default: return FText::AsNumber(*GetValue<uint64>(MetaData, Index));
 					}
-				 }
-				 break;
+				}
+				break;
 			case EFieldType::SignedInteger:
 				{
-					switch(Field.Size)
+					switch (Field.Size)
 					{
 					case 1: return FText::AsNumber(*GetValue<int8>(MetaData, Index));
 					case 2: return FText::AsNumber(*GetValue<int16>(MetaData, Index));
@@ -207,7 +207,7 @@ public:
 					case 1: { const auto* Ref = GetValueAs<UE::Trace::FEventRef8>(MetaData, Index); return FText::Format(Fmt, Ref->RefTypeId, Ref->Id); }
 					case 2: { const auto* Ref = GetValueAs<UE::Trace::FEventRef16>(MetaData, Index); return FText::Format(Fmt, Ref->RefTypeId, Ref->Id); }
 					case 4: { const auto* Ref = GetValueAs<UE::Trace::FEventRef32>(MetaData, Index); return FText::Format(Fmt, Ref->RefTypeId, Ref->Id); }
-					default:  { const auto* Ref = GetValueAs<UE::Trace::FEventRef64>(MetaData, Index); return FText::Format(Fmt, Ref->RefTypeId, Ref->Id); }
+					default: { const auto* Ref = GetValueAs<UE::Trace::FEventRef64>(MetaData, Index); return FText::Format(Fmt, Ref->RefTypeId, Ref->Id); }
 					}
 				}
 			}
@@ -241,29 +241,31 @@ public:
 
 	FMetadataSchema(const FMetadataSchema& InSchema)
 		: Fields(InSchema.Fields)
-	{}
+	{
+		check(Fields.Num() < 256);
+	}
 
-	FBuilder Builder() 
+	FBuilder Builder()
 	{
-		return FBuilder(Fields.GetData(), Fields.Num());
+		return FBuilder(Fields.GetData(), static_cast<uint8>(Fields.Num()));
 	}
 
 	FWriter Writer() const
 	{
-		return FWriter(Fields.GetData(), Fields.Num());
+		return FWriter(Fields.GetData(), static_cast<uint8>(Fields.Num()));
 	}
 
 	FReader Reader() const
 	{
-		return FReader(Fields.GetData(), Fields.Num());
+		return FReader(Fields.GetData(), static_cast<uint8>(Fields.Num()));
 	}
 
 private:
 	FMetadataSchema() = delete;
-	
+
 	TArray<FField> Fields;
 };
-	
+
 ////////////////////////////////////////////////////////////////////////////////////////////////////
 
 class IMetadataProvider : public IProvider
