[Insights] Fixed text conflict from NSLOCTEXT macro for namespace in TableImportTask.cpp.

diff --git a/Engine/Source/Developer/TraceInsights/Public/Insights/Version.h b/Engine/Source/Developer/TraceInsights/Public/Insights/Version.h
--- a/Engine/Source/Developer/TraceInsights/Public/Insights/Version.h
+++ b/Engine/Source/Developer/TraceInsights/Public/Insights/Version.h
@@ -5,10 +5,10 @@
 // See https://semver.org/
 
 #define UNREAL_INSIGHTS_VERSION_MAJOR 1
-#define UNREAL_INSIGHTS_VERSION_MINOR 1
+#define UNREAL_INSIGHTS_VERSION_MINOR 2
 #define UNREAL_INSIGHTS_VERSION_PATCH 0
 
-#define UNREAL_INSIGHTS_VERSION_STRING "1.01"
+#define UNREAL_INSIGHTS_VERSION_STRING "1.02"
 
 /** Extra identifier string (like: "alpha", "beta.1", "beta.2", "test", "featureX", etc.). */
 //#define UNREAL_INSIGHTS_VERSION_ID_STRING "beta.1"
diff --git a/Engine/Source/Developer/TraceServices/Private/Model/TableImportTask.cpp b/Engine/Source/Developer/TraceServices/Private/Model/TableImportTask.cpp
--- a/Engine/Source/Developer/TraceServices/Private/Model/TableImportTask.cpp
+++ b/Engine/Source/Developer/TraceServices/Private/Model/TableImportTask.cpp
@@ -108,7 +108,7 @@ bool FTableImportTask::CreateLayout(const FString& Line)
 
 	if (Values.Num() != ColumnNames.Num())
 	{
-		AddError(LOCTEXT("ValuesMismatchMsg", "Import failed because the number of values on line 1 does not match the number of columns in the header line."));
+		AddError(LOCTEXT("ValuesMismatchMsg1", "Import failed because the number of values on line 1 does not match the number of columns in the header line."));
 		return false;
 	}
 
@@ -152,7 +152,7 @@ bool FTableImportTask::ParseData(TArray<FString>& Lines)
 
 		if (Values.Num() != ColumnNames.Num())
 		{
-			AddError(FText::Format(LOCTEXT("ValuesMismatchMsg", "Import failed because the number of values on line {0} does not match the number of columns in the header line."), LineIndex + 1));
+			AddError(FText::Format(LOCTEXT("ValuesMismatchMsg2", "Import failed because the number of values on line {0} does not match the number of columns in the header line."), LineIndex + 1));
 			return false;
 		}
 
@@ -175,7 +175,7 @@ bool FTableImportTask::ParseData(TArray<FString>& Lines)
 					Restart = true;
 					break;
 				}
-				
+
 				NewRow.SetValue(ValueIndex, FCString::Atod(Value));
 			}
 			else if (ColumnType == ETableColumnType::TableColumnType_Int)
