Added more allocator-related symbols to skip when choosing TopFunction

diff --git a/Engine/Source/Developer/TraceServices/Private/Model/ModuleProvider.cpp b/Engine/Source/Developer/TraceServices/Private/Model/ModuleProvider.cpp
--- a/Engine/Source/Developer/TraceServices/Private/Model/ModuleProvider.cpp
+++ b/Engine/Source/Developer/TraceServices/Private/Model/ModuleProvider.cpp
@@ -416,12 +416,16 @@ FResolvedSymbolFilter::FResolvedSymbolFilter()
 	IgnoreSymbolsByFunctionName.Add(TEXT("FMemory::"));
 	IgnoreSymbolsByFunctionName.Add(TEXT("FMallocWrapper::"));
 	IgnoreSymbolsByFunctionName.Add(TEXT("FMallocPoisonProxy::"));
+	IgnoreSymbolsByFunctionName.Add(TEXT("FMallocLeakDetectionProxy::"));
 	IgnoreSymbolsByFunctionName.Add(TEXT("FVirtualWinApiHooks::"));
 	IgnoreSymbolsByFunctionName.Add(TEXT("Malloc"));
 	IgnoreSymbolsByFunctionName.Add(TEXT("Realloc"));
 	IgnoreSymbolsByFunctionName.Add(TEXT("MemoryTrace_"));
 	IgnoreSymbolsByFunctionName.Add(TEXT("operator new"));
 	IgnoreSymbolsByFunctionName.Add(TEXT("std::"));
+	IgnoreSymbolsByFunctionName.Add(TEXT("FWindowsPlatformMemory::"));
+	IgnoreSymbolsByFunctionName.Add(TEXT("FCachedOSPageAllocator::"));
+	IgnoreSymbolsByFunctionName.Add(TEXT("FMallocBinned"));
 
 	IgnoreSymbolsByFilePath.Add(FRegexPattern(FString(TEXT(".*/Containers/.*"))));
 	IgnoreSymbolsByFilePath.Add(FRegexPattern(FString(TEXT(".*/ConcurrentLinearAllocator.*"))));
