[Insights] Minor refactoring in MetadataProvider code in preparation for adding support for clear/save/restore events.

diff --git a/Engine/Source/Developer/TraceServices/Private/Model/MetadataProvider.cpp b/Engine/Source/Developer/TraceServices/Private/Model/MetadataProvider.cpp
--- a/Engine/Source/Developer/TraceServices/Private/Model/MetadataProvider.cpp
+++ b/Engine/Source/Developer/TraceServices/Private/Model/MetadataProvider.cpp
@@ -90,14 +90,9 @@ FMetadataProvider::~FMetadataProvider()
 	{
 		for (const FMetadataStackEntry& StackEntry : KV.Value->CurrentStack)
 		{
-			if (StackEntry.StoreIndex == InvalidMetadataStoreIndex &&
-				StackEntry.StoreEntry.Size > MaxInlinedMetadataSize)
+			if (StackEntry.StoreIndex == InvalidMetadataStoreIndex)
 			{
-				FMemory::Free(StackEntry.StoreEntry.Ptr);
-				check(AllocationCount > 0);
-				--AllocationCount;
-				check(TotalAllocatedMemory >= StackEntry.StoreEntry.Size);
-				TotalAllocatedMemory -= StackEntry.StoreEntry.Size;
+				InternalFreeStoreEntry(StackEntry.StoreEntry);
 			}
 		}
 		delete KV.Value;
@@ -107,14 +102,7 @@ FMetadataProvider::~FMetadataProvider()
 	{
 		for (const FMetadataStoreEntry& StoreEntry : MetadataStore)
 		{
-			if (StoreEntry.Size > MaxInlinedMetadataSize)
-			{
-				FMemory::Free(StoreEntry.Ptr);
-				check(AllocationCount > 0);
-				--AllocationCount;
-				check(TotalAllocatedMemory >= StoreEntry.Size);
-				TotalAllocatedMemory -= StoreEntry.Size;
-			}
+			InternalFreeStoreEntry(StoreEntry);
 		}
 	}
 
@@ -124,45 +112,45 @@ FMetadataProvider::~FMetadataProvider()
 
 ////////////////////////////////////////////////////////////////////////////////////////////////////
 
-uint16 FMetadataProvider::RegisterMetadataType(const TCHAR* Name, const FMetadataSchema& Schema)
+uint16 FMetadataProvider::RegisterMetadataType(const TCHAR* InName, const FMetadataSchema& InSchema)
 {
 	Lock.WriteAccessCheck();
 
 	check(RegisteredTypes.Num() <= MaxMetadataTypeId);
 	const uint16 Type = static_cast<uint16>(RegisteredTypes.Num());
-	RegisteredTypes.EmplaceBack(Schema);
-	RegisteredTypesMap.Add(Name, Type);
+	RegisteredTypes.EmplaceBack(InSchema);
+	RegisteredTypesMap.Add(InName, Type);
 
 	return Type;
 }
 
 ////////////////////////////////////////////////////////////////////////////////////////////////////
 
-uint16 FMetadataProvider::GetRegisteredMetadataType(FName Name) const
+uint16 FMetadataProvider::GetRegisteredMetadataType(FName InName) const
 {
 	Lock.ReadAccessCheck();
 
-	const uint16* TypePtr = RegisteredTypesMap.Find(Name);
+	const uint16* TypePtr = RegisteredTypesMap.Find(InName);
 	return TypePtr ? *TypePtr : InvalidMetadataType;
 }
 
 ////////////////////////////////////////////////////////////////////////////////////////////////////
 
-FName FMetadataProvider::GetRegisteredMetadataName(uint16 Type) const
+FName FMetadataProvider::GetRegisteredMetadataName(uint16 InType) const
 {
 	Lock.ReadAccessCheck();
 
-	const FName* Name = RegisteredTypesMap.FindKey(Type);
+	const FName* Name = RegisteredTypesMap.FindKey(InType);
 	return Name ? *Name : FName();
 }
 
 ////////////////////////////////////////////////////////////////////////////////////////////////////
 
-const FMetadataSchema* FMetadataProvider::GetRegisteredMetadataSchema(uint16 Type) const
+const FMetadataSchema* FMetadataProvider::GetRegisteredMetadataSchema(uint16 InType) const
 {
 	Lock.ReadAccessCheck();
 
-	const int32 Index = (int32)Type;
+	const int32 Index = (int32)InType;
 	return Index < RegisteredTypes.Num() ? &RegisteredTypes[Index] : nullptr;
 }
 
@@ -170,127 +158,168 @@ const FMetadataSchema* FMetadataProvider::GetRegisteredMetadataSchema(uint16 Typ
 
 FMetadataProvider::FMetadataThread::FMetadataThread(uint32 InThreadId, ILinearAllocator& InAllocator)
 	: ThreadId(InThreadId)
+	, CurrentStack()
 	, Metadata(InAllocator, 1024)
 {
 }
 
 ////////////////////////////////////////////////////////////////////////////////////////////////////
 
-FMetadataProvider::FMetadataThread& FMetadataProvider::GetOrAddThread(uint32 ThreadId)
+FMetadataProvider::FMetadataThread& FMetadataProvider::GetOrAddThread(uint32 InThreadId)
 {
-	FMetadataThread** MetadataThreadPtr = Threads.Find(ThreadId);
+	FMetadataThread** MetadataThreadPtr = Threads.Find(InThreadId);
 	if (MetadataThreadPtr)
 	{
 		return **MetadataThreadPtr;
 	}
-	FMetadataThread* MetadataThread = new FMetadataThread(ThreadId, Session.GetLinearAllocator());
-	Threads.Add(ThreadId, MetadataThread);
+	FMetadataThread* MetadataThread = new FMetadataThread(InThreadId, Session.GetLinearAllocator());
+	Threads.Add(InThreadId, MetadataThread);
 	return *MetadataThread;
 }
 
 ////////////////////////////////////////////////////////////////////////////////////////////////////
 
-FMetadataProvider::FMetadataThread* FMetadataProvider::GetThread(uint32 ThreadId)
+FMetadataProvider::FMetadataThread* FMetadataProvider::GetThread(uint32 InThreadId)
 {
-	FMetadataThread** MetadataThreadPtr = Threads.Find(ThreadId);
+	FMetadataThread** MetadataThreadPtr = Threads.Find(InThreadId);
 	return (MetadataThreadPtr) ? *MetadataThreadPtr : nullptr;
 }
 
 ////////////////////////////////////////////////////////////////////////////////////////////////////
 
-const FMetadataProvider::FMetadataThread* FMetadataProvider::GetThread(uint32 ThreadId) const
+const FMetadataProvider::FMetadataThread* FMetadataProvider::GetThread(uint32 InThreadId) const
 {
-	FMetadataThread* const* MetadataThreadPtr = Threads.Find(ThreadId);
+	FMetadataThread* const* MetadataThreadPtr = Threads.Find(InThreadId);
 	return (MetadataThreadPtr) ? *MetadataThreadPtr : nullptr;
 }
 
 ////////////////////////////////////////////////////////////////////////////////////////////////////
 
-void FMetadataProvider::PushScopedMetadata(uint32 ThreadId, uint16 Type, void* Data, uint32 Size)
+void FMetadataProvider::InternalAllocStoreEntry(FMetadataStoreEntry& InStoreEntry, uint16 InType, const void* InData, uint32 InSize)
 {
-	Lock.WriteAccessCheck();
-
-	check(Type < RegisteredTypes.Num());
-	check(Data != nullptr && Size > 0);
-
-	if (Size > MaxMetadataSize)
+	if (InSize > MaxInlinedMetadataSize)
 	{
-		UE_LOG(LogTraceServices, Error, TEXT("[Meta] Cannot push metadata (thread %u, type %u, size %u). Data size is too large."), ThreadId, Type, Size);
-		return;
+		InStoreEntry.Ptr = FMemory::Malloc(InSize);
+		++AllocationEventCount;
+		++AllocationCount;
+		TotalAllocatedMemory += InSize;
+		FMemory::Memcpy(InStoreEntry.Ptr, InData, InSize);
 	}
+	else
+	{
+		FMemory::Memcpy(InStoreEntry.Value, InData, InSize);
+	}
+	InStoreEntry.Size = static_cast<uint16>(InSize);
+	InStoreEntry.Type = InType;
+}
 
-	FMetadataThread& MetadataThread = GetOrAddThread(ThreadId);
+////////////////////////////////////////////////////////////////////////////////////////////////////
 
-	if (MetadataThread.CurrentStack.Num() == MaxMetadataStackSize)
+void FMetadataProvider::InternalFreeStoreEntry(const FMetadataStoreEntry& InStoreEntry)
+{
+	if (InStoreEntry.Size > MaxInlinedMetadataSize)
 	{
-		UE_LOG(LogTraceServices, Error, TEXT("[Meta] Cannot push metadata (thread %u, type %u). Stack size is too large."), ThreadId, Type);
-		return;
+		FMemory::Free(InStoreEntry.Ptr);
+		check(AllocationCount > 0);
+		--AllocationCount;
+		check(TotalAllocatedMemory >= InStoreEntry.Size);
+		TotalAllocatedMemory -= InStoreEntry.Size;
 	}
+}
+
+////////////////////////////////////////////////////////////////////////////////////////////////////
 
+void FMetadataProvider::InternalPushStackEntry(FMetadataThread& InMetadataThread, uint16 InType, const void* InData, uint32 InSize)
+{
 	FMetadataStackEntry StackEntry;
-	if (Size > MaxInlinedMetadataSize)
+	InternalAllocStoreEntry(StackEntry.StoreEntry, InType, InData, InSize);
+	StackEntry.StoreIndex = InvalidMetadataStoreIndex;
+	StackEntry.PinnedId = InvalidMetadataId;
+	InMetadataThread.CurrentStack.Push(StackEntry);
+}
+
+////////////////////////////////////////////////////////////////////////////////////////////////////
+
+void FMetadataProvider::InternalPopStackEntry(FMetadataThread& InMetadataThread)
+{
+	const FMetadataStackEntry& Top = InMetadataThread.CurrentStack.Top();
+	if (Top.StoreIndex == InvalidMetadataStoreIndex)
 	{
-		StackEntry.StoreEntry.Ptr = FMemory::Malloc(Size);
-		++AllocationEventCount;
-		++AllocationCount;
-		TotalAllocatedMemory += Size;
-		FMemory::Memcpy(StackEntry.StoreEntry.Ptr, Data, Size);
+		InternalFreeStoreEntry(Top.StoreEntry);
 	}
-	else
+	InMetadataThread.CurrentStack.Pop();
+}
+
+////////////////////////////////////////////////////////////////////////////////////////////////////
+
+void FMetadataProvider::InternalClearStack(FMetadataThread& InMetadataThread)
+{
+	while (InMetadataThread.CurrentStack.Num() > 0)
 	{
-		FMemory::Memcpy(StackEntry.StoreEntry.Value, Data, Size);
+		InternalPopStackEntry(InMetadataThread);
 	}
-	StackEntry.StoreEntry.Size = static_cast<uint16>(Size);
-	StackEntry.StoreEntry.Type = Type;
-	StackEntry.StoreIndex = InvalidMetadataStoreIndex;
-	StackEntry.PinnedId = InvalidMetadataId;
+}
+
+////////////////////////////////////////////////////////////////////////////////////////////////////
 
-	MetadataThread.CurrentStack.Push(StackEntry);
+void FMetadataProvider::PushScopedMetadata(uint32 InThreadId, uint16 InType, const void* InData, uint32 InSize)
+{
+	Lock.WriteAccessCheck();
 	++EventCount;
+
+	check(InType < RegisteredTypes.Num());
+	check(InData != nullptr && InSize > 0);
+
+	if (InSize > MaxMetadataSize)
+	{
+		UE_LOG(LogTraceServices, Error, TEXT("[Meta] Cannot push metadata (thread %u, type %u, size %u). Data size is too large."), InThreadId, InType, InSize);
+		return;
+	}
+
+	FMetadataThread& MetadataThread = GetOrAddThread(InThreadId);
+
+	if (MetadataThread.CurrentStack.Num() == MaxMetadataStackSize)
+	{
+		UE_LOG(LogTraceServices, Error, TEXT("[Meta] Cannot push metadata (thread %u, type %u). Stack size is too large."), InThreadId, InType);
+		return;
+	}
+
+	InternalPushStackEntry(MetadataThread, InType, InData, InSize);
 }
 
 ////////////////////////////////////////////////////////////////////////////////////////////////////
 
-void FMetadataProvider::PopScopedMetadata(uint32 ThreadId, uint16 Type)
+void FMetadataProvider::PopScopedMetadata(uint32 InThreadId, uint16 InType)
 {
 	Lock.WriteAccessCheck();
+	++EventCount;
 
-	FMetadataThread& MetadataThread = GetOrAddThread(ThreadId);
+	FMetadataThread& MetadataThread = GetOrAddThread(InThreadId);
 
 	if (MetadataThread.CurrentStack.Num() == 0)
 	{
-		UE_LOG(LogTraceServices, Error, TEXT("[Meta] Cannot pop metadata (thread %u, type %u). Stack is empty."), ThreadId, uint32(Type));
+		UE_LOG(LogTraceServices, Error, TEXT("[Meta] Cannot pop metadata (thread %u, type %u). Stack is empty."), InThreadId, uint32(InType));
 		return;
 	}
 
 	const FMetadataStackEntry& Top = MetadataThread.CurrentStack.Top();
 
-	if (Top.StoreEntry.Type != Type)
+	if (Top.StoreEntry.Type != InType)
 	{
-		UE_LOG(LogTraceServices, Error, TEXT("[Meta] Cannot pop metadata (thread %u, type %u). Type mismatch."), ThreadId, uint32(Type));
+		UE_LOG(LogTraceServices, Error, TEXT("[Meta] Cannot pop metadata (thread %u, type %u). Type mismatch."), InThreadId, uint32(InType));
 		return;
 	}
 
-	if (Top.StoreIndex == InvalidMetadataStoreIndex &&
-		Top.StoreEntry.Size > MaxInlinedMetadataSize)
-	{
-		FMemory::Free(Top.StoreEntry.Ptr);
-		check(AllocationCount > 0);
-		--AllocationCount;
-		check(TotalAllocatedMemory >= Top.StoreEntry.Size);
-		TotalAllocatedMemory -= Top.StoreEntry.Size;
-	}
-	MetadataThread.CurrentStack.Pop();
-	++EventCount;
+	InternalPopStackEntry(MetadataThread);
 }
 
 ////////////////////////////////////////////////////////////////////////////////////////////////////
 
-uint32 FMetadataProvider::PinAndGetId(uint32 ThreadId)
+uint32 FMetadataProvider::PinAndGetId(uint32 InThreadId)
 {
 	Lock.WriteAccessCheck();
 
-	FMetadataThread* MetadataThread = GetThread(ThreadId);
+	FMetadataThread* MetadataThread = GetThread(InThreadId);
 
 	if (!MetadataThread || MetadataThread->CurrentStack.Num() == 0)
 	{
diff --git a/Engine/Source/Developer/TraceServices/Private/Model/MetadataProvider.h b/Engine/Source/Developer/TraceServices/Private/Model/MetadataProvider.h
--- a/Engine/Source/Developer/TraceServices/Private/Model/MetadataProvider.h
+++ b/Engine/Source/Developer/TraceServices/Private/Model/MetadataProvider.h
@@ -95,19 +95,19 @@ public:
 	//////////////////////////////////////////////////
 	// Edit operations
 
-	uint16 RegisterMetadataType(const TCHAR* Name, const FMetadataSchema& Schema);
+	uint16 RegisterMetadataType(const TCHAR* InName, const FMetadataSchema& InSchema);
 
-	void PushScopedMetadata(uint32 ThreadId, uint16 Type, void* Data, uint32 Size);
-	void PopScopedMetadata(uint32 ThreadId, uint16 Type);
+	void PushScopedMetadata(uint32 InThreadId, uint16 InType, const void* InData, uint32 InSize);
+	void PopScopedMetadata(uint32 InThreadId, uint16 InType);
 
 	// Pins the metadata stack and returns an id for it.
-	uint32 PinAndGetId(uint32 ThreadId);
+	uint32 PinAndGetId(uint32 InThreadId);
 
 	//////////////////////////////////////////////////
 	// Read operations
 
-	virtual uint16 GetRegisteredMetadataType(FName Name) const override;
-	virtual FName GetRegisteredMetadataName(uint16 Type) const override;
+	virtual uint16 GetRegisteredMetadataType(FName InName) const override;
+	virtual FName GetRegisteredMetadataName(uint16 InType) const override;
 	virtual const FMetadataSchema* GetRegisteredMetadataSchema(uint16) const override;
 
 	virtual uint32 GetMetadataStackSize(uint32 InThreadId, uint32 InMetadataId) const override;
@@ -115,9 +115,14 @@ public:
 	virtual void EnumerateMetadata(uint32 InThreadId, uint32 InMetadataId, TFunctionRef<bool(uint32 StackDepth, uint16 Type, const void* Data, uint32 Size)> Callback) const override;
 
 private:
-	FMetadataThread& GetOrAddThread(uint32 ThreadId);
-	FMetadataThread* GetThread(uint32 ThreadId);
-	const FMetadataThread* GetThread(uint32 ThreadId) const;
+	FMetadataThread& GetOrAddThread(uint32 InThreadId);
+	FMetadataThread* GetThread(uint32 InThreadId);
+	const FMetadataThread* GetThread(uint32 InThreadId) const;
+	void InternalAllocStoreEntry(FMetadataStoreEntry& InStoreEntry, uint16 InType, const void* InData, uint32 InSize);
+	void InternalFreeStoreEntry(const FMetadataStoreEntry& InStoreEntry);
+	void InternalPushStackEntry(FMetadataThread& InMetadataThread, uint16 InType, const void* InData, uint32 InSize);
+	void InternalPopStackEntry(FMetadataThread& InMetadataThread);
+	void InternalClearStack(FMetadataThread& InMetadataThread);
 
 private:
 	IAnalysisSession& Session;
