[Insights]
  - Metadata: Implemented analysis for "MetadataStack.ClearScope", MetadataStack.SaveStack" and "MetadataStack.RestoreStack" trace events.
  - Fixed metadata for a pinned stack of multiple entries (i.e. being added in reverse order).

diff --git a/Engine/Source/Developer/TraceServices/Private/Analyzers/MetadataAnalysis.cpp b/Engine/Source/Developer/TraceServices/Private/Analyzers/MetadataAnalysis.cpp
--- a/Engine/Source/Developer/TraceServices/Private/Analyzers/MetadataAnalysis.cpp
+++ b/Engine/Source/Developer/TraceServices/Private/Analyzers/MetadataAnalysis.cpp
@@ -12,11 +12,11 @@
 
 namespace TraceServices
 {
-	
+
 ////////////////////////////////////////////////////////////////////////////////////////////////////
 template<typename ArrayType>
 static void WriteData(ArrayType& OutData, const UE::Trace::IAnalyzer::FEventData& EventData, uint16 MetadataTypeId, const FMetadataProvider* MetadataProvider, IAnalysisSession& Session);
-	
+
 ////////////////////////////////////////////////////////////////////////////////////////////////////
 FMetadataAnalysis::FMetadataAnalysis(IAnalysisSession& InSession, FMetadataProvider* InProvider)
 	: Session(InSession)
@@ -28,7 +28,7 @@ FMetadataAnalysis::FMetadataAnalysis(IAnalysisSession& InSession, FMetadataProvi
 FMetadataAnalysis::~FMetadataAnalysis()
 {
 }
-	
+
 ////////////////////////////////////////////////////////////////////////////////////////////////////
 void FMetadataAnalysis::OnAnalysisBegin(const FOnAnalysisContext& Context)
 {
@@ -46,6 +46,8 @@ void FMetadataAnalysis::OnAnalysisBegin(const FOnAnalysisContext& Context)
 ////////////////////////////////////////////////////////////////////////////////////////////////////
 void FMetadataAnalysis::OnAnalysisEnd()
 {
+	FProviderEditScopeLock _(*MetadataProvider);
+	MetadataProvider->OnAnalysisCompleted();
 }
 
 ////////////////////////////////////////////////////////////////////////////////////////////////////
@@ -56,38 +58,37 @@ bool FMetadataAnalysis::OnEvent(uint16 RouteId, EStyle Style, const FOnEventCont
 	{
 	case RouteId_ClearScope:
 		{
-			//todo: Add method in provider to add a clear scope
+			const uint32 ThreadId = Context.ThreadInfo.GetSystemId();
 			if (Style == EStyle::EnterScope)
 			{
+				MetadataProvider->BeginClearStackScope(ThreadId);
 			}
-			else
+			else if (ensure(Style == EStyle::LeaveScope))
 			{
+				MetadataProvider->EndClearStackScope(ThreadId);
 			}
 		}
 		break;
 
 	case RouteId_SaveStack:
 		{
+			const uint32 ThreadId = Context.ThreadInfo.GetSystemId();
 			const uint32 RuntimeId = Context.EventData.GetValue<uint32>("Id");
-			// todo: This should pin into a temporary buffer which can be released on analysis end
-			//auto AnalysisId = MetadataProvider->PinAndGetId(Context.ThreadInfo.GetSystemId());
-			//todo: Store RuntimeId <-> AnalysisId in map
-			//todo: Handle wrapping?
+			MetadataProvider->SaveStack(ThreadId, RuntimeId);
 		}
 		break;
 
 	case RouteId_RestoreStack:
 		{
+			const uint32 ThreadId = Context.ThreadInfo.GetSystemId();
 			if (Style == EStyle::EnterScope)
 			{
 				const uint32 RuntimeId = Context.EventData.GetValue<uint32>("Id");
-				//todo: Find analysis id in map
+				MetadataProvider->BeginRestoreSavedStackScope(ThreadId, RuntimeId);
 			}
-			else
+			else if (ensure(Style == EStyle::LeaveScope))
 			{
-				//MetadataProvider->PopScopedMetadata(Context.ThreadInfo.GetSystemId(), 0xFF /* All */);
-				//todo: At this point the mapping can be released, we can maintain a per thread memory of which saved
-				// stack is active. There can never be two restored stacks on top of each other.
+				MetadataProvider->EndRestoreSavedStackScope(ThreadId);
 			}
 		}
 		break;
@@ -101,19 +102,19 @@ bool FMetadataAnalysis::OnEvent(uint16 RouteId, EStyle Style, const FOnEventCont
 			// for this event type and register it with the provider.
 			const uint16 MetadataTypeId = GetOrRegisterType(EventInfo);
 
+			const uint32 ThreadId = Context.ThreadInfo.GetSystemId();
+
 			// Push the actual data of the scope to the provider or pop if we are leaving the scope.
 			if (Style == EStyle::EnterScope)
 			{
 				TArray<uint8, TInlineAllocator<32>> Data;
 				WriteData(Data, EventData, MetadataTypeId, MetadataProvider, Session);
 
-				MetadataProvider->PushScopedMetadata(Context.ThreadInfo.GetSystemId(), MetadataTypeId, Data.GetData(),
-													Data.Num());
+				MetadataProvider->PushScopedMetadata(ThreadId, MetadataTypeId, Data.GetData(), Data.Num());
 			}
-			else
+			else if (ensure(Style == EStyle::LeaveScope))
 			{
-				//todo: Check that event type info is available on pop scope
-				MetadataProvider->PopScopedMetadata(Context.ThreadInfo.GetSystemId(), MetadataTypeId);
+				MetadataProvider->PopScopedMetadata(ThreadId, MetadataTypeId);
 			}
 		}
 		break;
@@ -290,6 +291,6 @@ static void WriteData(ArrayType& OutData, const UE::Trace::IAnalyzer::FEventData
 		}
 	}
 }
-	
+
 ////////////////////////////////////////////////////////////////////////////////////////////////////
 } // namespace TraceServices
diff --git a/Engine/Source/Developer/TraceServices/Private/Model/MetadataProvider.cpp b/Engine/Source/Developer/TraceServices/Private/Model/MetadataProvider.cpp
--- a/Engine/Source/Developer/TraceServices/Private/Model/MetadataProvider.cpp
+++ b/Engine/Source/Developer/TraceServices/Private/Model/MetadataProvider.cpp
@@ -4,6 +4,9 @@
 
 #include "Common/Utils.h"
 
+#define METADATA_PROVIDER_DEBUG_LOG(Thread, Fmt, ...) //{ if (Thread == 64108) { UE_LOG(LogTraceServices, Warning, Fmt, __VA_ARGS__); } }
+#define METADATA_PROVIDER_ERROR(ErrorCount, Fmt, ...) { if (++ErrorCount <= MaxLogMessagesPerErrorType) { UE_LOG(LogTraceServices, Error, Fmt, __VA_ARGS__); } }
+
 ////////////////////////////////////////////////////////////////////////////////////////////////////
 
 namespace TraceServices
@@ -160,6 +163,10 @@ FMetadataProvider::FMetadataThread::FMetadataThread(uint32 InThreadId, ILinearAl
 	: ThreadId(InThreadId)
 	, CurrentStack()
 	, Metadata(InAllocator, 1024)
+	, bIsClearStackScope(false)
+	, bIsRestoreSavedStackScope(false)
+	, RestoreSavedStackId(0)
+	, RestoreSavedStackSize(0)
 {
 }
 
@@ -236,6 +243,8 @@ void FMetadataProvider::InternalPushStackEntry(FMetadataThread& InMetadataThread
 	StackEntry.StoreIndex = InvalidMetadataStoreIndex;
 	StackEntry.PinnedId = InvalidMetadataId;
 	InMetadataThread.CurrentStack.Push(StackEntry);
+
+	METADATA_PROVIDER_DEBUG_LOG(InMetadataThread.ThreadId, TEXT("  --> push stack : size %d"), InMetadataThread.CurrentStack.Num());
 }
 
 ////////////////////////////////////////////////////////////////////////////////////////////////////
@@ -248,12 +257,16 @@ void FMetadataProvider::InternalPopStackEntry(FMetadataThread& InMetadataThread)
 		InternalFreeStoreEntry(Top.StoreEntry);
 	}
 	InMetadataThread.CurrentStack.Pop();
+
+	METADATA_PROVIDER_DEBUG_LOG(InMetadataThread.ThreadId, TEXT("  --> pop stack : size %d"), InMetadataThread.CurrentStack.Num());
 }
 
 ////////////////////////////////////////////////////////////////////////////////////////////////////
 
 void FMetadataProvider::InternalClearStack(FMetadataThread& InMetadataThread)
 {
+	METADATA_PROVIDER_DEBUG_LOG(InMetadataThread.ThreadId, TEXT("  --> clear stack of %d"), InMetadataThread.CurrentStack.Num());
+
 	while (InMetadataThread.CurrentStack.Num() > 0)
 	{
 		InternalPopStackEntry(InMetadataThread);
@@ -267,12 +280,14 @@ void FMetadataProvider::PushScopedMetadata(uint32 InThreadId, uint16 InType, con
 	Lock.WriteAccessCheck();
 	++EventCount;
 
+	METADATA_PROVIDER_DEBUG_LOG(InThreadId, TEXT("PushScopedMetadata(type=%u, size=%u)"), uint32(InType), InSize);
+
 	check(InType < RegisteredTypes.Num());
 	check(InData != nullptr && InSize > 0);
 
 	if (InSize > MaxMetadataSize)
 	{
-		UE_LOG(LogTraceServices, Error, TEXT("[Meta] Cannot push metadata (thread %u, type %u, size %u). Data size is too large."), InThreadId, InType, InSize);
+		METADATA_PROVIDER_ERROR(MetaScopeErrors, TEXT("[Meta] Cannot push metadata (thread %u, type %u, size %u). Data size is too large."), InThreadId, InType, InSize);
 		return;
 	}
 
@@ -280,7 +295,7 @@ void FMetadataProvider::PushScopedMetadata(uint32 InThreadId, uint16 InType, con
 
 	if (MetadataThread.CurrentStack.Num() == MaxMetadataStackSize)
 	{
-		UE_LOG(LogTraceServices, Error, TEXT("[Meta] Cannot push metadata (thread %u, type %u). Stack size is too large."), InThreadId, InType);
+		METADATA_PROVIDER_ERROR(MetaScopeErrors, TEXT("[Meta] Cannot push metadata (thread %u, type %u). Stack size is too large."), InThreadId, InType);
 		return;
 	}
 
@@ -294,11 +309,13 @@ void FMetadataProvider::PopScopedMetadata(uint32 InThreadId, uint16 InType)
 	Lock.WriteAccessCheck();
 	++EventCount;
 
+	METADATA_PROVIDER_DEBUG_LOG(InThreadId, TEXT("PopScopedMetadata(type=%u)"), uint32(InType));
+
 	FMetadataThread& MetadataThread = GetOrAddThread(InThreadId);
 
 	if (MetadataThread.CurrentStack.Num() == 0)
 	{
-		UE_LOG(LogTraceServices, Error, TEXT("[Meta] Cannot pop metadata (thread %u, type %u). Stack is empty."), InThreadId, uint32(InType));
+		METADATA_PROVIDER_ERROR(MetaScopeErrors, TEXT("[Meta] Cannot pop metadata (thread %u, type %u). Stack is empty."), InThreadId, uint32(InType));
 		return;
 	}
 
@@ -306,7 +323,7 @@ void FMetadataProvider::PopScopedMetadata(uint32 InThreadId, uint16 InType)
 
 	if (Top.StoreEntry.Type != InType)
 	{
-		UE_LOG(LogTraceServices, Error, TEXT("[Meta] Cannot pop metadata (thread %u, type %u). Type mismatch."), InThreadId, uint32(InType));
+		METADATA_PROVIDER_ERROR(MetaScopeErrors, TEXT("[Meta] Cannot pop metadata (thread %u, type %u). Type mismatch."), InThreadId, uint32(InType));
 		return;
 	}
 
@@ -315,6 +332,211 @@ void FMetadataProvider::PopScopedMetadata(uint32 InThreadId, uint16 InType)
 
 ////////////////////////////////////////////////////////////////////////////////////////////////////
 
+void FMetadataProvider::BeginClearStackScope(uint32 InThreadId)
+{
+	Lock.WriteAccessCheck();
+	++EventCount;
+
+	METADATA_PROVIDER_DEBUG_LOG(InThreadId, TEXT("BeginClearStackScope()"));
+
+	const uint32 ClearSavedStackId = 0;
+
+	FMetadataThread& MetadataThread = GetOrAddThread(InThreadId);
+	if (MetadataThread.bIsClearStackScope)
+	{
+		METADATA_PROVIDER_ERROR(ClearScopeErrors, TEXT("[Meta] Cannot begin clear stack (thread %u). Already in a clear scope."), InThreadId);
+		return;
+	}
+	MetadataThread.bIsClearStackScope = true;
+
+	// Save the current stack.
+	const uint32 MetadataId = PinAndGetId(InThreadId);
+	SavedStackMap.Add(ClearSavedStackId, { InThreadId, MetadataId });
+
+	if (MetadataId != InvalidMetadataId)
+	{
+		METADATA_PROVIDER_DEBUG_LOG(InThreadId, TEXT("  --> id = %u"), MetadataId);
+	}
+
+	// Clear the current stack.
+	InternalClearStack(MetadataThread);
+}
+
+////////////////////////////////////////////////////////////////////////////////////////////////////
+
+void FMetadataProvider::EndClearStackScope(uint32 InThreadId)
+{
+	Lock.WriteAccessCheck();
+	++EventCount;
+
+	METADATA_PROVIDER_DEBUG_LOG(InThreadId, TEXT("EndClearStackScope()"));
+
+	const uint32 ClearSavedStackId = 0;
+
+	FMetadataThread& MetadataThread = GetOrAddThread(InThreadId);
+	if (!MetadataThread.bIsClearStackScope)
+	{
+		METADATA_PROVIDER_ERROR(ClearScopeErrors, TEXT("[Meta] Cannot end clear stack (thread %u). Not in a clear scope."), InThreadId);
+		return;
+	}
+	MetadataThread.bIsClearStackScope = false;
+
+	FMetadataSavedStackInfo SavedStackInfo;
+	if (!SavedStackMap.RemoveAndCopyValue(ClearSavedStackId, SavedStackInfo))
+	{
+		METADATA_PROVIDER_ERROR(ClearScopeErrors, TEXT("[Meta] Cannot end clear stack (thread %u). Invalid saved stack."), InThreadId);
+		return;
+	}
+
+	if (SavedStackInfo.MetadataId == InvalidMetadataId)
+	{
+		METADATA_PROVIDER_DEBUG_LOG(InThreadId, TEXT("  --> empty"));
+
+		// Empty stack. Nothing to push.
+		return;
+	}
+
+	METADATA_PROVIDER_DEBUG_LOG(InThreadId, TEXT("  --> id = %u"), SavedStackInfo.MetadataId);
+
+	FMetadataThread* SavedMetadataThread = GetThread(SavedStackInfo.ThreadId);
+
+	if (!SavedMetadataThread || SavedStackInfo.MetadataId >= SavedMetadataThread->Metadata.Num())
+	{
+		METADATA_PROVIDER_ERROR(ClearScopeErrors, TEXT("[Meta] Cannot end clear stack (thread %u). Invalid saved metadata."), InThreadId);
+		return;
+	}
+
+	// Push the cleared saved stack.
+	auto Iterator = SavedMetadataThread->Metadata.GetIteratorFromItem(SavedStackInfo.MetadataId);
+	const FMetadataEntry* Entry = Iterator.GetCurrentItem();
+	while (true)
+	{
+		check(Entry->StoreIndex < MetadataStore.Num());
+		const FMetadataStoreEntry& StoreEntry = *MetadataStore.GetIteratorFromItem(Entry->StoreIndex);
+		const void* Data = (StoreEntry.Size > MaxInlinedMetadataSize) ? StoreEntry.Ptr : StoreEntry.Value;
+
+		if (MetadataThread.CurrentStack.Num() == MaxMetadataStackSize)
+		{
+			METADATA_PROVIDER_ERROR(ClearScopeErrors, TEXT("[Meta] Cannot end clear stack (thread %u). Stack size is too large."), InThreadId);
+			break;
+		}
+
+		InternalPushStackEntry(MetadataThread, StoreEntry.Type, Data, StoreEntry.Size);
+
+		if (Entry->StackSize == 1) // last one
+		{
+			break;
+		}
+		Entry = Iterator.PrevItem();
+	}
+}
+
+////////////////////////////////////////////////////////////////////////////////////////////////////
+
+void FMetadataProvider::SaveStack(uint32 InThreadId, uint32 InSavedStackId)
+{
+	Lock.WriteAccessCheck();
+	++EventCount;
+
+	const uint32 MetadataId = PinAndGetId(InThreadId);
+	SavedStackMap.Add(InSavedStackId, { InThreadId, MetadataId });
+}
+
+////////////////////////////////////////////////////////////////////////////////////////////////////
+
+void FMetadataProvider::BeginRestoreSavedStackScope(uint32 InThreadId, uint32 InSavedStackId)
+{
+	Lock.WriteAccessCheck();
+	++EventCount;
+
+	FMetadataThread& MetadataThread = GetOrAddThread(InThreadId);
+
+	if (MetadataThread.bIsRestoreSavedStackScope)
+	{
+		METADATA_PROVIDER_ERROR(RestoreScopeErrors, TEXT("[Meta] Cannot begin restore saved stack (thread %u, id %u). Already in a restore scope."), InThreadId, InSavedStackId);
+		return;
+	}
+	MetadataThread.bIsRestoreSavedStackScope = true;
+
+	MetadataThread.RestoreSavedStackId = InSavedStackId;
+	MetadataThread.RestoreSavedStackSize = 0;
+
+	const FMetadataSavedStackInfo* SavedStackInfo = SavedStackMap.Find(InSavedStackId);
+	if (!SavedStackInfo)
+	{
+		METADATA_PROVIDER_ERROR(RestoreScopeErrors, TEXT("[Meta] Cannot begin restore saved stack (thread %u, id %u). Invalid saved stack."), InThreadId, InSavedStackId);
+		return;
+	}
+
+	if (SavedStackInfo->MetadataId == InvalidMetadataId)
+	{
+		// Empty stack. Nothing to push.
+		return;
+	}
+
+	FMetadataThread* SavedMetadataThread = GetThread(SavedStackInfo->ThreadId);
+
+	if (!SavedMetadataThread || SavedStackInfo->MetadataId >= SavedMetadataThread->Metadata.Num())
+	{
+		METADATA_PROVIDER_ERROR(RestoreScopeErrors, TEXT("[Meta] Cannot begin restore saved stack (thread %u, id %u). Invalid saved metadata."), InThreadId, InSavedStackId);
+		return;
+	}
+
+	// Push the saved stack.
+	auto Iterator = SavedMetadataThread->Metadata.GetIteratorFromItem(SavedStackInfo->MetadataId);
+	const FMetadataEntry* Entry = Iterator.GetCurrentItem();
+	while (true)
+	{
+		check(Entry->StoreIndex < MetadataStore.Num());
+		const FMetadataStoreEntry& StoreEntry = *MetadataStore.GetIteratorFromItem(Entry->StoreIndex);
+		const void* Data = (StoreEntry.Size > MaxInlinedMetadataSize) ? StoreEntry.Ptr : StoreEntry.Value;
+
+		InternalPushStackEntry(MetadataThread, StoreEntry.Type, Data, StoreEntry.Size);
+
+		MetadataThread.RestoreSavedStackSize++;
+
+		if (Entry->StackSize == 1) // last one
+		{
+			break;
+		}
+		Entry = Iterator.PrevItem();
+	}
+}
+
+////////////////////////////////////////////////////////////////////////////////////////////////////
+
+void FMetadataProvider::EndRestoreSavedStackScope(uint32 InThreadId)
+{
+	Lock.WriteAccessCheck();
+	++EventCount;
+
+	FMetadataThread& MetadataThread = GetOrAddThread(InThreadId);
+
+	if (!MetadataThread.bIsRestoreSavedStackScope)
+	{
+		METADATA_PROVIDER_ERROR(RestoreScopeErrors, TEXT("[Meta] Cannot end restore saved stack (thread %u). Not in a restore scope."), InThreadId);
+		return;
+	}
+	MetadataThread.bIsRestoreSavedStackScope = false;
+
+	// Pop the pushed saved stack.
+	for (uint32 Index = MetadataThread.RestoreSavedStackSize; Index > 0; --Index)
+	{
+		if (MetadataThread.CurrentStack.Num() == 0)
+		{
+			METADATA_PROVIDER_ERROR(RestoreScopeErrors, TEXT("[Meta] Cannot end restore saved stack (thread %u). Invalid current stack."), InThreadId);
+			break;
+		}
+
+		InternalPopStackEntry(MetadataThread);
+	}
+
+	MetadataThread.RestoreSavedStackId = 0;
+	MetadataThread.RestoreSavedStackSize = 0;
+}
+
+////////////////////////////////////////////////////////////////////////////////////////////////////
+
 uint32 FMetadataProvider::PinAndGetId(uint32 InThreadId)
 {
 	Lock.WriteAccessCheck();
@@ -352,39 +574,31 @@ uint32 FMetadataProvider::PinAndGetId(uint32 InThreadId)
 		}
 	}
 
-	int32 FirstUnpinnedStackIndex = StackSize - 1;
-	for (int32 StackIndex = StackSize - 2; StackIndex >= 0; --StackIndex)
-	{
-		const FMetadataStackEntry& StackEntry = CurrentStack[StackIndex];
-		if (StackEntry.PinnedId != InvalidMetadataId)
-		{
-			break;
-		}
-		FirstUnpinnedStackIndex = StackIndex;
-	}
-	if (FirstUnpinnedStackIndex > 0)
+	// Check if we can reuse the last pinned stack.
+	int32 LastPinnedStackIndex = -1;
+	if (Metadata.Num() > 0)
 	{
-		const FMetadataStackEntry& LastPinnedStackEntry = CurrentStack[FirstUnpinnedStackIndex - 1];
-		if (LastPinnedStackEntry.PinnedId == Metadata.Num() - 1)
+		const uint32 LastPinnedMetadataId = static_cast<uint32>(Metadata.Num()) - 1;
+		for (int32 StackIndex = StackSize - 2; StackIndex >= 0; --StackIndex)
 		{
-			// Reuse the partial pinned metadata stack.
-			// Only pin additional metadata entries.
-			for (int32 StackIndex = FirstUnpinnedStackIndex; StackIndex < StackSize; ++StackIndex)
+			const FMetadataStackEntry& StackEntry = CurrentStack[StackIndex];
+			if (StackEntry.PinnedId != InvalidMetadataId)
 			{
-				FMetadataStackEntry& StackEntry = CurrentStack[StackIndex];
-				StackEntry.PinnedId = static_cast<uint32>(Metadata.Num());
-				FMetadataEntry& Entry = Metadata.PushBack();
-				Entry.StoreIndex = StackEntry.StoreIndex;
-				Entry.Type = StackEntry.StoreEntry.Type;
-				Entry.StackSize = static_cast<uint16>(StackIndex + 1);
+				if (StackEntry.PinnedId == LastPinnedMetadataId)
+				{
+					// We can reuse the last pinned metadata stack.
+					// In this case we pin only the additional entries.
+					LastPinnedStackIndex = StackIndex;
+				}
+				break;
 			}
-
-			return CurrentStack.Top().PinnedId;
 		}
 	}
 
-	// Pin all metadata entries in the current stack.
-	for (int32 StackIndex = StackSize - 1; StackIndex >= 0; --StackIndex)
+	// Pin metadata entries in the current stack (reusing last pinned stack if possible).
+	// It needs to be done in increasing stack index order.
+	// Each entry gives the size of the partial stack (previous entries).
+	for (int32 StackIndex = LastPinnedStackIndex + 1; StackIndex < StackSize; ++StackIndex)
 	{
 		FMetadataStackEntry& StackEntry = CurrentStack[StackIndex];
 		StackEntry.PinnedId = static_cast<uint32>(Metadata.Num());
@@ -399,6 +613,29 @@ uint32 FMetadataProvider::PinAndGetId(uint32 InThreadId)
 
 ////////////////////////////////////////////////////////////////////////////////////////////////////
 
+void FMetadataProvider::OnAnalysisCompleted()
+{
+	Lock.WriteAccessCheck();
+
+	if (MetaScopeErrors > 0)
+	{
+		UE_LOG(LogTraceServices, Error, TEXT("[Meta] Meta scope errors: %u"), MetaScopeErrors);
+	}
+	if (ClearScopeErrors > 0)
+	{
+		UE_LOG(LogTraceServices, Error, TEXT("[Meta] Clear scope errors: %u"), ClearScopeErrors);
+	}
+	if (RestoreScopeErrors > 0)
+	{
+		UE_LOG(LogTraceServices, Error, TEXT("[Meta] Restore scope errors: %u"), RestoreScopeErrors);
+	}
+
+	UE_LOG(LogTraceServices, Log, TEXT("[Meta] Analysis Completed (%llu events; %llu alloc events; %llu allocs; %llu bytes)"),
+		EventCount, AllocationEventCount, AllocationCount, TotalAllocatedMemory);
+}
+
+////////////////////////////////////////////////////////////////////////////////////////////////////
+
 uint32 FMetadataProvider::GetMetadataStackSize(uint32 InThreadId, uint32 InMetadataId) const
 {
 	Lock.ReadAccessCheck();
@@ -508,3 +745,6 @@ const IMetadataProvider* ReadMetadataProvider(const IAnalysisSession& Session)
 ////////////////////////////////////////////////////////////////////////////////////////////////////
 
 } // namespace TraceServices
+
+#undef METADATA_PROVIDER_ERROR
+#undef METADATA_PROVIDER_DEBUG_LOG
diff --git a/Engine/Source/Developer/TraceServices/Private/Model/MetadataProvider.h b/Engine/Source/Developer/TraceServices/Private/Model/MetadataProvider.h
--- a/Engine/Source/Developer/TraceServices/Private/Model/MetadataProvider.h
+++ b/Engine/Source/Developer/TraceServices/Private/Model/MetadataProvider.h
@@ -28,13 +28,14 @@ public:
 private:
 	FRWLock RWLock;
 };
-	
+
 ////////////////////////////////////////////////////////////////////////////////////////////////////
 
 class FMetadataProvider : public IMetadataProvider
 {
 public:
 	static constexpr uint32 MaxMetadataSize = 0xFFFF;
+	static constexpr uint32 MaxLogMessagesPerErrorType = 100;
 
 private:
 	static constexpr uint32 MaxInlinedMetadataSize = 12;
@@ -78,6 +79,16 @@ private:
 		uint32 ThreadId;
 		TArray<FMetadataStackEntry> CurrentStack;
 		TPagedArray<FMetadataEntry> Metadata; // a metadata id is an index in this array
+		bool bIsClearStackScope;
+		bool bIsRestoreSavedStackScope;
+		uint32 RestoreSavedStackId;
+		uint32 RestoreSavedStackSize;
+	};
+
+	struct FMetadataSavedStackInfo
+	{
+		uint32 ThreadId;
+		uint32 MetadataId;
 	};
 
 public:
@@ -100,9 +111,18 @@ public:
 	void PushScopedMetadata(uint32 InThreadId, uint16 InType, const void* InData, uint32 InSize);
 	void PopScopedMetadata(uint32 InThreadId, uint16 InType);
 
+	void BeginClearStackScope(uint32 InThreadId);
+	void EndClearStackScope(uint32 InThreadId);
+
+	void SaveStack(uint32 InThreadId, uint32 InSavedStackId);
+	void BeginRestoreSavedStackScope(uint32 InThreadId, uint32 InSavedStackId);
+	void EndRestoreSavedStackScope(uint32 InThreadId);
+
 	// Pins the metadata stack and returns an id for it.
 	uint32 PinAndGetId(uint32 InThreadId);
 
+	void OnAnalysisCompleted();
+
 	//////////////////////////////////////////////////
 	// Read operations
 
@@ -135,6 +155,12 @@ private:
 
 	TMap<uint32, FMetadataThread*> Threads;
 
+	TMap<uint32, FMetadataSavedStackInfo> SavedStackMap; // SavedStackId --> (ThreadId, MetadataId)
+
+	uint32 MetaScopeErrors = 0; // debug
+	uint32 ClearScopeErrors = 0; // debug
+	uint32 RestoreScopeErrors = 0; // debug
+
 	uint64 EventCount = 0; // debug
 	uint64 AllocationEventCount = 0; // debug
 	uint64 AllocationCount = 0; // debug
