[Insights] Added fake allocations for the free events without a matching alloc event. These allocs will have zero size and same callstack as the free event.

diff --git a/Engine/Source/Developer/TraceServices/Private/Model/AllocationsProvider.cpp b/Engine/Source/Developer/TraceServices/Private/Model/AllocationsProvider.cpp
--- a/Engine/Source/Developer/TraceServices/Private/Model/AllocationsProvider.cpp
+++ b/Engine/Source/Developer/TraceServices/Private/Model/AllocationsProvider.cpp
@@ -1433,6 +1433,19 @@ void FAllocationsProvider::EditFree(double Time, uint32 CallstackId, uint64 Addr
 	{
 		INSIGHTS_SLOW_CHECK(!AllocationPtr->IsHeap());
 	}
+	if (!AllocationPtr)
+	{
+		++FreeErrors;
+		if (FreeErrors <= MaxLogMessagesPerErrorType)
+		{
+			UE_LOG(LogTraceServices, Error, TEXT("[MemAlloc] Invalid FREE event (Address=0x%llX, RootHeap=%u, Time=%f)! A fake alloc will be created with size 0."), Address, RootHeap, Time);
+		}
+		// Fake the missing alloc.
+		constexpr uint64 FakeAllocSize = 0;
+		constexpr uint32 FakeAllocAlignment = 0;
+		EditAlloc(Time, CallstackId, Address, FakeAllocSize, FakeAllocAlignment, RootHeap);
+		AllocationPtr = LiveAllocs[RootHeap]->Remove(Address); // we take ownership of AllocationPtr
+	}
 	if (AllocationPtr)
 	{
 		check(EventIndex[RootHeap] > AllocationPtr->StartEventIndex);
