[Insights]
  - Memory Insights: Removed INSIGHTS_MEM_TRACE_LEGACY_FORMAT (support for legacy memory trace format from UE 5.0 Early Access).
  - Fixed "implicit conversion" warnings.

diff --git a/Engine/Source/Developer/TraceServices/Private/AnalysisCache.cpp b/Engine/Source/Developer/TraceServices/Private/AnalysisCache.cpp
--- a/Engine/Source/Developer/TraceServices/Private/AnalysisCache.cpp
+++ b/Engine/Source/Developer/TraceServices/Private/AnalysisCache.cpp
@@ -417,8 +417,8 @@ FAnalysisCache::~FAnalysisCache()
 	// Delete file contents wrapper
 	Contents.Reset();
 	
-	UE_LOG(LogAnalysisCache, Display, TEXT("Closing analysis cache, %0.2f Mb read, %0.2f Mb written."),
-	       Stats.BytesRead / (1024.f*1024.f), Stats.BytesWritten / (1024.f*1024.f));
+	UE_LOG(LogAnalysisCache, Display, TEXT("Closing analysis cache, %0.2f MiB read, %0.2f MiB written."),
+	       double(Stats.BytesRead) / (1024.0 * 1024.0), double(Stats.BytesWritten) / (1024.0 * 1024.0));
 }
 
 /////////////////////////////////////////////////////////////////////
diff --git a/Engine/Source/Developer/TraceServices/Private/Analyzers/AllocationsAnalysis.cpp b/Engine/Source/Developer/TraceServices/Private/Analyzers/AllocationsAnalysis.cpp
--- a/Engine/Source/Developer/TraceServices/Private/Analyzers/AllocationsAnalysis.cpp
+++ b/Engine/Source/Developer/TraceServices/Private/Analyzers/AllocationsAnalysis.cpp
@@ -15,7 +15,6 @@
 namespace TraceServices
 {
 
-#define INSIGHTS_MEM_TRACE_LEGACY_FORMAT 1 // backward compatibility with legacy memory trace format (5.0-EA)
 #define INSIGHTS_MEM_TRACE_METADATA_TEST 0
 
 ////////////////////////////////////////////////////////////////////////////////////////////////////
@@ -55,8 +54,8 @@ void FAllocationsAnalyzer::OnAnalysisBegin(const FOnAnalysisContext& Context)
 	Builder.RouteEvent(RouteId_HeapSpec,            "Memory", "HeapSpec");
 	Builder.RouteEvent(RouteId_HeapMarkAlloc,       "Memory", "HeapMarkAlloc");
 	Builder.RouteEvent(RouteId_HeapUnmarkAlloc,     "Memory", "HeapUnmarkAlloc");
-	Builder.RouteEvent(RouteId_MemScopeTag,			"Memory", "MemoryScope", true);
-	Builder.RouteEvent(RouteId_MemScopePtr,			"Memory", "MemoryScopePtr", true);
+	Builder.RouteEvent(RouteId_MemScopeTag,         "Memory", "MemoryScope", true);
+	Builder.RouteEvent(RouteId_MemScopePtr,         "Memory", "MemoryScopePtr", true);
 
 #if INSIGHTS_MEM_TRACE_METADATA_TEST
 	{
@@ -98,7 +97,6 @@ bool FAllocationsAnalyzer::OnEvent(uint16 RouteId, EStyle Style, const FOnEventC
 	{
 		case RouteId_Init:
 		{
-#if !INSIGHTS_MEM_TRACE_LEGACY_FORMAT
 			const uint8 Version = EventData.GetValue<uint8>("Version");
 
 			constexpr uint8 MinSupportedVersion = 1; // UE 5.0
@@ -107,12 +105,8 @@ bool FAllocationsAnalyzer::OnEvent(uint16 RouteId, EStyle Style, const FOnEventC
 			{
 				break;
 			}
-#endif // !INSIGHTS_MEM_TRACE_LEGACY_FORMAT
 
 			const double Time = GetCurrentTime();
-#if INSIGHTS_MEM_TRACE_LEGACY_FORMAT
-			BaseCycle = EventData.GetValue<uint64>("BaseCycle", 0);
-#endif // INSIGHTS_MEM_TRACE_LEGACY_FORMAT
 			MarkerPeriod = EventData.GetValue<uint32>("MarkerPeriod");
 
 			const uint8 MinAlignment = EventData.GetValue<uint8>("MinAlignment");
@@ -157,42 +151,17 @@ bool FAllocationsAnalyzer::OnEvent(uint16 RouteId, EStyle Style, const FOnEventC
 
 			// CallstackId is optional. If the field is not present CallstackId will be 0 (i.e. "no callstack").
 			uint32 CallstackId = EventData.GetValue<uint32>("CallstackId", 0);
-#if INSIGHTS_MEM_TRACE_LEGACY_FORMAT
-			if (!CallstackId)
-			{
-				const ICallstacksProvider* CallstacksProvider = ReadCallstacksProvider(Session);
-				if (CallstacksProvider)
-				{
-					// Legacy format of sending the hash value
-					const uint64 CallstackHash = EventData.GetValue<uint64>("Owner");
-					CallstackId = CallstacksProvider->GetCallstackIdForHash(CallstackHash);
-				}
-			}
-#endif // INSIGHTS_MEM_TRACE_LEGACY_FORMAT
 
 			uint64 Address = EventData.GetValue<uint64>("Address");
 
 			RootHeap = EventData.GetValue<uint8>("RootHeap", static_cast<uint8>(RootHeap));
 
 			uint64 SizeUpper = EventData.GetValue<uint32>("Size");
-			const uint8 SizeLowerMask = ((1 << SizeShift) - 1);
+			const uint8 SizeLowerMask = static_cast<uint8>((1u << SizeShift) - 1u);
 			const uint8 AlignmentMask = ~SizeLowerMask;
-			uint64 Size = 0;
-			uint32 Alignment = 0;
-#if INSIGHTS_MEM_TRACE_LEGACY_FORMAT
-			const uint8 Alignment_SizeLower = EventData.GetValue<uint8>("Alignment_SizeLower");
-			if (Alignment_SizeLower)
-			{
-				Size = SizeUpper << SizeShift | static_cast<uint64>(Alignment_SizeLower & SizeLowerMask);
-				Alignment = Alignment_SizeLower & AlignmentMask;
-			}
-			else
-#endif // INSIGHTS_MEM_TRACE_LEGACY_FORMAT
-			{
-				const uint8 AlignmentPow2_SizeLower = EventData.GetValue<uint8>("AlignmentPow2_SizeLower");
-				Size = SizeUpper << SizeShift | static_cast<uint64>(AlignmentPow2_SizeLower & SizeLowerMask);
-				Alignment = 1 << (AlignmentPow2_SizeLower >> SizeShift);
-			}
+			const uint8 AlignmentPow2_SizeLower = EventData.GetValue<uint8>("AlignmentPow2_SizeLower");
+			const uint64 Size = (SizeUpper << SizeShift) | static_cast<uint64>(AlignmentPow2_SizeLower & SizeLowerMask);
+			const uint32 Alignment = 1u << (AlignmentPow2_SizeLower >> SizeShift);
 
 			const uint32 TraceThreadId = Context.ThreadInfo.GetId();
 			const uint32 SystemThreadId = Context.ThreadInfo.GetSystemId();
@@ -224,17 +193,6 @@ bool FAllocationsAnalyzer::OnEvent(uint16 RouteId, EStyle Style, const FOnEventC
 			uint32 CallstackId = EventData.GetValue<uint32>("CallstackId", 0);
 
 			uint64 Address = EventData.GetValue<uint64>("Address");
-#if INSIGHTS_MEM_TRACE_LEGACY_FORMAT
-			if (Address == 0)
-			{
-				// Legacy format of sending Address and RootHeap packed togather (after 5.0-EA).
-				constexpr uint32 HeapShift = 60;
-				constexpr uint64 RootHeapMask = uint64(0xF) << HeapShift;
-				const uint64 Address_RootHeap = EventData.GetValue<uint64>("Address_RootHeap");
-				Address = Address_RootHeap & ~RootHeapMask;
-				RootHeap = (Address_RootHeap & RootHeapMask) >> HeapShift;
-			}
-#endif // INSIGHTS_MEM_TRACE_LEGACY_FORMAT
 
 			RootHeap = EventData.GetValue<uint8>("RootHeap", static_cast<uint8>(RootHeap));
 
@@ -398,6 +356,5 @@ double FAllocationsAnalyzer::GetCurrentTime() const
 ////////////////////////////////////////////////////////////////////////////////////////////////////
 
 #undef INSIGHTS_MEM_TRACE_METADATA_TEST
-#undef INSIGHTS_MEM_TRACE_LEGACY_FORMAT
 
 } // namespace TraceServices
