[Insights]
  - Counters: Added "Export Operations..." to the context menu of the Counters list view. Exports the incremental operations/values of the selected counter to a text file (tab-separated values or comma-separated values). Exports the ops/values only in the selected time region (if any) or the entire session if no time region is selected. Note that the CountersTrace API only use the "Set" operation (i.e. Add/Substract are converted to Set at runtime).

diff --git a/Engine/Source/Developer/TraceInsights/Private/Insights/ViewModels/TimingExporter.cpp b/Engine/Source/Developer/TraceInsights/Private/Insights/ViewModels/TimingExporter.cpp
--- a/Engine/Source/Developer/TraceInsights/Private/Insights/ViewModels/TimingExporter.cpp
+++ b/Engine/Source/Developer/TraceInsights/Private/Insights/ViewModels/TimingExporter.cpp
@@ -966,7 +966,56 @@ int32 FTimingExporter::ExportCounterAsText(const FString& Filename, uint32 Count
 			// Iterate the counter values.
 			if (Params.bExportOps)
 			{
-				//TODO: export counter events / operations
+				if (Counter.IsFloatingPoint())
+				{
+					Counter.EnumerateFloatOps(Params.IntervalStartTime, Params.IntervalEndTime, false, [&](double Time, TraceServices::ECounterOpType Op, double Value)
+						{
+							StringBuilder.Reset();
+							StringBuilder.Appendf(UTF8TEXT("%.9f"), Time);
+							StringBuilder.AppendChar(Separator);
+							switch (Op)
+							{
+							case TraceServices::ECounterOpType::Set:
+								StringBuilder.Append(UTF8TEXT("Set"));
+								break;
+							case TraceServices::ECounterOpType::Add:
+								StringBuilder.Append(UTF8TEXT("Add"));
+								break;
+							default:
+								StringBuilder.Appendf(UTF8TEXT("%d"), int32(Op));
+							}
+							StringBuilder.AppendChar(Separator);
+							StringBuilder.Appendf(UTF8TEXT("%.9f"), Value);
+							StringBuilder.AppendChar(LineEnd);
+							ExportFileHandle->Write((const uint8*)StringBuilder.ToString(), StringBuilder.Len() * sizeof(UTF8CHAR));
+							++ValueCount;
+						});
+				}
+				else
+				{
+					Counter.EnumerateOps(Params.IntervalStartTime, Params.IntervalEndTime, false, [&](double Time, TraceServices::ECounterOpType Op, int64 IntValue)
+						{
+							StringBuilder.Reset();
+							StringBuilder.Appendf(UTF8TEXT("%.9f"), Time);
+							StringBuilder.AppendChar(Separator);
+							switch (Op)
+							{
+							case TraceServices::ECounterOpType::Set:
+								StringBuilder.Append(UTF8TEXT("Set"));
+								break;
+							case TraceServices::ECounterOpType::Add:
+								StringBuilder.Append(UTF8TEXT("Add"));
+								break;
+							default:
+								StringBuilder.Appendf(UTF8TEXT("%d"), int32(Op));
+							}
+							StringBuilder.AppendChar(Separator);
+							StringBuilder.Appendf(UTF8TEXT("%lli"), IntValue);
+							StringBuilder.AppendChar(LineEnd);
+							ExportFileHandle->Write((const uint8*)StringBuilder.ToString(), StringBuilder.Len() * sizeof(UTF8CHAR));
+							++ValueCount;
+						});
+				}
 			}
 			else
 			{
diff --git a/Engine/Source/Developer/TraceInsights/Private/Insights/Widgets/SStatsView.cpp b/Engine/Source/Developer/TraceInsights/Private/Insights/Widgets/SStatsView.cpp
--- a/Engine/Source/Developer/TraceInsights/Private/Insights/Widgets/SStatsView.cpp
+++ b/Engine/Source/Developer/TraceInsights/Private/Insights/Widgets/SStatsView.cpp
@@ -498,7 +498,6 @@ TSharedPtr<SWidget> SStatsView::TreeView_GetMenuContent()
 			FSlateIcon(FAppStyle::Get().GetStyleSetName(), "Icons.Save")
 		);
 
-#if 0 // not implemented yet
 		MenuBuilder.AddMenuEntry
 		(
 			FStatsViewCommands::Get().Command_ExportOps,
@@ -507,7 +506,6 @@ TSharedPtr<SWidget> SStatsView::TreeView_GetMenuContent()
 			TAttribute<FText>(),
 			FSlateIcon(FAppStyle::Get().GetStyleSetName(), "Icons.Save")
 		);
-#endif
 
 		MenuBuilder.AddMenuEntry
 		(
diff --git a/Engine/Source/Developer/TraceServices/Private/Analyzers/CountersTraceAnalysis.cpp b/Engine/Source/Developer/TraceServices/Private/Analyzers/CountersTraceAnalysis.cpp
--- a/Engine/Source/Developer/TraceServices/Private/Analyzers/CountersTraceAnalysis.cpp
+++ b/Engine/Source/Developer/TraceServices/Private/Analyzers/CountersTraceAnalysis.cpp
@@ -60,9 +60,9 @@ bool FCountersAnalyzer::OnEvent(uint16 RouteId, EStyle Style, const FOnEventCont
 	}
 	case RouteId_SetValueInt:
 	{
-		uint16 CounterId = EventData.GetValue<uint16>("CounterId");
-		int64 Value = EventData.GetValue<int64>("Value");
-		double Timestamp = Context.EventTime.AsSeconds(EventData.GetValue<uint64>("Cycle"));
+		const double Timestamp = Context.EventTime.AsSeconds(EventData.GetValue<uint64>("Cycle"));
+		const int64 Value = EventData.GetValue<int64>("Value");
+		const uint16 CounterId = EventData.GetValue<uint16>("CounterId");
 		IEditableCounter* FindEditableCounter = EditableCountersMap.FindRef(CounterId);
 		if (ensure(FindEditableCounter))
 		{
@@ -72,9 +72,9 @@ bool FCountersAnalyzer::OnEvent(uint16 RouteId, EStyle Style, const FOnEventCont
 	}
 	case RouteId_SetValueFloat:
 	{
-		uint16 CounterId = EventData.GetValue<uint16>("CounterId");
-		float Value = EventData.GetValue<float>("Value");
-		double Timestamp = Context.EventTime.AsSeconds(EventData.GetValue<uint64>("Cycle"));
+		const double Timestamp = Context.EventTime.AsSeconds(EventData.GetValue<uint64>("Cycle"));
+		const float Value = EventData.GetValue<float>("Value");
+		const uint16 CounterId = EventData.GetValue<uint16>("CounterId");
 		IEditableCounter* FindEditableCounter = EditableCountersMap.FindRef(CounterId);
 		if (ensure(FindEditableCounter))
 		{
diff --git a/Engine/Source/Developer/TraceServices/Private/Model/Counters.cpp b/Engine/Source/Developer/TraceServices/Private/Model/Counters.cpp
--- a/Engine/Source/Developer/TraceServices/Private/Model/Counters.cpp
+++ b/Engine/Source/Developer/TraceServices/Private/Model/Counters.cpp
@@ -17,6 +17,12 @@ FCounter::FCounter(ILinearAllocator& Allocator, const TArray64<double>& InFrameS
 
 }
 
+void FCounter::SetIsFloatingPoint(bool bInIsFloatingPoint)
+{
+	check(!ModCount);
+	bIsFloatingPoint = bInIsFloatingPoint;
+}
+
 template<typename CounterType, typename EnumerationType>
 static void EnumerateCounterValuesInternal(const TCounterData<CounterType>& CounterData, const TArray64<double>& FrameStartTimes, bool bResetEveryFrame, double IntervalStart, double IntervalEnd, bool bIncludeExternalBounds, TFunctionRef<void(double, EnumerationType)> Callback)
 {
@@ -28,13 +34,12 @@ static void EnumerateCounterValuesInternal(const TCounterData<CounterType>& Coun
 	auto CounterIterator = bResetEveryFrame ? CounterData.GetIterator(FrameStartTimes) : CounterData.GetIterator(NoFrameStartTimes);
 	bool bFirstValue = true;
 	bool bFirstEnumeratedValue = true;
-	CounterType LastValue = CounterType();
 	double LastTime = 0.0;
+	CounterType LastValue = CounterType();
 	while (CounterIterator)
 	{
-		const TTuple<double, CounterType>& Current = *CounterIterator;
-		double Time = Current.template Get<0>();
-		CounterType CurrentValue = Current.template Get<1>();
+		const double Time = CounterIterator.GetCurrentTime();
+		const CounterType CurrentValue = CounterIterator.GetCurrentValue();
 		if (Time >= IntervalStart)
 		{
 			if (bFirstEnumeratedValue)
@@ -62,12 +67,6 @@ static void EnumerateCounterValuesInternal(const TCounterData<CounterType>& Coun
 	}
 }
 
-void FCounter::SetIsFloatingPoint(bool bInIsFloatingPoint)
-{
-	check(!ModCount);
-	bIsFloatingPoint = bInIsFloatingPoint;
-}
-
 void FCounter::EnumerateValues(double IntervalStart, double IntervalEnd, bool bIncludeExternalBounds, TFunctionRef<void(double, int64)> Callback) const
 {
 	if (bIsFloatingPoint)
@@ -92,15 +91,86 @@ void FCounter::EnumerateFloatValues(double IntervalStart, double IntervalEnd, bo
 	}
 }
 
+template<typename CounterType, typename EnumerationType>
+static void EnumerateCounterOpsInternal(const TCounterData<CounterType>& CounterData, const TArray64<double>& FrameStartTimes, bool bResetEveryFrame, double IntervalStart, double IntervalEnd, bool bIncludeExternalBounds, TFunctionRef<void(double, ECounterOpType, EnumerationType)> Callback)
+{
+	if (!CounterData.Num())
+	{
+		return;
+	}
+	TArray64<double> NoFrameStartTimes;
+	auto CounterIterator = bResetEveryFrame ? CounterData.GetIterator(FrameStartTimes) : CounterData.GetIterator(NoFrameStartTimes);
+	bool bFirstValue = true;
+	bool bFirstEnumeratedValue = true;
+	double LastTime = 0.0;
+	ECounterOpType LastOp = ECounterOpType::Set;
+	CounterType LastOpArgument = CounterType();
+	while (CounterIterator)
+	{
+		const double Time = CounterIterator.GetCurrentTime();
+		const ECounterOpType CurrentOp = CounterIterator.GetCurrentOp();
+		const CounterType CurrentOpArgument = CounterIterator.GetCurrentOpArgument();
+		if (Time >= IntervalStart)
+		{
+			if (bFirstEnumeratedValue)
+			{
+				if (!bFirstValue && bIncludeExternalBounds)
+				{
+					Callback(LastTime, LastOp, static_cast<EnumerationType>(LastOpArgument));
+				}
+				bFirstEnumeratedValue = false;
+			}
+			if (Time > IntervalEnd)
+			{
+				if (bIncludeExternalBounds)
+				{
+					Callback(Time, CurrentOp, static_cast<EnumerationType>(CurrentOpArgument));
+				}
+				break;
+			}
+			Callback(Time, CurrentOp, static_cast<EnumerationType>(CurrentOpArgument));
+		}
+		LastTime = Time;
+		LastOp = CurrentOp;
+		LastOpArgument = CurrentOpArgument;
+		bFirstValue = false;
+		++CounterIterator;
+	}
+}
+
+void FCounter::EnumerateOps(double IntervalStart, double IntervalEnd, bool bIncludeExternalBounds, TFunctionRef<void(double, ECounterOpType, int64)> Callback) const
+{
+	if (bIsFloatingPoint)
+	{
+		EnumerateCounterOpsInternal<double, int64>(DoubleCounterData, FrameStartTimes, bIsResetEveryFrame, IntervalStart, IntervalEnd, bIncludeExternalBounds, Callback);
+	}
+	else
+	{
+		EnumerateCounterOpsInternal<int64, int64>(IntCounterData, FrameStartTimes, bIsResetEveryFrame, IntervalStart, IntervalEnd, bIncludeExternalBounds, Callback);
+	}
+}
+
+void FCounter::EnumerateFloatOps(double IntervalStart, double IntervalEnd, bool bIncludeExternalBounds, TFunctionRef<void(double, ECounterOpType, double)> Callback) const
+{
+	if (bIsFloatingPoint)
+	{
+		EnumerateCounterOpsInternal<double, double>(DoubleCounterData, FrameStartTimes, bIsResetEveryFrame, IntervalStart, IntervalEnd, bIncludeExternalBounds, Callback);
+	}
+	else
+	{
+		EnumerateCounterOpsInternal<int64, double>(IntCounterData, FrameStartTimes, bIsResetEveryFrame, IntervalStart, IntervalEnd, bIncludeExternalBounds, Callback);
+	}
+}
+
 void FCounter::AddValue(double Time, int64 Value)
 {
 	if (bIsFloatingPoint)
 	{
-		DoubleCounterData.InsertOp(Time, CounterOpType_Add, double(Value));
+		DoubleCounterData.InsertOp(Time, ECounterOpType::Add, double(Value));
 	}
 	else
 	{
-		IntCounterData.InsertOp(Time, CounterOpType_Add, Value);
+		IntCounterData.InsertOp(Time, ECounterOpType::Add, Value);
 	}
 	++ModCount;
 }
@@ -109,11 +179,11 @@ void FCounter::AddValue(double Time, double Value)
 {
 	if (bIsFloatingPoint)
 	{
-		DoubleCounterData.InsertOp(Time, CounterOpType_Add, Value);
+		DoubleCounterData.InsertOp(Time, ECounterOpType::Add, Value);
 	}
 	else
 	{
-		IntCounterData.InsertOp(Time, CounterOpType_Add, int64(Value));
+		IntCounterData.InsertOp(Time, ECounterOpType::Add, int64(Value));
 	}
 	++ModCount;
 }
@@ -122,11 +192,11 @@ void FCounter::SetValue(double Time, int64 Value)
 {
 	if (bIsFloatingPoint)
 	{
-		DoubleCounterData.InsertOp(Time, CounterOpType_Set, double(Value));
+		DoubleCounterData.InsertOp(Time, ECounterOpType::Set, double(Value));
 	}
 	else
 	{
-		IntCounterData.InsertOp(Time, CounterOpType_Set, Value);
+		IntCounterData.InsertOp(Time, ECounterOpType::Set, Value);
 	}
 	++ModCount;
 }
@@ -135,11 +205,11 @@ void FCounter::SetValue(double Time, double Value)
 {
 	if (bIsFloatingPoint)
 	{
-		DoubleCounterData.InsertOp(Time, CounterOpType_Set, Value);
+		DoubleCounterData.InsertOp(Time, ECounterOpType::Set, Value);
 	}
 	else
 	{
-		IntCounterData.InsertOp(Time, CounterOpType_Set, int64(Value));
+		IntCounterData.InsertOp(Time, ECounterOpType::Set, int64(Value));
 	}
 	++ModCount;
 }
diff --git a/Engine/Source/Developer/TraceServices/Private/Model/CountersPrivate.h b/Engine/Source/Developer/TraceServices/Private/Model/CountersPrivate.h
--- a/Engine/Source/Developer/TraceServices/Private/Model/CountersPrivate.h
+++ b/Engine/Source/Developer/TraceServices/Private/Model/CountersPrivate.h
@@ -20,12 +20,6 @@
 namespace TraceServices
 {
 
-enum ECounterOpType : uint8
-{
-	CounterOpType_Add,
-	CounterOpType_Set,
-};
-
 template<typename ValueType>
 class TCounterData;
 
@@ -38,14 +32,37 @@ public:
 		, TimestampsIterator(Outer.Timestamps.GetIterator())
 		, OpTypesIterator(Outer.OpTypes.GetIterator())
 		, OpArgumentsIterator(Outer.OpArguments.GetIterator())
+		, CurrentTime(0.0)
+		, CurrentOp(ECounterOpType::Set)
+		, CurrentOpArgument(ValueType())
+		, CurrentValue(ValueType())
 	{
-		Current = MakeTuple(0.0, ValueType());
 		UpdateValue();
 	}
 
-	const TTuple<double, ValueType>& operator*() const
+	const ValueType operator*() const
+	{
+		return CurrentValue;
+	}
+
+	const double GetCurrentTime() const
+	{
+		return CurrentTime;
+	}
+
+	const ECounterOpType GetCurrentOp() const
+	{
+		return CurrentOp;
+	}
+
+	const ValueType GetCurrentOpArgument() const
+	{
+		return CurrentOpArgument;
+	}
+
+	const ValueType GetCurrentValue() const
 	{
-		return Current;
+		return CurrentValue;
 	}
 
 	explicit operator bool() const
@@ -81,21 +98,16 @@ private:
 				bIsNewFrame = true;
 				++FrameStartTimesIterator;
 			}
-			Current.template Get<0>() = *Time;
-			switch (*OpTypesIterator)
+			CurrentTime = *Time;
+			CurrentOp = bIsNewFrame ? ECounterOpType::Set : *OpTypesIterator;
+			CurrentOpArgument = *OpArgumentsIterator;
+			switch (CurrentOp)
 			{
-			case CounterOpType_Add:
-				if (bIsNewFrame)
-				{
-					Current.template Get<1>() = *OpArgumentsIterator;
-				}
-				else
-				{
-					Current.template Get<1>() += *OpArgumentsIterator;
-				}
+			case ECounterOpType::Set:
+				CurrentValue = CurrentOpArgument;
 				break;
-			case CounterOpType_Set:
-				Current.template Get<1>() = *OpArgumentsIterator;
+			case ECounterOpType::Add:
+				CurrentValue += CurrentOpArgument;
 				break;
 			}
 		}
@@ -111,7 +123,10 @@ private:
 	TPagedArray<ECounterOpType>::TIterator OpTypesIterator;
 	typename TPagedArray<ValueType>::TIterator OpArgumentsIterator;
 #endif
-	TTuple<double, ValueType> Current;
+	double CurrentTime;
+	ECounterOpType CurrentOp;
+	ValueType CurrentOpArgument;
+	ValueType CurrentValue;
 };
 
 template<typename ValueType>
@@ -209,6 +224,8 @@ public:
 	virtual void SetDisplayHint(ECounterDisplayHint InDisplayHint) override { DisplayHint = InDisplayHint; }
 	virtual void EnumerateValues(double IntervalStart, double IntervalEnd, bool bIncludeExternalBounds, TFunctionRef<void(double, int64)> Callback) const override;
 	virtual void EnumerateFloatValues(double IntervalStart, double IntervalEnd, bool bIncludeExternalBounds, TFunctionRef<void(double, double)> Callback) const override;
+	virtual void EnumerateOps(double IntervalStart, double IntervalEnd, bool bIncludeExternalBounds, TFunctionRef<void(double, ECounterOpType, int64)> Callback) const override;
+	virtual void EnumerateFloatOps(double IntervalStart, double IntervalEnd, bool bIncludeExternalBounds, TFunctionRef<void(double, ECounterOpType, double)> Callback) const override;
 	virtual void AddValue(double Time, int64 Value) override;
 	virtual void AddValue(double Time, double Value) override;
 	virtual void SetValue(double Time, int64 Value) override;
diff --git a/Engine/Source/Developer/TraceServices/Private/Model/LoadTimeProfiler.cpp b/Engine/Source/Developer/TraceServices/Private/Model/LoadTimeProfiler.cpp
--- a/Engine/Source/Developer/TraceServices/Private/Model/LoadTimeProfiler.cpp
+++ b/Engine/Source/Developer/TraceServices/Private/Model/LoadTimeProfiler.cpp
@@ -62,7 +62,7 @@ ECounterDisplayHint FLoadTimeProfilerProvider::FLoaderFrameCounter::GetDisplayHi
 
 void FLoadTimeProfilerProvider::FLoaderFrameCounter::EnumerateValues(double IntervalStart, double IntervalEnd, bool bIncludeExternalBounds, TFunctionRef<void(double, int64)> Callback) const
 {
-
+	unimplemented();
 }
 
 void FLoadTimeProfilerProvider::FLoaderFrameCounter::EnumerateFloatValues(double IntervalStart, double IntervalEnd, bool bIncludeExternalBounds, TFunctionRef<void(double, double)> Callback) const
@@ -92,6 +92,16 @@ void FLoadTimeProfilerProvider::FLoaderFrameCounter::EnumerateFloatValues(double
 	}
 }
 
+void FLoadTimeProfilerProvider::FLoaderFrameCounter::EnumerateOps(double IntervalStart, double IntervalEnd, bool bIncludeExternalBounds, TFunctionRef<void(double, ECounterOpType, int64)> Callback) const
+{
+	unimplemented();
+}
+
+void FLoadTimeProfilerProvider::FLoaderFrameCounter::EnumerateFloatOps(double IntervalStart, double IntervalEnd, bool bIncludeExternalBounds, TFunctionRef<void(double, ECounterOpType, double)> Callback) const
+{
+	unimplemented();
+}
+
 FLoadTimeProfilerProvider::FLoadTimeProfilerProvider(IAnalysisSession& InSession, IEditableCounterProvider& InEditableCounterProvider)
 	: Session(InSession)
 	, EditableCounterProvider(InEditableCounterProvider)
diff --git a/Engine/Source/Developer/TraceServices/Private/Model/LoadTimeProfilerPrivate.h b/Engine/Source/Developer/TraceServices/Private/Model/LoadTimeProfilerPrivate.h
--- a/Engine/Source/Developer/TraceServices/Private/Model/LoadTimeProfilerPrivate.h
+++ b/Engine/Source/Developer/TraceServices/Private/Model/LoadTimeProfilerPrivate.h
@@ -93,6 +93,8 @@ private:
 		virtual ECounterDisplayHint GetDisplayHint() const override;
 		virtual void EnumerateValues(double IntervalStart, double IntervalEnd, bool bIncludeExternalBounds, TFunctionRef<void(double, int64)> Callback) const override;
 		virtual void EnumerateFloatValues(double IntervalStart, double IntervalEnd, bool bIncludeExternalBounds, TFunctionRef<void(double, double)> Callback) const override;
+		virtual void EnumerateOps(double IntervalStart, double IntervalEnd, bool bIncludeExternalBounds, TFunctionRef<void(double, ECounterOpType, int64)> Callback) const override;
+		virtual void EnumerateFloatOps(double IntervalStart, double IntervalEnd, bool bIncludeExternalBounds, TFunctionRef<void(double, ECounterOpType, double)> Callback) const override;
 
 	private:
 		const TPagedArray<FLoaderFrame>& Frames;
diff --git a/Engine/Source/Developer/TraceServices/Public/TraceServices/Model/Counters.h b/Engine/Source/Developer/TraceServices/Public/TraceServices/Model/Counters.h
--- a/Engine/Source/Developer/TraceServices/Public/TraceServices/Model/Counters.h
+++ b/Engine/Source/Developer/TraceServices/Public/TraceServices/Model/Counters.h
@@ -17,6 +17,12 @@ enum ECounterDisplayHint
 	CounterDisplayHint_Memory,
 };
 
+enum class ECounterOpType : uint8
+{
+	Set = 0,
+	Add = 1,
+};
+
 class ICounter
 {
 public:
@@ -30,6 +36,8 @@ public:
 	virtual ECounterDisplayHint GetDisplayHint() const = 0;
 	virtual void EnumerateValues(double IntervalStart, double IntervalEnd, bool bIncludeExternalBounds, TFunctionRef<void(double, int64)> Callback) const = 0;
 	virtual void EnumerateFloatValues(double IntervalStart, double IntervalEnd, bool bIncludeExternalBounds, TFunctionRef<void(double, double)> Callback) const = 0;
+	virtual void EnumerateOps(double IntervalStart, double IntervalEnd, bool bIncludeExternalBounds, TFunctionRef<void(double, ECounterOpType, int64)> Callback) const = 0;
+	virtual void EnumerateFloatOps(double IntervalStart, double IntervalEnd, bool bIncludeExternalBounds, TFunctionRef<void(double, ECounterOpType, double)> Callback) const = 0;
 };
 
 class IEditableCounterProvider;
