Insights nullptr protection in order to avoid a crash I've been getting opening recent insight traces.

diff --git a/Engine/Source/Developer/TraceServices/Private/Common/FormatArgs.cpp b/Engine/Source/Developer/TraceServices/Private/Common/FormatArgs.cpp
--- a/Engine/Source/Developer/TraceServices/Private/Common/FormatArgs.cpp
+++ b/Engine/Source/Developer/TraceServices/Private/Common/FormatArgs.cpp
@@ -162,6 +162,13 @@ const TCHAR* FFormatArgsHelper::ExtractNextFormatArg(const TCHAR* FormatString,
 
 void FFormatArgsHelper::InitArgumentStream(FFormatArgsStreamContext& Context, const uint8* ArgumentsData)
 {
+	if (ArgumentsData == nullptr)
+	{
+		Context.ArgumentTypeCategory = 0;
+		Context.ArgumentTypeSize = 0;
+		return;
+	}
+
 	Context.ArgumentCount = *ArgumentsData++;
 	Context.DescriptorPtr = ArgumentsData;
 	Context.PayloadPtr = ArgumentsData + Context.ArgumentCount;
