AssetLoadingInsights:
Fixed recursive post loads not being tracked properly
Fixed time calculations for recursive scopes
Added support for load time tracing in the editor
Added time selection filtering to the requests tab

diff --git a/Engine/Source/Developer/TraceInsights/Private/Insights/LoadingProfiler/Widgets/SLoadingProfilerWindow.cpp b/Engine/Source/Developer/TraceInsights/Private/Insights/LoadingProfiler/Widgets/SLoadingProfilerWindow.cpp
--- a/Engine/Source/Developer/TraceInsights/Private/Insights/LoadingProfiler/Widgets/SLoadingProfilerWindow.cpp
+++ b/Engine/Source/Developer/TraceInsights/Private/Insights/LoadingProfiler/Widgets/SLoadingProfilerWindow.cpp
@@ -197,20 +197,6 @@ void SLoadingProfilerWindow::UpdateExportDetailsTreeView()
 
 ////////////////////////////////////////////////////////////////////////////////////////////////////
 
-class FProxyUntypedTable : public  TraceServices::IUntypedTable
-{
-public:
-	FProxyUntypedTable(const TraceServices::IUntypedTable* InTable) : TablePtr(InTable) {}
-	virtual ~FProxyUntypedTable() = default;
-
-	virtual const TraceServices::ITableLayout& GetLayout() const { return TablePtr->GetLayout(); }
-	virtual uint64 GetRowCount() const { return TablePtr->GetRowCount(); }
-	virtual TraceServices::IUntypedTableReader* CreateReader() const { return TablePtr->CreateReader(); }
-
-private:
-	const TraceServices::IUntypedTable* TablePtr;
-};
-
 void SLoadingProfilerWindow::UpdateRequestsTreeView()
 {
 	if (RequestsTreeView)
@@ -221,8 +207,8 @@ void SLoadingProfilerWindow::UpdateRequestsTreeView()
 			TraceServices::FAnalysisSessionReadScope SessionReadScope(*Session.Get());
 			const TraceServices::ILoadTimeProfilerProvider& LoadTimeProfilerProvider = *TraceServices::ReadLoadTimeProfilerProvider(*Session.Get());
 
-			const TraceServices::ITable<TraceServices::FLoadRequest>& RequestsTable = LoadTimeProfilerProvider.GetRequestsTable();
-			RequestsTreeView->UpdateSourceTable(MakeShared<FProxyUntypedTable>(&RequestsTable));
+			TraceServices::ITable<TraceServices::FRequestsTableRow>* RequestsTable = LoadTimeProfilerProvider.CreateRequestsTable(SelectionStartTime, SelectionEndTime);
+			RequestsTreeView->UpdateSourceTable(MakeShareable(RequestsTable));
 		}
 		else
 		{
@@ -542,7 +528,7 @@ TSharedRef<FTabManager::FLayout> SLoadingProfilerWindow::CreateDefaultTabLayout(
 				->AddTab(FLoadingProfilerTabs::PackageDetailsTreeViewID, ETabState::OpenedTab)
 				->AddTab(FLoadingProfilerTabs::ExportDetailsTreeViewID, ETabState::OpenedTab)
 				->AddTab(FLoadingProfilerTabs::RequestsTreeViewID, ETabState::OpenedTab)
-				->SetForegroundTab(FLoadingProfilerTabs::EventAggregationTreeViewID)
+				->SetForegroundTab(FLoadingProfilerTabs::PackageDetailsTreeViewID)
 			)
 		);
 }
diff --git a/Engine/Source/Developer/TraceInsights/Private/Insights/ViewModels/LoadingTimingTrack.cpp b/Engine/Source/Developer/TraceInsights/Private/Insights/ViewModels/LoadingTimingTrack.cpp
--- a/Engine/Source/Developer/TraceInsights/Private/Insights/ViewModels/LoadingTimingTrack.cpp
+++ b/Engine/Source/Developer/TraceInsights/Private/Insights/ViewModels/LoadingTimingTrack.cpp
@@ -80,7 +80,7 @@ void FLoadingSharedState::OnBeginSession(Insights::ITimingViewSession& InSession
 
 	LoadTimeProfilerTimelineCount = 0;
 
-	SetColorSchema(3);
+	SetColorSchema(1);
 }
 
 ////////////////////////////////////////////////////////////////////////////////////////////////////
@@ -199,7 +199,14 @@ void FLoadingSharedState::SetAllLoadingTracksToggle(bool bOnOff)
 
 const TCHAR* FLoadingSharedState::GetEventNameByEventType(uint32 Depth, const TraceServices::FLoadTimeProfilerCpuEvent& Event) const
 {
-	return TraceServices::GetLoadTimeProfilerObjectEventTypeString(Event.EventType);
+	if (Event.Export)
+	{
+		return TraceServices::GetLoadTimeProfilerObjectEventTypeString(Event.EventType);
+	}
+	else
+	{
+		return TEXT("ProcessPackageHeader");
+	}
 }
 
 ////////////////////////////////////////////////////////////////////////////////////////////////////
@@ -288,10 +295,13 @@ void FLoadingTimingTrack::BuildDrawState(ITimingEventsTrackDrawStateBuilder& Bui
 				const double SecondsPerPixel = 1.0 / Viewport.GetScaleX();
 				Timeline.EnumerateEventsDownSampled(Viewport.GetStartTime(), Viewport.GetEndTime(), SecondsPerPixel, [this, &Builder](double StartTime, double EndTime, uint32 Depth, const TraceServices::FLoadTimeProfilerCpuEvent& Event)
 				{
-					const TCHAR* Name = SharedState.GetEventName(Depth, Event);
-					const uint64 Type = static_cast<uint64>(Event.EventType);
-					const uint32 Color = 0;
-					Builder.AddEvent(StartTime, EndTime, Depth, Name, Type, Color);
+					if (Event.Package)
+					{
+						const TCHAR* Name = SharedState.GetEventName(Depth, Event);
+						const uint64 Type = static_cast<uint64>(Event.EventType);
+						const uint32 Color = 0;
+						Builder.AddEvent(StartTime, EndTime, Depth, Name, Type, Color);
+					}
 					return TraceServices::EEventEnumerate::Continue;
 				});
 			}
@@ -299,10 +309,13 @@ void FLoadingTimingTrack::BuildDrawState(ITimingEventsTrackDrawStateBuilder& Bui
 			{
 				Timeline.EnumerateEvents(Viewport.GetStartTime(), Viewport.GetEndTime(), [this, &Builder](double StartTime, double EndTime, uint32 Depth, const TraceServices::FLoadTimeProfilerCpuEvent& Event)
 				{
-					const TCHAR* Name = SharedState.GetEventName(Depth, Event);
-					const uint64 Type = static_cast<uint64>(Event.EventType);
-					const uint32 Color = 0;
-					Builder.AddEvent(StartTime, EndTime, Depth, Name, Type, Color);
+					if (Event.Package)
+					{
+						const TCHAR* Name = SharedState.GetEventName(Depth, Event);
+						const uint64 Type = static_cast<uint64>(Event.EventType);
+						const uint32 Color = 0;
+						Builder.AddEvent(StartTime, EndTime, Depth, Name, Type, Color);
+					}
 					return TraceServices::EEventEnumerate::Continue;
 				});
 			}
@@ -332,8 +345,8 @@ void FLoadingTimingTrack::InitTooltip(FTooltipDrawState& InOutTooltip, const ITi
 
 			InOutTooltip.AddTitle(SharedState.GetEventName(TooltipEvent.GetDepth(), InFoundEvent));
 
-			const TraceServices::FPackageInfo* Package = InFoundEvent.Package;
 			const TraceServices::FPackageExportInfo* Export = InFoundEvent.Export;
+			const TraceServices::FPackageInfo* Package = InFoundEvent.Export ? InFoundEvent.Export->Package : InFoundEvent.Package;
 
 			InOutTooltip.AddNameValueTextLine(TEXT("Duration:"), TimeUtils::FormatTimeAuto(TooltipEvent.GetDuration()));
 			InOutTooltip.AddNameValueTextLine(TEXT("Depth:"), FString::Printf(TEXT("%d"), TooltipEvent.GetDepth()));
@@ -343,12 +356,15 @@ void FLoadingTimingTrack::InitTooltip(FTooltipDrawState& InOutTooltip, const ITi
 				InOutTooltip.AddNameValueTextLine(TEXT("Package Name:"), Package->Name);
 				InOutTooltip.AddNameValueTextLine(TEXT("Header Size:"), FString::Printf(TEXT("%s bytes"), *FText::AsNumber(Package->Summary.TotalHeaderSize).ToString()));
 				InOutTooltip.AddNameValueTextLine(TEXT("Package Summary:"), FString::Printf(TEXT("%d imports, %d exports"), Package->Summary.ImportCount, Package->Summary.ExportCount));
+				if (!Export)
+				{
+					InOutTooltip.AddNameValueTextLine(TEXT("Event:"), TEXT("ProcessPackageSummary"));
+				}
 			}
 
-			InOutTooltip.AddNameValueTextLine(TEXT("Export Event:"), FString::Printf(TEXT("%s"), TraceServices::GetLoadTimeProfilerObjectEventTypeString(InFoundEvent.EventType)));
-
 			if (Export)
 			{
+				InOutTooltip.AddNameValueTextLine(TEXT("Event:"), FString::Printf(TEXT("%s"), TraceServices::GetLoadTimeProfilerObjectEventTypeString(InFoundEvent.EventType)));
 				InOutTooltip.AddNameValueTextLine(TEXT("Export Class:"), Export->Class ? Export->Class->Name : TEXT("N/A"));
 				InOutTooltip.AddNameValueTextLine(TEXT("Serial Size:"), FString::Printf(TEXT("%s bytes"), *FText::AsNumber(Export->SerialSize).ToString()));
 			}
@@ -394,8 +410,15 @@ bool FLoadingTimingTrack::FindLoadTimeProfilerCpuEvent(const FTimingEventSearchP
 					{
 						Timeline.EnumerateEvents(InContext.GetParameters().StartTime, InContext.GetParameters().EndTime, [&InContext](double EventStartTime, double EventEndTime, uint32 EventDepth, const TraceServices::FLoadTimeProfilerCpuEvent& Event)
 						{
-							InContext.Check(EventStartTime, EventEndTime, EventDepth, Event);
-							return InContext.ShouldContinueSearching() ? TraceServices::EEventEnumerate::Continue : TraceServices::EEventEnumerate::Stop;
+							if (Event.Package)
+							{
+								InContext.Check(EventStartTime, EventEndTime, EventDepth, Event);
+								return InContext.ShouldContinueSearching() ? TraceServices::EEventEnumerate::Continue : TraceServices::EEventEnumerate::Stop;
+							}
+							else
+							{
+								return TraceServices::EEventEnumerate::Continue;
+							}
 						});
 					});
 				}
diff --git a/Engine/Source/Developer/TraceServices/Private/Analyzers/LoadTimeTraceAnalysis.cpp b/Engine/Source/Developer/TraceServices/Private/Analyzers/LoadTimeTraceAnalysis.cpp
--- a/Engine/Source/Developer/TraceServices/Private/Analyzers/LoadTimeTraceAnalysis.cpp
+++ b/Engine/Source/Developer/TraceServices/Private/Analyzers/LoadTimeTraceAnalysis.cpp
@@ -34,9 +34,12 @@
 	{
 		delete KV.Value;
 	}
+	for (const auto& KV : ActiveLinkersMap)
+	{
+		delete KV.Value;
+	}
 }
 
-// It must be called under the edit session scope (write session lock)!
 FAsyncLoadingTraceAnalyzer::FThreadState& FAsyncLoadingTraceAnalyzer::GetThreadState(uint32 ThreadId)
 {
 	FThreadState* ThreadState = ThreadStatesMap.FindRef(ThreadId);
@@ -44,6 +47,7 @@
 	{
 		ThreadState = new FThreadState();
 		ThreadStatesMap.Add(ThreadId, ThreadState);
+		FAnalysisSessionEditScope _(Session);
 		ThreadState->CpuTimeline = &LoadTimeProfilerProvider.EditCpuTimeline(ThreadId);
 	}
 	return *ThreadState;
@@ -70,21 +74,26 @@
 	Builder.RouteEvent(RouteId_SuspendAsyncLoading, "LoadTime", "SuspendAsyncLoading");
 	Builder.RouteEvent(RouteId_ResumeAsyncLoading, "LoadTime", "ResumeAsyncLoading");
 	Builder.RouteEvent(RouteId_PackageSummary, "LoadTime", "PackageSummary");
+	Builder.RouteEvent(RouteId_BeginProcessSummary, "LoadTime", "BeginProcessSummary");
+	Builder.RouteEvent(RouteId_EndProcessSummary, "LoadTime", "EndProcessSummary");
 	Builder.RouteEvent(RouteId_BeginCreateExport, "LoadTime", "BeginCreateExport");
 	Builder.RouteEvent(RouteId_EndCreateExport, "LoadTime", "EndCreateExport");
 	Builder.RouteEvent(RouteId_BeginSerializeExport, "LoadTime", "BeginSerializeExport");
 	Builder.RouteEvent(RouteId_EndSerializeExport, "LoadTime", "EndSerializeExport");
-	Builder.RouteEvent(RouteId_BeginPostLoadExport, "LoadTime", "BeginPostLoadExport");
-	Builder.RouteEvent(RouteId_EndPostLoadExport, "LoadTime", "EndPostLoadExport");
+	Builder.RouteEvent(RouteId_BeginPostLoad, "LoadTime", "BeginPostLoad");
+	Builder.RouteEvent(RouteId_EndPostLoad, "LoadTime", "EndPostLoad");
+	Builder.RouteEvent(RouteId_BeginPostLoadObject, "LoadTime", "BeginPostLoadObject");
+	Builder.RouteEvent(RouteId_EndPostLoadObject, "LoadTime", "EndPostLoadObject");
 	Builder.RouteEvent(RouteId_NewAsyncPackage, "LoadTime", "NewAsyncPackage");
-	Builder.RouteEvent(RouteId_BeginLoadAsyncPackage, "LoadTime", "BeginLoadAsyncPackage");
-	Builder.RouteEvent(RouteId_EndLoadAsyncPackage, "LoadTime", "EndLoadAsyncPackage");
 	Builder.RouteEvent(RouteId_DestroyAsyncPackage, "LoadTime", "DestroyAsyncPackage");
+	Builder.RouteEvent(RouteId_NewLinker, "LoadTime", "NewLinker");
+	Builder.RouteEvent(RouteId_DestroyLinker, "LoadTime", "DestroyLinker");
 	Builder.RouteEvent(RouteId_BeginRequest, "LoadTime", "BeginRequest");
 	Builder.RouteEvent(RouteId_EndRequest, "LoadTime", "EndRequest");
 	Builder.RouteEvent(RouteId_BeginRequestGroup, "LoadTime", "BeginRequestGroup");
 	Builder.RouteEvent(RouteId_EndRequestGroup, "LoadTime", "EndRequestGroup");
 	Builder.RouteEvent(RouteId_AsyncPackageRequestAssociation, "LoadTime", "AsyncPackageRequestAssociation");
+	Builder.RouteEvent(RouteId_AsyncPackageLinkerAssociation, "LoadTime", "AsyncPackageLinkerAssociation");
 	Builder.RouteEvent(RouteId_AsyncPackageImportDependency, "LoadTime", "AsyncPackageImportDependency");
 	Builder.RouteEvent(RouteId_ClassInfo, "LoadTime", "ClassInfo");
 	Builder.RouteEvent(RouteId_BatchIssued, "IoDispatcher", "BatchIssued");
@@ -93,7 +102,8 @@
 	// Backwards compatibility
 	Builder.RouteEvent(RouteId_BeginObjectScope, "LoadTime", "BeginObjectScope");
 	Builder.RouteEvent(RouteId_EndObjectScope, "LoadTime", "EndObjectScope");
-	Builder.RouteEvent(RouteId_AsyncPackageLinkerAssociation, "LoadTime", "AsyncPackageLinkerAssociation");
+	Builder.RouteEvent(RouteId_BeginPostLoadExport, "LoadTime", "BeginPostLoadExport");
+	Builder.RouteEvent(RouteId_EndPostLoadExport, "LoadTime", "EndPostLoadExport");
 }
 
 bool FAsyncLoadingTraceAnalyzer::OnEvent(uint16 RouteId, EStyle Style, const FOnEventContext& Context)
@@ -105,67 +115,130 @@
 	{
 	case RouteId_PackageSummary:
 	{
-		FAsyncPackageState* AsyncPackage = nullptr;
 		uint64 AsyncPackagePtr = EventData.GetValue<uint64>("AsyncPackage");
-		if (AsyncPackagePtr)
+		FAsyncPackageState* AsyncPackage = ActiveAsyncPackagesMap.FindRef(AsyncPackagePtr);
+		if (!AsyncPackage)
 		{
-			AsyncPackage = ActiveAsyncPackagesMap.FindRef(AsyncPackagePtr);
+			FLinkerLoadState* Linker = ActiveLinkersMap.FindRef(AsyncPackagePtr);
+			if (Linker)
+			{
+				AsyncPackage = Linker->AsyncPackage;
+			}
+		}
+		if (AsyncPackage)
+		{
+			FPackageInfo* Package = AsyncPackage->PackageInfo;
+			if (Package)
+			{
+				FAnalysisSessionEditScope _(Session);
+				FString PackageName = FTraceAnalyzerUtils::LegacyAttachmentString<TCHAR>("Name", Context);
+				if (PackageName.Len())
+				{
+					Package->Name = Session.StoreString(*PackageName);
+				}
+				FPackageSummaryInfo& Summary = Package->Summary;
+				Summary.TotalHeaderSize = EventData.GetValue<uint32>("TotalHeaderSize");
+				Summary.ImportCount = EventData.GetValue<uint32>("ImportCount");
+				Summary.ExportCount = EventData.GetValue<uint32>("ExportCount");
+			}
+		}
+		break;
+	}
+	case RouteId_BeginProcessSummary:
+	{
+		FPackageInfo* Package = nullptr;
+		uint64 AsyncPackagePtr = EventData.GetValue<uint64>("AsyncPackage");
+		FAsyncPackageState* AsyncPackage = ActiveAsyncPackagesMap.FindRef(AsyncPackagePtr);
+		if (AsyncPackage)
+		{
+			Package = AsyncPackage->PackageInfo;
 		}
 		else
 		{
-			uint64 LinkerPtr = EventData.GetValue<uint64>("Linker"); // Backwards compatibility
-			AsyncPackage = LinkerToAsyncPackageMap.FindRef(LinkerPtr);
+			FLinkerLoadState* Linker = ActiveLinkersMap.FindRef(AsyncPackagePtr);
+			if (Linker)
+			{
+				if (Linker->AsyncPackage)
+				{
+					Package = Linker->AsyncPackage->PackageInfo;
+					Linker->PackageInfo = Package;
+				}
+				else
+				{
+					{
+						FAnalysisSessionEditScope _(Session);
+						Package = &LoadTimeProfilerProvider.CreatePackage();
+					}
+					FAsyncPackageState* FakeAsyncPackageState = new FAsyncPackageState();
+					FakeAsyncPackageState->PackageInfo = Package;
+					FakeAsyncPackageState->Linker = Linker;
+					Linker->PackageInfo = Package;
+					Linker->AsyncPackage = FakeAsyncPackageState;
+					Linker->bHasFakeAsyncPackageState = true;
+				}
+			}
 		}
-		if (AsyncPackage)
+		uint32 ThreadId = FTraceAnalyzerUtils::GetThreadIdField(Context);
+		FThreadState& ThreadState = GetThreadState(ThreadId);
+		uint64 Cycle = EventData.GetValue<uint64>("Cycle");
+		double Time = Context.EventTime.AsSeconds(Cycle);
 		{
 			FAnalysisSessionEditScope _(Session);
-			FString PackageName = FTraceAnalyzerUtils::LegacyAttachmentString<TCHAR>("Name", Context);
-			if (PackageName.Len())
-			{
-				AsyncPackage->PackageInfo->Name = Session.StoreString(*PackageName);
-			}
-			FPackageSummaryInfo& Summary = AsyncPackage->PackageInfo->Summary;
-			Summary.TotalHeaderSize = EventData.GetValue<uint32>("TotalHeaderSize");
-			Summary.ImportCount = EventData.GetValue<uint32>("ImportCount");
-			Summary.ExportCount = EventData.GetValue<uint32>("ExportCount");
+			ThreadState.EnterScope(Time, Package);
 		}
 		break;
 	}
+	case RouteId_EndProcessSummary:
+	{
+		uint32 ThreadId = FTraceAnalyzerUtils::GetThreadIdField(Context);
+		FThreadState& ThreadState = GetThreadState(ThreadId);
+		if (ensure(ThreadState.GetCurrentScopeEventType() == LoadTimeProfilerObjectEventType_None))
+		{
+			uint64 Cycle = EventData.GetValue<uint64>("Cycle");
+			double Time = Context.EventTime.AsSeconds(Cycle);
+			FAnalysisSessionEditScope _(Session);
+			ThreadState.LeaveScope(Time);
+		}
+		break;
+	}
 	case RouteId_BeginCreateExport:
 	{
+		uint64 AsyncPackagePtr = EventData.GetValue<uint64>("AsyncPackage");
+		FAsyncPackageState* AsyncPackage = ActiveAsyncPackagesMap.FindRef(AsyncPackagePtr);
+		if (!AsyncPackage)
+		{
+			FLinkerLoadState* Linker = ActiveLinkersMap.FindRef(AsyncPackagePtr);
+			if (Linker)
+			{
+				AsyncPackage = Linker->AsyncPackage;
+			}
+		}
+		uint32 ThreadId = FTraceAnalyzerUtils::GetThreadIdField(Context);
+		FThreadState& ThreadState = GetThreadState(ThreadId);
+		uint64 Cycle = EventData.GetValue<uint64>("Cycle");
+		double Time = Context.EventTime.AsSeconds(Cycle);
+
 		FAnalysisSessionEditScope _(Session);
 		FPackageExportInfo& Export = LoadTimeProfilerProvider.CreateExport();
 		Export.SerialSize = EventData.GetValue<uint64>("SerialSize"); // Backwards compatibility
-		FAsyncPackageState* AsyncPackage = nullptr;
-		uint64 AsyncPackagePtr = EventData.GetValue<uint64>("AsyncPackage");
-		if (AsyncPackagePtr)
-		{
-			AsyncPackage = ActiveAsyncPackagesMap.FindRef(AsyncPackagePtr);
-		}
-		else
-		{
-			uint64 LinkerPtr = EventData.GetValue<uint64>("Linker");  // Backwards compatibility
-			AsyncPackage = LinkerToAsyncPackageMap.FindRef(LinkerPtr);
-		}
 		if (AsyncPackage)
 		{
-			AsyncPackage->PackageInfo->Exports.Add(&Export);
-			Export.Package = AsyncPackage->PackageInfo;
+			if (FPackageInfo* PackageInfo = AsyncPackage->PackageInfo)
+			{
+				PackageInfo->Exports.Add(&Export);
+				Export.Package = PackageInfo;
+			}
 		}
-		uint32 ThreadId = FTraceAnalyzerUtils::GetThreadIdField(Context);
-		FThreadState& ThreadState = GetThreadState(ThreadId);
-		uint64 Cycle = EventData.GetValue<uint64>("Cycle");
-		double Time = Context.EventTime.AsSeconds(Cycle);
-		ThreadState.EnterExportScope(Time, &Export, LoadTimeProfilerObjectEventType_Create);
+		ThreadState.EnterScope(Time, &Export, LoadTimeProfilerObjectEventType_Create);
 		break;
 	}
 	case RouteId_EndCreateExport:
 	{
-		FAnalysisSessionEditScope _(Session);
 		uint32 ThreadId = FTraceAnalyzerUtils::GetThreadIdField(Context);
 		FThreadState& ThreadState = GetThreadState(ThreadId);
-		if (ensure(ThreadState.GetCurrentExportScopeEventType() == LoadTimeProfilerObjectEventType_Create))
+		if (ensure(ThreadState.GetCurrentScopeEventType() == LoadTimeProfilerObjectEventType_Create))
 		{
+			FAnalysisSessionEditScope _(Session);
 			FPackageExportInfo* Export = ThreadState.GetCurrentExportScope();
 			if (Export)
 			{
@@ -177,15 +250,19 @@
 
 			uint64 Cycle = EventData.GetValue<uint64>("Cycle");
 			double Time = Context.EventTime.AsSeconds(Cycle);
-			ThreadState.LeaveExportScope(Time);
+			ThreadState.LeaveScope(Time);
 		}
 		break;
 	}
 	case RouteId_BeginSerializeExport:
 	{
-		FAnalysisSessionEditScope _(Session);
+		uint32 ThreadId = FTraceAnalyzerUtils::GetThreadIdField(Context);
+		FThreadState& ThreadState = GetThreadState(ThreadId);
+		uint64 Cycle = EventData.GetValue<uint64>("Cycle");
+		double Time = Context.EventTime.AsSeconds(Cycle);
 		uint64 ObjectPtr = EventData.GetValue<uint64>("Object");
 		FPackageExportInfo* Export = ExportsMap.FindRef(ObjectPtr);
+		FAnalysisSessionEditScope _(Session);
 		if (Export)
 		{
 			uint64 SerialSize = EventData.GetValue<uint64>("SerialSize");
@@ -198,69 +275,94 @@
 				const_cast<FPackageInfo*>(Export->Package)->TotalExportsSerialSize += SerialSize;
 			}
 		}
-		uint32 ThreadId = FTraceAnalyzerUtils::GetThreadIdField(Context);
-		FThreadState& ThreadState = GetThreadState(ThreadId);
-		uint64 Cycle = EventData.GetValue<uint64>("Cycle");
-		double Time = Context.EventTime.AsSeconds(Cycle);
-		ThreadState.EnterExportScope(Time, Export, LoadTimeProfilerObjectEventType_Serialize);
+		ThreadState.EnterScope(Time, Export, LoadTimeProfilerObjectEventType_Serialize);
 		break;
 	}
 	case RouteId_EndSerializeExport:
 	{
-		FAnalysisSessionEditScope _(Session);
 		uint32 ThreadId = FTraceAnalyzerUtils::GetThreadIdField(Context);
 		FThreadState& ThreadState = GetThreadState(ThreadId);
-		if (ensure(ThreadState.GetCurrentExportScopeEventType() == LoadTimeProfilerObjectEventType_Serialize))
+		if (ensure(ThreadState.GetCurrentScopeEventType() == LoadTimeProfilerObjectEventType_Serialize))
 		{
 			uint64 Cycle = EventData.GetValue<uint64>("Cycle");
 			double Time = Context.EventTime.AsSeconds(Cycle);
-			ThreadState.LeaveExportScope(Time);
+			FAnalysisSessionEditScope _(Session);
+			ThreadState.LeaveScope(Time);
 		}
 		break;
 	}
+	case RouteId_BeginPostLoad:
+	{
+		uint32 ThreadId = FTraceAnalyzerUtils::GetThreadIdField(Context);
+		FThreadState& ThreadState = GetThreadState(ThreadId);
+		++ThreadState.PostLoadScopeDepth;
+		break;
+	}
+	case RouteId_EndPostLoad:
+	{
+		uint32 ThreadId = FTraceAnalyzerUtils::GetThreadIdField(Context);
+		FThreadState& ThreadState = GetThreadState(ThreadId);
+		--ThreadState.PostLoadScopeDepth;
+		break;
+	}
 	case RouteId_BeginPostLoadExport:
+	case RouteId_BeginPostLoadObject:
 	{
-		FAnalysisSessionEditScope _(Session);
+		uint32 ThreadId = FTraceAnalyzerUtils::GetThreadIdField(Context);
+		FThreadState& ThreadState = GetThreadState(ThreadId);
+		if (RouteId == RouteId_BeginPostLoadObject && !ThreadState.PostLoadScopeDepth)
+		{
+			break;
+		}
 		uint64 ObjectPtr = EventData.GetValue<uint64>("Object");
 		FPackageExportInfo* Export = ExportsMap.FindRef(ObjectPtr);
-		uint32 ThreadId = FTraceAnalyzerUtils::GetThreadIdField(Context);
-		FThreadState& ThreadState = GetThreadState(ThreadId);
 		uint64 Cycle = EventData.GetValue<uint64>("Cycle");
 		double Time = Context.EventTime.AsSeconds(Cycle);
-		ThreadState.EnterExportScope(Time, Export, LoadTimeProfilerObjectEventType_PostLoad);
+		FAnalysisSessionEditScope _(Session);
+		ThreadState.EnterScope(Time, Export, LoadTimeProfilerObjectEventType_PostLoad);
 		break;
 	}
 	case RouteId_EndPostLoadExport:
+	case RouteId_EndPostLoadObject:
 	{
-		FAnalysisSessionEditScope _(Session);
 		uint32 ThreadId = FTraceAnalyzerUtils::GetThreadIdField(Context);
 		FThreadState& ThreadState = GetThreadState(ThreadId);
-		if (ensure(ThreadState.GetCurrentExportScopeEventType() == LoadTimeProfilerObjectEventType_PostLoad))
+		if (RouteId == RouteId_EndPostLoadObject && !ThreadState.PostLoadScopeDepth)
+		{
+			break;
+		}
+		if (ensure(ThreadState.GetCurrentScopeEventType() == LoadTimeProfilerObjectEventType_PostLoad))
 		{
 			uint64 Cycle = EventData.GetValue<uint64>("Cycle");
 			double Time = Context.EventTime.AsSeconds(Cycle);
-			ThreadState.LeaveExportScope(Time);
+			FAnalysisSessionEditScope _(Session);
+			ThreadState.LeaveScope(Time);
 		}
 		break;
 	}
 	case RouteId_BeginRequest:
 	{
-		FAnalysisSessionEditScope _(Session);
 		uint64 RequestId = EventData.GetValue<uint64>("RequestId");
 		uint32 ThreadId = FTraceAnalyzerUtils::GetThreadIdField(Context);
-		//check(!ActiveRequestsMap.Contains(RequestId));
 		FRequestState* RequestState = new FRequestState();
 		RequestState->WallTimeStartCycle = EventData.GetValue<uint64>("Cycle");
 		RequestState->WallTimeEndCycle = 0;
 		RequestState->ThreadId = ThreadId;
-		ActiveRequestsMap.Add(RequestId, RequestState);
+	
+		if (RequestId == 0)
+		{
+			FakeRequestsStack.Push(RequestState);
+		}
+		else
+		{
+			ActiveRequestsMap.Add(RequestId, RequestState);
+		}
 		FThreadState& ThreadState = GetThreadState(ThreadId);
-		TSharedPtr<FRequestGroupState> RequestGroup = ThreadState.RequestGroupStack.Num() ? ThreadState.RequestGroupStack.Top() : nullptr;
+		TSharedPtr<FRequestGroupState> RequestGroup = ThreadState.RequestGroupStack.Num() ? ThreadState.RequestGroupStack[0] : nullptr;
 		if (!RequestGroup)
 		{
 			RequestGroup = MakeShared<FRequestGroupState>();
 			RequestGroup->Name = TEXT("[ungrouped]");
-			RequestGroup->bIsClosed = true;
 		}
 		RequestGroup->Requests.Add(RequestState);
 		++RequestGroup->ActiveRequestsCount;
@@ -270,17 +372,29 @@
 	case RouteId_EndRequest:
 	{
 		uint64 RequestId = EventData.GetValue<uint64>("RequestId");
-		FRequestState* RequestState = ActiveRequestsMap.FindRef(RequestId);
+		FRequestState* RequestState = nullptr;
+		if (RequestId == 0)
+		{
+			if (!FakeRequestsStack.IsEmpty())
+			{
+				RequestState = FakeRequestsStack.Pop();
+			}
+		}
+		else
+		{
+			ActiveRequestsMap.RemoveAndCopyValue(RequestId, RequestState);
+		}
 		if (RequestState)
 		{
 			RequestState->WallTimeEndCycle = EventData.GetValue<uint64>("Cycle");
 			RequestState->Group->LatestEndCycle = FMath::Max(RequestState->Group->LatestEndCycle, RequestState->WallTimeEndCycle);
 			--RequestState->Group->ActiveRequestsCount;
-			if (RequestState->Group->LoadRequest && RequestState->Group->bIsClosed && RequestState->Group->ActiveRequestsCount == 0)
+			if (RequestState->Group->LoadRequest && RequestState->Group->ActiveRequestsCount == 0)
 			{
 				FAnalysisSessionEditScope _(Session);
 				RequestState->Group->LoadRequest->EndTime = Context.EventTime.AsSeconds(RequestState->Group->LatestEndCycle);
 			}
+			delete RequestState;
 		}
 		break;
 	}
@@ -299,96 +413,111 @@
 			FormatString = FormatStringPtr;
 		}
 
+		uint32 ThreadId = FTraceAnalyzerUtils::GetThreadIdField(Context);
+		FThreadState& ThreadState = GetThreadState(ThreadId);
+
 		FAnalysisSessionEditScope _(Session);
 		TSharedRef<FRequestGroupState> GroupState = MakeShared<FRequestGroupState>();
 		FFormatArgsHelper::Format(FormatBuffer, FormatBufferSize - 1, TempBuffer, FormatBufferSize - 1, *FormatString, FormatArgs);
 		GroupState->Name = Session.StoreString(FormatBuffer);
-		uint32 ThreadId = FTraceAnalyzerUtils::GetThreadIdField(Context);
-		FThreadState& ThreadState = GetThreadState(ThreadId);
 		ThreadState.RequestGroupStack.Push(GroupState);
 		break;
 	}
 	case RouteId_EndRequestGroup:
 	{
-		FAnalysisSessionEditScope _(Session);
 		uint32 ThreadId = FTraceAnalyzerUtils::GetThreadIdField(Context);
 		FThreadState& ThreadState = GetThreadState(ThreadId);
 		if (ThreadState.RequestGroupStack.Num())
 		{
-			TSharedPtr<FRequestGroupState> GroupState = ThreadState.RequestGroupStack.Pop(false);
-			GroupState->bIsClosed = true;
-			if (GroupState->LoadRequest && GroupState->ActiveRequestsCount == 0)
-			{
-				GroupState->LoadRequest->EndTime = Context.EventTime.AsSeconds(GroupState->LatestEndCycle);
-			}
+			ThreadState.RequestGroupStack.Pop(false);
 		}
 		break;
 	}
 	case RouteId_NewAsyncPackage:
 	{
-		FAnalysisSessionEditScope _(Session);
 		uint64 AsyncPackagePtr = EventData.GetValue<uint64>("AsyncPackage");
-		//check(!ActivePackagesMap.Contains(AsyncPackagePtr));
 		FAsyncPackageState* AsyncPackageState = new FAsyncPackageState();
-		AsyncPackageState->PackageInfo = &LoadTimeProfilerProvider.CreatePackage();
-		if (EventData.GetAttachmentSize()) // For backwards compatibility
 		{
-			const TCHAR* PackageName = reinterpret_cast<const TCHAR*>(EventData.GetAttachment());
-			AsyncPackageState->PackageInfo->Name = Session.StoreString(PackageName);
+			FAnalysisSessionEditScope _(Session);
+			AsyncPackageState->PackageInfo = &LoadTimeProfilerProvider.CreatePackage();
 		}
 		ActiveAsyncPackagesMap.Add(AsyncPackagePtr, AsyncPackageState);
 		break;
 	}
-	case RouteId_BeginLoadAsyncPackage:
-	{
-		uint64 AsyncPackagePtr = EventData.GetValue<uint64>("AsyncPackage");
-		FAsyncPackageState* AsyncPackage = ActiveAsyncPackagesMap.FindRef(AsyncPackagePtr);
-		if (AsyncPackage)
-		{
-			AsyncPackage->LoadStartCycle = EventData.GetValue<uint64>("Cycle");
-			if (AsyncPackage->PackageInfo)
-			{
-				FAnalysisSessionEditScope _(Session);
-				double Time = Context.EventTime.AsSeconds(AsyncPackage->LoadStartCycle);
-				AsyncPackage->LoadHandle = LoadTimeProfilerProvider.BeginLoadPackage(*AsyncPackage->PackageInfo, Time);
-			}
-		}
-		break;
-	}
-	case RouteId_EndLoadAsyncPackage:
-	{
-		uint64 AsyncPackagePtr = EventData.GetValue<uint64>("AsyncPackage");
-		FAsyncPackageState* AsyncPackage = ActiveAsyncPackagesMap.FindRef(AsyncPackagePtr);
-		if (AsyncPackage)
-		{
-			AsyncPackage->LoadEndCycle = EventData.GetValue<uint64>("Cycle");
-			if (AsyncPackage->PackageInfo && AsyncPackage->LoadHandle != uint64(-1))
-			{
-				FAnalysisSessionEditScope _(Session);
-				double Time = Context.EventTime.AsSeconds(AsyncPackage->LoadEndCycle);
-				LoadTimeProfilerProvider.EndLoadPackage(AsyncPackage->LoadHandle, Time);
-			}
-		}
-		break;
-	}
 	case RouteId_DestroyAsyncPackage:
 	{
 		uint64 AsyncPackagePtr = EventData.GetValue<uint64>("AsyncPackage");
-		//check(ActiveAsyncPackagesMap.Contains(AsyncPackagePtr));
-		ActiveAsyncPackagesMap.Remove(AsyncPackagePtr);
+		FAsyncPackageState* AsyncPackageState;
+		if (ActiveAsyncPackagesMap.RemoveAndCopyValue(AsyncPackagePtr, AsyncPackageState))
+		{
+			if (AsyncPackageState->Linker)
+			{
+				AsyncPackageState->Linker->AsyncPackage = nullptr;
+			}
+			delete AsyncPackageState;
+		}
 		break;
 	}
+	case RouteId_NewLinker:
+	{
+		uint64 LinkerPtr = EventData.GetValue<uint64>("Linker");
+		FLinkerLoadState* LinkerState = new FLinkerLoadState();
+		ActiveLinkersMap.Add(LinkerPtr, LinkerState);
+		break;
+	}
+	case RouteId_DestroyLinker:
+	{
+		uint64 LinkerPtr = EventData.GetValue<uint64>("Linker");
+		FLinkerLoadState* LinkerState;
+		if (ActiveLinkersMap.RemoveAndCopyValue(LinkerPtr, LinkerState))
+		{
+			if (LinkerState->AsyncPackage)
+			{
+				LinkerState->AsyncPackage->Linker = nullptr;
+			}
+			if (LinkerState->bHasFakeAsyncPackageState)
+			{
+				delete LinkerState->AsyncPackage;
+			}
+			delete LinkerState;
+		}
+		break;
+	}
 	case RouteId_AsyncPackageImportDependency:
 	{
 		uint64 AsyncPackagePtr = EventData.GetValue<uint64>("AsyncPackage");
-		FAsyncPackageState* AsyncPackage = ActiveAsyncPackagesMap.FindRef(AsyncPackagePtr);
 		uint64 ImportedAsyncPackagePtr = EventData.GetValue<uint64>("ImportedAsyncPackage");
-		FAsyncPackageState* ImportedAsyncPackage = ActiveAsyncPackagesMap.FindRef(ImportedAsyncPackagePtr);
+
+		FAsyncPackageState* AsyncPackage = ActiveAsyncPackagesMap.FindRef(AsyncPackagePtr);
+		FAsyncPackageState* ImportedAsyncPackage = nullptr;
+		if (AsyncPackage)
+		{
+			ImportedAsyncPackage = ActiveAsyncPackagesMap.FindRef(ImportedAsyncPackagePtr);
+		}
+		else
+		{
+			FLinkerLoadState* LinkerState = ActiveLinkersMap.FindRef(AsyncPackagePtr);
+			if (LinkerState)
+			{
+				AsyncPackage = LinkerState->AsyncPackage;
+				FLinkerLoadState* ImportedLinkerState = ActiveLinkersMap.FindRef(ImportedAsyncPackagePtr);
+				if (ImportedLinkerState)
+				{
+					ImportedAsyncPackage = ImportedLinkerState->AsyncPackage;
+				}
+			}
+		}
 		if (AsyncPackage && ImportedAsyncPackage)
 		{
-			if (AsyncPackage->Request)
+			if (!AsyncPackage->ImportedAsyncPackages.Contains(ImportedAsyncPackage))
 			{
-				PackageRequestAssociation(Context, ImportedAsyncPackage, AsyncPackage->Request);
+				AsyncPackage->ImportedAsyncPackages.Add(ImportedAsyncPackage);
+				FAnalysisSessionEditScope _(Session);
+				AsyncPackage->PackageInfo->ImportedPackages.Add(ImportedAsyncPackage->PackageInfo);
+				if (AsyncPackage->Request)
+				{
+					PackageRequestAssociation(Context, ImportedAsyncPackage, AsyncPackage->Request);
+				}
 			}
 		}
 		break;
@@ -396,12 +525,36 @@
 	case RouteId_AsyncPackageRequestAssociation:
 	{
 		uint64 AsyncPackagePtr = EventData.GetValue<uint64>("AsyncPackage");
-		FAsyncPackageState* AsyncPackageState = ActiveAsyncPackagesMap.FindRef(AsyncPackagePtr);
+		FAsyncPackageState* AsyncPackageState = nullptr;
 		uint64 RequestId = EventData.GetValue<uint64>("RequestId");
-		FRequestState* RequestState = ActiveRequestsMap.FindRef(RequestId);
-		if (AsyncPackageState && RequestState)
+		FRequestState* RequestState = nullptr;
+		if (RequestId == 0)
 		{
-			PackageRequestAssociation(Context, AsyncPackageState, RequestState);
+			if (!FakeRequestsStack.IsEmpty())
+			{
+				RequestState = FakeRequestsStack.Top();
+			}
 		}
+		else
+		{
+			RequestState = ActiveRequestsMap.FindRef(RequestId);
+		}
+		if (RequestState)
+		{
+			AsyncPackageState = ActiveAsyncPackagesMap.FindRef(AsyncPackagePtr);
+			if (!AsyncPackageState)
+			{
+				FLinkerLoadState* Linker = ActiveLinkersMap.FindRef(AsyncPackagePtr);
+				if (Linker)
+				{
+					AsyncPackageState = Linker->AsyncPackage;
+				}
+			}
+			if (AsyncPackageState)
+			{
+				FAnalysisSessionEditScope _(Session);
+				PackageRequestAssociation(Context, AsyncPackageState, RequestState);
+			}
+		}
 		break;
 	}
@@ -440,28 +593,29 @@
 	}
 	case RouteId_BeginObjectScope:
 	{
-		FAnalysisSessionEditScope _(Session);
+		uint32 ThreadId = FTraceAnalyzerUtils::GetThreadIdField(Context);
+		FThreadState& ThreadState = GetThreadState(ThreadId);
 		uint64 ObjectPtr = EventData.GetValue<uint64>("Object");
 		FPackageExportInfo* Export = ExportsMap.FindRef(ObjectPtr);
-		uint32 ThreadId = FTraceAnalyzerUtils::GetThreadIdField(Context);
-		FThreadState& ThreadState = GetThreadState(ThreadId);
 		ELoadTimeProfilerObjectEventType EventType = static_cast<ELoadTimeProfilerObjectEventType>(EventData.GetValue<uint8>("EventType"));
 		//EventType = static_cast<ELoadTimeProfilerObjectEventType>(99); // debug
 		uint64 Cycle = EventData.GetValue<uint64>("Cycle");
 		double Time = Context.EventTime.AsSeconds(Cycle);
-		ThreadState.EnterExportScope(Time, Export, EventType);
+
+		FAnalysisSessionEditScope _(Session);
+		ThreadState.EnterScope(Time, Export, EventType);
 		break;
 	}
 	case RouteId_EndObjectScope:
 	{
-		FAnalysisSessionEditScope _(Session);
 		uint32 ThreadId = FTraceAnalyzerUtils::GetThreadIdField(Context);
 		FThreadState& ThreadState = GetThreadState(ThreadId);
 		//if (ensure(ThreadState.GetCurrentExportScopeEventType() == static_cast<ELoadTimeProfilerObjectEventType>(99))) // debug
 		{
 			uint64 Cycle = EventData.GetValue<uint64>("Cycle");
 			double Time = Context.EventTime.AsSeconds(Cycle);
-			ThreadState.LeaveExportScope(Time);
+			FAnalysisSessionEditScope _(Session);
+			ThreadState.LeaveScope(Time);
 		}
 		break;
 	}
@@ -470,9 +624,17 @@
 		uint64 LinkerPtr = EventData.GetValue<uint64>("Linker");
 		uint64 AsyncPackagePtr = EventData.GetValue<uint64>("AsyncPackage");
 		FAsyncPackageState* AsyncPackageState = ActiveAsyncPackagesMap.FindRef(AsyncPackagePtr);
-		if (AsyncPackageState)
+		FLinkerLoadState* LinkerState = ActiveLinkersMap.FindRef(LinkerPtr);
+		if (AsyncPackageState && LinkerState)
 		{
-			LinkerToAsyncPackageMap.Add(LinkerPtr, AsyncPackageState);
+			if (LinkerState->PackageInfo)
+			{
+				FAnalysisSessionEditScope _(Session);
+				AsyncPackageState->PackageInfo->Name = LinkerState->PackageInfo->Name;
+				AsyncPackageState->PackageInfo->Summary = LinkerState->PackageInfo->Summary;
+			}
+			LinkerState->AsyncPackage = AsyncPackageState;
+			AsyncPackageState->Linker = LinkerState;
 		}
 		break;
 	}
@@ -481,17 +643,27 @@
 	return true;
 }
 
-void FAsyncLoadingTraceAnalyzer::FThreadState::EnterExportScope(double Time, const FPackageExportInfo* ExportInfo, ELoadTimeProfilerObjectEventType EventType)
+void FAsyncLoadingTraceAnalyzer::FThreadState::EnterScope(double Time, const FPackageInfo* PackageInfo)
 {
 	FScopeStackEntry& StackEntry = CpuScopeStack[CpuScopeStackDepth++];
+	StackEntry.Event.Export = nullptr;
+	StackEntry.Event.Package = PackageInfo;
+	StackEntry.Event.EventType = LoadTimeProfilerObjectEventType_None;
+	CurrentEvent = StackEntry.Event;
+	CpuTimeline->AppendBeginEvent(Time, StackEntry.Event);
+}
+
+void FAsyncLoadingTraceAnalyzer::FThreadState::EnterScope(double Time, const FPackageExportInfo* ExportInfo, ELoadTimeProfilerObjectEventType EventType)
+{
+	FScopeStackEntry& StackEntry = CpuScopeStack[CpuScopeStackDepth++];
 	StackEntry.Event.Export = ExportInfo;
-	StackEntry.Event.EventType = EventType;
 	StackEntry.Event.Package = ExportInfo ? ExportInfo->Package : nullptr;
+	StackEntry.Event.EventType = EventType;
 	CurrentEvent = StackEntry.Event;
 	CpuTimeline->AppendBeginEvent(Time, StackEntry.Event);
 }
 
-void FAsyncLoadingTraceAnalyzer::FThreadState::LeaveExportScope(double Time)
+void FAsyncLoadingTraceAnalyzer::FThreadState::LeaveScope(double Time)
 {
 	check(CpuScopeStackDepth > 0);
 	FScopeStackEntry& StackEntry = CpuScopeStack[--CpuScopeStackDepth];
@@ -506,7 +678,7 @@
 	}
 }
 
-ELoadTimeProfilerObjectEventType FAsyncLoadingTraceAnalyzer::FThreadState::GetCurrentExportScopeEventType()
+ELoadTimeProfilerObjectEventType FAsyncLoadingTraceAnalyzer::FThreadState::GetCurrentScopeEventType()
 {
 	if (CpuScopeStackDepth > 0)
 	{
@@ -515,7 +687,7 @@
 	}
 	else
 	{
-		return ELoadTimeProfilerObjectEventType::LoadTimeProfilerObjectEventType_None;
+		return LoadTimeProfilerObjectEventType_None;
 	}
 }
 
@@ -539,7 +711,6 @@
 		RequestState->AsyncPackages.Add(AsyncPackageState);
 		AsyncPackageState->Request = RequestState;
 		FLoadRequest* LoadRequest = RequestState->Group->LoadRequest;
-		FAnalysisSessionEditScope _(Session);
 		if (!LoadRequest)
 		{
 			LoadRequest = &LoadTimeProfilerProvider.CreateRequest();
@@ -549,7 +720,15 @@
 			LoadRequest->ThreadId = RequestState->ThreadId;
 			RequestState->Group->LoadRequest = LoadRequest;
 		}
-		LoadRequest->Packages.Add(AsyncPackageState->PackageInfo);
+		if (AsyncPackageState->PackageInfo)
+		{
+			LoadRequest->Packages.Add(AsyncPackageState->PackageInfo);
+			AsyncPackageState->PackageInfo->RequestId = LoadRequest->Id;
+		}
+		for (FAsyncPackageState* ImportedAsyncPackage : AsyncPackageState->ImportedAsyncPackages)
+		{
+			PackageRequestAssociation(Context, ImportedAsyncPackage, RequestState);
+		}
 	}
 }
 
diff --git a/Engine/Source/Developer/TraceServices/Private/Analyzers/LoadTimeTraceAnalysis.h b/Engine/Source/Developer/TraceServices/Private/Analyzers/LoadTimeTraceAnalysis.h
--- a/Engine/Source/Developer/TraceServices/Private/Analyzers/LoadTimeTraceAnalysis.h
+++ b/Engine/Source/Developer/TraceServices/Private/Analyzers/LoadTimeTraceAnalysis.h
@@ -31,6 +31,7 @@ public:
 private:
 	struct FRequestState;
 	struct FAsyncPackageState;
+	struct FLinkerLoadState;
 
 	struct FRequestGroupState
 	{
@@ -39,7 +40,6 @@ private:
 		FLoadRequest* LoadRequest = nullptr;
 		uint64 LatestEndCycle = 0;
 		uint64 ActiveRequestsCount = 0;
-		bool bIsClosed = false;
 	};
 
 	struct FRequestState
@@ -55,9 +55,15 @@ private:
 	{
 		FPackageInfo* PackageInfo = nullptr;
 		FRequestState* Request = nullptr;
-		uint64 LoadHandle = uint64(-1);
-		uint64 LoadStartCycle = 0;
-		uint64 LoadEndCycle = 0;
+		FLinkerLoadState* Linker = nullptr;
+		TSet<FAsyncPackageState*> ImportedAsyncPackages;
+	};
+
+	struct FLinkerLoadState
+	{
+		FPackageInfo* PackageInfo = nullptr;
+		FAsyncPackageState* AsyncPackage = nullptr;
+		bool bHasFakeAsyncPackageState = false;
 	};
 
 	struct FScopeStackEntry
@@ -70,14 +76,16 @@ private:
 	{
 		FScopeStackEntry CpuScopeStack[256];
 		uint64 CpuScopeStackDepth = 0;
+		int64 PostLoadScopeDepth = 0;
 		FLoadTimeProfilerCpuEvent CurrentEvent;
 		TArray<TSharedPtr<FRequestGroupState>> RequestGroupStack;
 		
 		FLoadTimeProfilerProvider::CpuTimelineInternal* CpuTimeline;
 
-		void EnterExportScope(double Time, const FPackageExportInfo* ExportInfo, ELoadTimeProfilerObjectEventType EventType);
-		void LeaveExportScope(double Time);
-		ELoadTimeProfilerObjectEventType GetCurrentExportScopeEventType();
+		void EnterScope(double Time, const FPackageInfo* PackageInfo);
+		void EnterScope(double Time, const FPackageExportInfo* ExportInfo, ELoadTimeProfilerObjectEventType EventType);
+		void LeaveScope(double Time);
+		ELoadTimeProfilerObjectEventType GetCurrentScopeEventType();
 		FPackageExportInfo* GetCurrentExportScope();
 	};
 
@@ -91,22 +99,27 @@ private:
 		RouteId_SuspendAsyncLoading,
 		RouteId_ResumeAsyncLoading,
 		RouteId_NewAsyncPackage,
-		RouteId_BeginLoadAsyncPackage,
-		RouteId_EndLoadAsyncPackage,
 		RouteId_DestroyAsyncPackage,
+		RouteId_NewLinker,
+		RouteId_DestroyLinker,
 		RouteId_BeginRequest,
 		RouteId_EndRequest,
 		RouteId_BeginRequestGroup,
 		RouteId_EndRequestGroup,
 		RouteId_PackageSummary,
 		RouteId_AsyncPackageRequestAssociation,
+		RouteId_AsyncPackageLinkerAssociation,
 		RouteId_AsyncPackageImportDependency,
+		RouteId_BeginProcessSummary,
+		RouteId_EndProcessSummary,
 		RouteId_BeginCreateExport,
 		RouteId_EndCreateExport,
 		RouteId_BeginSerializeExport,
 		RouteId_EndSerializeExport,
-		RouteId_BeginPostLoadExport,
-		RouteId_EndPostLoadExport,
+		RouteId_BeginPostLoad,
+		RouteId_EndPostLoad,
+		RouteId_BeginPostLoadObject,
+		RouteId_EndPostLoadObject,
 		RouteId_ClassInfo,
 		RouteId_BatchIssued,
 		RouteId_BatchResolved,
@@ -114,7 +127,8 @@ private:
 		// Backwards compatibility
 		RouteId_BeginObjectScope,
 		RouteId_EndObjectScope,
-		RouteId_AsyncPackageLinkerAssociation,
+		RouteId_BeginPostLoadExport,
+		RouteId_EndPostLoadExport,
 	};
 
 	enum
@@ -156,14 +170,13 @@ private:
 	using TPointerMap = TMap<uint64, ValueType, FDefaultSetAllocator, FPointerMapKeyFuncs<ValueType>>;
 
 	TPointerMap<FAsyncPackageState*> ActiveAsyncPackagesMap;
+	TPointerMap<FLinkerLoadState*> ActiveLinkersMap;
 	TPointerMap<FPackageExportInfo*> ExportsMap;
 	TMap<uint64, FRequestState*> ActiveRequestsMap;
+	TArray<FRequestState*> FakeRequestsStack;
 	TPointerMap<uint64> ActiveBatchesMap;
 	TMap<uint32, FThreadState*> ThreadStatesMap;
 	TPointerMap<const FClassInfo*> ClassInfosMap;
-
-	// Backwards compatibility
-	TPointerMap<FAsyncPackageState*> LinkerToAsyncPackageMap;
 };
 
 } // namespace TraceServices
diff --git a/Engine/Source/Developer/TraceServices/Private/Model/LoadTimeProfiler.cpp b/Engine/Source/Developer/TraceServices/Private/Model/LoadTimeProfiler.cpp
--- a/Engine/Source/Developer/TraceServices/Private/Model/LoadTimeProfiler.cpp
+++ b/Engine/Source/Developer/TraceServices/Private/Model/LoadTimeProfiler.cpp
@@ -108,7 +108,6 @@ FLoadTimeProfilerProvider::FLoadTimeProfilerProvider(IAnalysisSession& InSession
 	, ClassInfos(Session.GetLinearAllocator(), 16384)
 	, Requests(Session.GetLinearAllocator(), 16384)
 	, Packages(Session.GetLinearAllocator(), 16384)
-	, PackageLoads(Session.GetLinearAllocator(), 16384)
 	, IoDispatcherBatches(Session.GetLinearAllocator(), 16384)
 	, Exports(Session.GetLinearAllocator(), 16384)
 	, RequestsTable(Requests)
@@ -149,11 +148,22 @@ FLoadTimeProfilerProvider::FLoadTimeProfilerProvider(IAnalysisSession& InSession
 				return Row.PackageInfo->Name;
 			},
 			TEXT("Package")).
+		AddColumn<uint64>([](const FPackagesTableRow& Row)
+			{
+				return Row.PackageInfo ? Row.PackageInfo->RequestId : 0;
+			},
+			TEXT("RequestId")).
+		AddColumn(&FPackagesTableRow::TotalSerializedSize, TEXT("SerializedSize"), TableColumnDisplayHint_Memory | TableColumnDisplayHint_Summable).
 		AddColumn(&FPackagesTableRow::SerializedHeaderSize, TEXT("SerializedHeaderSize"), TableColumnDisplayHint_Memory | TableColumnDisplayHint_Summable).
 		AddColumn(&FPackagesTableRow::SerializedExportsCount, TEXT("SerializedExportsCount"), TableColumnDisplayHint_Summable).
 		AddColumn(&FPackagesTableRow::SerializedExportsSize, TEXT("SerializedExportsSize"), TableColumnDisplayHint_Memory | TableColumnDisplayHint_Summable).
 		AddColumn(&FPackagesTableRow::MainThreadTime, TEXT("MainThreadTime"), TableColumnDisplayHint_Time | TableColumnDisplayHint_Summable).
-		AddColumn(&FPackagesTableRow::AsyncLoadingThreadTime, TEXT("AsyncLoadingThreadTime"), TableColumnDisplayHint_Time | TableColumnDisplayHint_Summable);
+		AddColumn(&FPackagesTableRow::AsyncLoadingThreadTime, TEXT("AsyncLoadingThreadTime"), TableColumnDisplayHint_Time | TableColumnDisplayHint_Summable).
+		AddColumn<int32>([](const FPackagesTableRow& Row)
+			{
+				return Row.PackageInfo ? Row.PackageInfo->ImportedPackages.Num() : 0;
+			},
+			TEXT("ImportedPackagesCount"), TableColumnDisplayHint_Summable);
 
 	ExportsTableLayout.
 		AddColumn<const TCHAR*>([](const FExportsTableRow& Row)
@@ -161,6 +171,11 @@ FLoadTimeProfilerProvider::FLoadTimeProfilerProvider(IAnalysisSession& InSession
 				return Row.ExportInfo->Package ? Row.ExportInfo->Package->Name : TEXT("[unknown]");
 			},
 			TEXT("Package")).
+		AddColumn<uint64>([](const FExportsTableRow& Row)
+			{
+				return Row.ExportInfo->Package ? Row.ExportInfo->Package->RequestId : 0;
+			},
+			TEXT("RequestId")).
 		AddColumn<const TCHAR*>([](const FExportsTableRow& Row)
 			{
 				return Row.ExportInfo->Class ? Row.ExportInfo->Class->Name : TEXT("[unknown]");
@@ -174,6 +189,38 @@ FLoadTimeProfilerProvider::FLoadTimeProfilerProvider(IAnalysisSession& InSession
 		AddColumn(&FExportsTableRow::SerializedSize, TEXT("SerializedSize"), TableColumnDisplayHint_Memory | TableColumnDisplayHint_Summable).
 		AddColumn(&FExportsTableRow::MainThreadTime, TEXT("MainThreadTime"), TableColumnDisplayHint_Time | TableColumnDisplayHint_Summable).
 		AddColumn(&FExportsTableRow::AsyncLoadingThreadTime, TEXT("AsyncLoadingThreadTime"), TableColumnDisplayHint_Time | TableColumnDisplayHint_Summable);
+
+	RequestsTableLayout.
+		AddColumn(&FRequestsTableRow::Id, TEXT("Id")).
+		AddColumn(&FRequestsTableRow::Name, TEXT("Name")).
+		AddColumn(&FRequestsTableRow::StartTime, TEXT("StartTime"), TableColumnDisplayHint_Time).
+		AddColumn(&FRequestsTableRow::Duration, TEXT("Duration"), TableColumnDisplayHint_Time | TableColumnDisplayHint_Summable).
+		AddColumn<int32>(
+			[](const FRequestsTableRow& Row)
+			{
+				return Row.Packages.Num();
+			},
+			TEXT("PackageCount"), TableColumnDisplayHint_Summable).
+		AddColumn<uint64>([](const FRequestsTableRow& Row)
+			{
+				uint64 Sum = 0;
+				for (const FPackageInfo* Package : Row.Packages)
+				{
+					Sum += Package->Summary.TotalHeaderSize;
+					for (const FPackageExportInfo* Export : Package->Exports)
+					{
+						Sum += Export->SerialSize;
+					}
+				}
+				return Sum;
+			},
+			TEXT("SerializedSize"), TableColumnDisplayHint_Memory | TableColumnDisplayHint_Summable).
+		AddColumn<const TCHAR*>(
+			[](const FRequestsTableRow& Row)
+			{
+				return Row.Packages.Num() ? Row.Packages[0]->Name : TEXT("N/A");
+			},
+			TEXT("FirstPackage"));
 }
 
 bool FLoadTimeProfilerProvider::GetCpuThreadTimelineIndex(uint32 ThreadId, uint32& OutTimelineIndex) const
@@ -268,61 +315,105 @@ ITable<FLoadTimeProfilerAggregatedStats>* FLoadTimeProfilerProvider::CreateObjec
 
 ITable<FPackagesTableRow>* FLoadTimeProfilerProvider::CreatePackageDetailsTable(double IntervalStart, double IntervalEnd) const
 {
-	TTable<FPackagesTableRow>* Table = new TTable<FPackagesTableRow>(PackagesTableLayout);
-
-	TMap<const FPackageInfo*, FPackagesTableRow*> PackagesMap;
-
-	auto FindRow = [Table, &PackagesMap](const FLoadTimeProfilerCpuEvent& Event) -> FPackagesTableRow*
+	struct FTempRow
 	{
-		if (Event.Package)
+		double StartTime = 0.0;
+		const FPackageInfo* PackageInfo = nullptr;
+		uint64 SerializedHeaderSize = 0;
+		uint64 SerializedExportsCount = 0;
+		uint64 SerializedExportsSize = 0;
+		double MainThreadTime = 0.0;
+		double AsyncLoadingThreadTime = 0.0;
+	};
+	TArray<FTempRow> Rows;
+	struct FStackEntry
+	{
+		double StartTime;
+		int32 RowIndex;
+	};
+	TArray<FStackEntry, TInlineAllocator<64>> Stack;
+	double LastTime = 0.0;
+	int32 TimelineIndex = 0;
+	TMap<const FPackageInfo*, int32> RowMap;
+	for (const TSharedRef<CpuTimelineInternal>& Timeline : CpuTimelines)
+	{
+		Timeline->EnumerateEvents(IntervalStart, IntervalEnd, [&Stack, &Rows, &RowMap, IntervalStart, IntervalEnd, &LastTime, TimelineIndex](bool IsEnter, double Time, const FLoadTimeProfilerCpuEvent& Event)
 		{
-			FPackagesTableRow** FindIt = PackagesMap.Find(Event.Package);
-			FPackagesTableRow* Row = nullptr;
-			if (!FindIt)
+			if (!Event.Package)
 			{
-				Row = &Table->AddRow();
-				Row->PackageInfo = Event.Package;
-				PackagesMap.Add(Event.Package, Row);
+				return EEventEnumerate::Continue;
 			}
-			else
+			double ClampedTime = FMath::Clamp(Time, IntervalStart, IntervalEnd);
+			if (Stack.Num())
 			{
-				Row = *FindIt;
-			}
-
-			if (Event.Export && Event.EventType == LoadTimeProfilerObjectEventType_Serialize)
-			{
-				++Row->SerializedExportsCount;
-				Row->SerializedExportsSize += Event.Export->SerialSize;
+				FStackEntry& StackEntry = Stack.Top();
+				FTempRow& Row = Rows[StackEntry.RowIndex];
+				if (TimelineIndex == 0)
+				{
+					Row.MainThreadTime += ClampedTime - LastTime;
+				}
+				else
+				{
+					Row.AsyncLoadingThreadTime += ClampedTime - LastTime;
+				}
 			}
-
-			Row->SerializedHeaderSize = Event.Package->Summary.TotalHeaderSize;
-
-			return Row;
-		}
-		else
-		{
-			return nullptr;
-		}
-	};
-
-	for (int32 TimelineIndex = 0, TimelineCount = CpuTimelines.Num(); TimelineIndex < TimelineCount; ++TimelineIndex)
-	{
-		CpuTimelines[TimelineIndex]->EnumerateEvents(IntervalStart, IntervalEnd, [Table, &PackagesMap, FindRow, TimelineIndex](double StartTime, double EndTime, uint32, const FLoadTimeProfilerCpuEvent& Event)
-		{
-			FPackagesTableRow* Row = FindRow(Event);
-			if (Row)
+			LastTime = ClampedTime;
+			if (IsEnter)
 			{
-				if (TimelineIndex == 0)
+				FTempRow* Row;
+				int32 RowIndex = -1;
+				int32* FindRowIndex = RowMap.Find(Event.Package);
+				if (FindRowIndex)
 				{
-					Row->MainThreadTime += EndTime - StartTime; // TODO: Should be exclusive time
+					Row = &Rows[*FindRowIndex];
+					RowIndex = *FindRowIndex;
 				}
 				else
 				{
-					Row->AsyncLoadingThreadTime += EndTime - StartTime; // TODO: Should be exclusive time
+					RowIndex = Rows.Num();
+					Row = &Rows.AddDefaulted_GetRef();
+					Row->StartTime = Time;
+					Row->PackageInfo = Event.Package;
+					RowMap.Add(Event.Package, RowIndex);
+				}
+				FStackEntry& StackEntry = Stack.AddDefaulted_GetRef();
+				StackEntry.StartTime = ClampedTime;
+				StackEntry.RowIndex = RowIndex;
+				if (Event.Export && Event.EventType == LoadTimeProfilerObjectEventType_Serialize)
+				{
+					Row->SerializedExportsSize += Event.Export->SerialSize;
+					++Row->SerializedExportsCount;
+				}
+				else if (!Event.Export && Event.Package && !Row->SerializedHeaderSize)
+				{
+					Row->SerializedHeaderSize = Event.Package->Summary.TotalHeaderSize;
 				}
 			}
+			else
+			{
+				Stack.Pop(false);
+			}
 			return EEventEnumerate::Continue;
 		});
+		++TimelineIndex;
+	}
+
+	Algo::Sort(Rows, [](const FTempRow& A, const FTempRow& B)
+		{
+			return A.StartTime < B.StartTime;
+		});
+	TTable<FPackagesTableRow>* Table = new TTable<FPackagesTableRow>(PackagesTableLayout);
+	for (const FTempRow& TempRow : Rows)
+	{
+		FPackagesTableRow& Row = Table->AddRow();
+		Row.PackageInfo = TempRow.PackageInfo;
+		Row.TotalSerializedSize = TempRow.SerializedHeaderSize + TempRow.SerializedExportsSize;
+		Row.SerializedHeaderSize = TempRow.SerializedHeaderSize;
+		Row.SerializedExportsCount = TempRow.SerializedExportsCount;
+		Row.SerializedExportsSize = TempRow.SerializedExportsSize;
+
+		Row.MainThreadTime = TempRow.MainThreadTime;
+		Row.AsyncLoadingThreadTime = TempRow.AsyncLoadingThreadTime;
 	}
 
 	return Table;
@@ -330,61 +421,105 @@ ITable<FPackagesTableRow>* FLoadTimeProfilerProvider::CreatePackageDetailsTable(
 
 ITable<FExportsTableRow>* FLoadTimeProfilerProvider::CreateExportDetailsTable(double IntervalStart, double IntervalEnd) const
 {
-	TTable<FExportsTableRow>* Table = new TTable<FExportsTableRow>(ExportsTableLayout);
-
-	TMap<TTuple<const FPackageExportInfo*, ELoadTimeProfilerObjectEventType>, FExportsTableRow*> ExportsMap;
-
-	auto FindRow = [Table, &ExportsMap](const FLoadTimeProfilerCpuEvent& Event) -> FExportsTableRow*
+	struct FTempRow
 	{
-		if (Event.Export)
-		{
-			auto Key = MakeTuple(Event.Export, Event.EventType);
-			FExportsTableRow** FindIt = ExportsMap.Find(Key);
-			FExportsTableRow* Row = nullptr;
-			if (!FindIt)
-			{
-				Row = &Table->AddRow();
-				Row->ExportInfo = Event.Export;
-				Row->EventType = Event.EventType;
-				ExportsMap.Add(Key, Row);
-			}
-			else
-			{
-				Row = *FindIt;
-			}
-
-			if (Event.EventType == LoadTimeProfilerObjectEventType_Serialize)
-			{
-				Row->SerializedSize += Event.Export->SerialSize;
-			}
-
-			return Row;
-		}
-		else
-		{
-			return nullptr;
-		}
+		double StartTime = 0.0;
+		const FPackageExportInfo* ExportInfo = nullptr;
+		uint64 SerializedSize = 0;
+		double MainThreadTime = 0.0;
+		double AsyncLoadingThreadTime = 0.0;
+		ELoadTimeProfilerObjectEventType EventType;
 	};
-
-	for (int32 TimelineIndex = 0, TimelineCount = CpuTimelines.Num(); TimelineIndex < TimelineCount; ++TimelineIndex)
+	TArray<FTempRow> Rows;
+	struct FStackEntry
+	{
+		double StartTime;
+		int32 RowIndex;
+	};
+	TArray<FStackEntry, TInlineAllocator<64>> Stack;
+	double LastTime = 0.0;
+	int32 TimelineIndex = 0;
+	for (const TSharedRef<CpuTimelineInternal>& Timeline : CpuTimelines)
 	{
-		CpuTimelines[TimelineIndex]->EnumerateEvents(IntervalStart, IntervalEnd, [Table, &ExportsMap, FindRow, TimelineIndex](double StartTime, double EndTime, uint32, const FLoadTimeProfilerCpuEvent& Event)
+		Timeline->EnumerateEvents(IntervalStart, IntervalEnd, [&Stack, &Rows, IntervalStart, IntervalEnd, &LastTime, TimelineIndex](bool IsEnter, double Time, const FLoadTimeProfilerCpuEvent& Event)
 		{
-			FExportsTableRow* Row = FindRow(Event);
-			if (Row)
+			if (!Event.Export)
+			{
+				return EEventEnumerate::Continue;
+			}
+			double ClampedTime = FMath::Clamp(Time, IntervalStart, IntervalEnd);
+			if (Stack.Num())
 			{
+				FStackEntry& StackEntry = Stack.Top();
+				FTempRow& Row = Rows[StackEntry.RowIndex];
 				if (TimelineIndex == 0)
 				{
-					Row->MainThreadTime += EndTime - StartTime; // TODO: Should be exclusive time
+					Row.MainThreadTime += ClampedTime - LastTime;
 				}
 				else
 				{
-					Row->AsyncLoadingThreadTime += EndTime - StartTime; // TODO: Should be exclusive time
+					Row.AsyncLoadingThreadTime += ClampedTime - LastTime;
 				}
 			}
+			LastTime = ClampedTime;
+			if (IsEnter)
+			{
+				FStackEntry& StackEntry = Stack.AddDefaulted_GetRef();
+				StackEntry.StartTime = ClampedTime;
+				StackEntry.RowIndex = Rows.Num();
+				FTempRow& Row = Rows.AddDefaulted_GetRef();
+				Row.StartTime = Time;
+				Row.ExportInfo = Event.Export;
+				Row.EventType = Event.EventType;
+				if ( Event.EventType == LoadTimeProfilerObjectEventType_Serialize)
+				{
+					Row.SerializedSize = Event.Export->SerialSize;
+				}
+			}
+			else
+			{
+				Stack.Pop(false);
+			}
 			return EEventEnumerate::Continue;
 		});
+		++TimelineIndex;
+	}
+
+	Algo::Sort(Rows, [](const FTempRow& A, const FTempRow& B)
+		{
+			return A.StartTime < B.StartTime;
+		});
+	TTable<FExportsTableRow>* Table = new TTable<FExportsTableRow>(ExportsTableLayout);
+	for (const FTempRow& TempRow : Rows)
+	{
+		FExportsTableRow& Row = Table->AddRow();
+		Row.ExportInfo = TempRow.ExportInfo;
+		Row.SerializedSize = TempRow.SerializedSize;
+		Row.MainThreadTime = TempRow.MainThreadTime;
+		Row.AsyncLoadingThreadTime = TempRow.AsyncLoadingThreadTime;
+		Row.EventType = TempRow.EventType;
 	}
+
+	return Table;
+}
+
+ITable<FRequestsTableRow>* FLoadTimeProfilerProvider::CreateRequestsTable(double IntervalStart, double IntervalEnd) const
+{
+	TTable<FRequestsTableRow>* Table = new TTable<FRequestsTableRow>(RequestsTableLayout);
+	for (const FLoadRequest& Request : Requests)
+	{
+		if (Request.EndTime < IntervalStart || Request.StartTime > IntervalEnd)
+		{
+			continue;
+		}
+		FRequestsTableRow& Row = Table->AddRow();
+		Row.Id = Request.Id;
+		Row.Name = Request.Name;
+		Row.StartTime = Request.StartTime;
+		Row.Duration = Request.EndTime - Request.StartTime;
+		Row.Packages = Request.Packages;
+	}
+
 	return Table;
 }
 
@@ -402,6 +537,7 @@ FLoadRequest& FLoadTimeProfilerProvider::CreateRequest()
 	Session.WriteAccessCheck();
 
 	FLoadRequest& RequestInfo = Requests.PushBack();
+	RequestInfo.Id = Requests.Num();
 	return RequestInfo;
 }
 
@@ -467,36 +603,6 @@ void FLoadTimeProfilerProvider::DistributeBytesAcrossFrames(uint64 ByteCount, do
 	check(TotalDistributed == OriginalByteCount);
 }
 
-uint64 FLoadTimeProfilerProvider::BeginLoadPackage(const FPackageInfo& PackageInfo, double Time)
-{
-	Session.WriteAccessCheck();
-
-	CreateCounters();
-
-	LoadingPackagesCounter->AddValue(Time, int64(1));
-
-	uint64 LoadHandle = PackageLoads.Num();
-	FPackageLoad& PackageLoad = PackageLoads.PushBack();
-	PackageLoad.Package = &PackageInfo;
-	PackageLoad.StartTime = Time;
-	return LoadHandle;
-}
-
-void FLoadTimeProfilerProvider::EndLoadPackage(uint64 LoadHandle, double Time)
-{
-	Session.WriteAccessCheck();
-
-	FPackageLoad& PackageLoad = PackageLoads[LoadHandle];
-	PackageLoad.EndTime = Time;
-	check(PackageLoad.EndTime >= PackageLoad.StartTime);
-
-	DistributeBytesAcrossFrames(PackageLoad.Package->Summary.TotalHeaderSize, PackageLoad.StartTime, PackageLoad.EndTime, &FLoaderFrame::HeaderLoadedBytes);
-	DistributeBytesAcrossFrames(PackageLoad.Package->TotalExportsSerialSize, PackageLoad.StartTime, PackageLoad.EndTime, &FLoaderFrame::ExportLoadedBytes);
-
-	TotalLoaderBytesLoadedCounter->AddValue(Time, int64(PackageLoad.Package->Summary.TotalHeaderSize + PackageLoad.Package->TotalExportsSerialSize));
-	LoadingPackagesCounter->AddValue(Time, int64(-1));
-}
-
 uint64 FLoadTimeProfilerProvider::BeginIoDispatcherBatch(uint64 BatchId, double Time)
 {
 	Session.WriteAccessCheck();
diff --git a/Engine/Source/Developer/TraceServices/Private/Model/LoadTimeProfilerPrivate.h b/Engine/Source/Developer/TraceServices/Private/Model/LoadTimeProfilerPrivate.h
--- a/Engine/Source/Developer/TraceServices/Private/Model/LoadTimeProfilerPrivate.h
+++ b/Engine/Source/Developer/TraceServices/Private/Model/LoadTimeProfilerPrivate.h
@@ -31,12 +31,11 @@ public:
 	virtual ITable<FLoadTimeProfilerAggregatedStats>* CreateObjectTypeAggregation(double IntervalStart, double IntervalEnd) const override;
 	virtual ITable<FPackagesTableRow>* CreatePackageDetailsTable(double IntervalStart, double IntervalEnd) const override;
 	virtual ITable<FExportsTableRow>* CreateExportDetailsTable(double IntervalStart, double IntervalEnd) const override;
+	virtual ITable<FRequestsTableRow>* CreateRequestsTable(double IntervalStart, double IntervalEnd) const override;
 	virtual const ITable<FLoadRequest>& GetRequestsTable() const override { return RequestsTable; }
 	const FClassInfo& AddClassInfo(const TCHAR* ClassName);
 	FLoadRequest& CreateRequest();
 	FPackageInfo& CreatePackage();
-	uint64 BeginLoadPackage(const FPackageInfo& PackageInfo, double Time);
-	void EndLoadPackage(uint64 LoadHandle, double Time);
 	FPackageExportInfo& CreateExport();
 	CpuTimelineInternal& EditCpuTimeline(uint32 ThreadId);
 	uint64 BeginIoDispatcherBatch(uint64 BatchId, double Time);
@@ -106,7 +105,6 @@ private:
 	TPagedArray<FClassInfo> ClassInfos;
 	TPagedArray<FLoadRequest> Requests;
 	TPagedArray<FPackageInfo> Packages;
-	TPagedArray<FPackageLoad> PackageLoads;
 	TPagedArray<FIoDispatcherBatch> IoDispatcherBatches;
 	TPagedArray<FPackageExportInfo> Exports;
 	TArray<TSharedRef<CpuTimelineInternal>> CpuTimelines;
@@ -115,6 +113,7 @@ private:
 	TTableLayout<FLoadTimeProfilerAggregatedStats> AggregatedStatsTableLayout;
 	TTableLayout<FPackagesTableRow> PackagesTableLayout;
 	TTableLayout<FExportsTableRow> ExportsTableLayout;
+	TTableLayout<FRequestsTableRow> RequestsTableLayout;
 	TPagedArray<FLoaderFrame> Frames;
 	bool bHasCreatedCounters = false;
 	IEditableCounter* ActiveIoDispatcherBatchesCounter = nullptr;
diff --git a/Engine/Source/Developer/TraceServices/Private/Modules/LoadTimeProfilerModule.cpp b/Engine/Source/Developer/TraceServices/Private/Modules/LoadTimeProfilerModule.cpp
--- a/Engine/Source/Developer/TraceServices/Private/Modules/LoadTimeProfilerModule.cpp
+++ b/Engine/Source/Developer/TraceServices/Private/Modules/LoadTimeProfilerModule.cpp
@@ -89,7 +89,8 @@ void FLoadTimeProfilerModule::GenerateReports(const IAnalysisSession& Session, c
 		Table2Csv(*PackagesTable.Get(), *(ReportDirectory / TEXT("Packages.csv")));
 		TUniquePtr<ITable<FExportsTableRow>> ExportsTable(LoadTimeProfiler->CreateExportDetailsTable(CaptureStartTime, CaptureEndTime));
 		Table2Csv(*ExportsTable.Get(), *(ReportDirectory / TEXT("Exports.csv")));
-		Table2Csv(LoadTimeProfiler->GetRequestsTable(), *(ReportDirectory / TEXT("Requests.csv")));
+		TUniquePtr<ITable<FRequestsTableRow>> RequestsTable(LoadTimeProfiler->CreateRequestsTable(CaptureStartTime, CaptureEndTime));
+		Table2Csv(*RequestsTable.Get(), *(ReportDirectory / TEXT("Requests.csv")));
 	}
 	const IFileActivityProvider* FileActivityProvider = ReadFileActivityProvider(Session);
 	if (FileActivityProvider)
diff --git a/Engine/Source/Developer/TraceServices/Public/TraceServices/Model/LoadTimeProfiler.h b/Engine/Source/Developer/TraceServices/Public/TraceServices/Model/LoadTimeProfiler.h
--- a/Engine/Source/Developer/TraceServices/Public/TraceServices/Model/LoadTimeProfiler.h
+++ b/Engine/Source/Developer/TraceServices/Public/TraceServices/Model/LoadTimeProfiler.h
@@ -88,12 +88,15 @@ struct FPackageInfo
 	uint32 Id;
 	const TCHAR* Name = nullptr;
 	FPackageSummaryInfo Summary;
+	TArray<const FPackageInfo*> ImportedPackages;
 	TArray<const FPackageExportInfo*> Exports;
 	uint64 TotalExportsSerialSize = 0;
+	uint64 RequestId = 0;
 };
 
 struct FLoadRequest
 {
+	uint64 Id = 0;
 	const TCHAR* Name = nullptr;
 	uint32 ThreadId = uint32(-1);
 	double StartTime = 0.0;
@@ -121,6 +124,7 @@ struct FExportsTableRow
 struct FPackagesTableRow
 {
 	const FPackageInfo* PackageInfo = nullptr;
+	uint64 TotalSerializedSize = 0;
 	uint64 SerializedHeaderSize = 0;
 	uint64 SerializedExportsCount = 0;
 	uint64 SerializedExportsSize = 0;
@@ -128,6 +132,15 @@ struct FPackagesTableRow
 	double AsyncLoadingThreadTime = 0.0;
 };
 
+struct FRequestsTableRow
+{
+	uint64 Id = 0;
+	const TCHAR* Name = nullptr;
+	double StartTime = 0.0;
+	double Duration = 0.0;
+	TArray<const FPackageInfo*> Packages;
+};
+
 struct FLoadTimeProfilerCpuEvent
 {
 	const FPackageInfo* Package = nullptr;
@@ -162,6 +175,7 @@ public:
 	virtual ITable<FLoadTimeProfilerAggregatedStats>* CreateObjectTypeAggregation(double IntervalStart, double IntervalEnd) const = 0;
 	virtual ITable<FPackagesTableRow>* CreatePackageDetailsTable(double IntervalStart, double IntervalEnd) const = 0;
 	virtual ITable<FExportsTableRow>* CreateExportDetailsTable(double IntervalStart, double IntervalEnd) const = 0;
+	virtual ITable<FRequestsTableRow>* CreateRequestsTable(double IntervalStart, double IntervalEnd) const = 0;
 	virtual const ITable<FLoadRequest>& GetRequestsTable() const = 0;
 };
 
diff --git a/Engine/Source/Editor/UnrealEd/Private/EditorServer.cpp b/Engine/Source/Editor/UnrealEd/Private/EditorServer.cpp
--- a/Engine/Source/Editor/UnrealEd/Private/EditorServer.cpp
+++ b/Engine/Source/Editor/UnrealEd/Private/EditorServer.cpp
@@ -173,6 +173,7 @@
 
 #include "Serialization/StructuredArchive.h"
 #include "Serialization/Formatters/JsonArchiveInputFormatter.h"
+#include "Serialization/LoadTimeTrace.h"
 
 DEFINE_LOG_CATEGORY_STATIC(LogEditorServer, Log, All);
 
@@ -2445,6 +2446,7 @@ bool UEditorEngine::PackageIsAMapFile( const TCHAR* PackageFilename, FText& OutN
 bool UEditorEngine::Map_Load(const TCHAR* Str, FOutputDevice& Ar)
 {
 	TRACE_CPUPROFILER_EVENT_SCOPE(UEditorEngine::Map_Load);
+	TRACE_LOADTIME_REQUEST_GROUP_SCOPE(TEXT("LoadMap - %s"), Str);
 
 	auto FindWorldInPackageOrFollowRedirector = [](UPackage*& InOutPackage)
 	{
diff --git a/Engine/Source/Runtime/CoreUObject/Private/Serialization/AsyncLoading.cpp b/Engine/Source/Runtime/CoreUObject/Private/Serialization/AsyncLoading.cpp
--- a/Engine/Source/Runtime/CoreUObject/Private/Serialization/AsyncLoading.cpp
+++ b/Engine/Source/Runtime/CoreUObject/Private/Serialization/AsyncLoading.cpp
@@ -1774,8 +1774,6 @@ FORCEINLINE static bool CanAddWaitingPackages(FAsyncLoadingThread& AsyncLoadingT
 
 void FAsyncLoadingThread::QueueEvent_CreateLinker(FAsyncPackage* Package, int32 EventSystemPriority)
 {
-	TRACE_LOADTIME_BEGIN_LOAD_ASYNC_PACKAGE(Package);
-
 	FileOpenLogActive();  // make sure GRandomizeLoadOrder is set up
 	check(Package);
 	Package->AddNode(EEventLoadNode::Package_LoadSummary);
@@ -1977,8 +1975,6 @@ void FAsyncPackage::Event_FinishLinker()
 			}
 		}
 
-		TRACE_LOADTIME_PACKAGE_SUMMARY(this, Desc.Name, Linker->Summary.TotalHeaderSize, Linker->Summary.ImportCount, Linker->Summary.ExportCount);
-
 		check(AsyncPackageLoadingState == EAsyncPackageLoadingState::WaitingForSummary);
 		AsyncPackageLoadingState = EAsyncPackageLoadingState::StartImportPackages;
 		AsyncLoadingThread.QueueEvent_StartImportPackages(this, FAsyncLoadEvent::EventSystemPriority_MAX - 1);
@@ -4821,8 +4817,6 @@ EAsyncPackageState::Type FAsyncLoadingThread::ProcessLoadedPackages(bool bUseTim
 
 				CompletedPackages.Add(Package);
 
-				TRACE_LOADTIME_END_LOAD_ASYNC_PACKAGE(Package);
-
 				if (IsTimeLimitExceeded(TickStartTime, bUseTimeLimit, TimeLimit, TEXT("ProcessLoadedPackages Misc")) || (FlushRequest && !ContainsRequestInternal(FlushRequest.GetId())))
 				{
 					// The only package we care about has finished loading, so we're good to exit
@@ -6322,6 +6316,7 @@ EAsyncPackageState::Type FAsyncPackage::CreateLinker()
 		Linker->AsyncRoot = this;
 
 		UE_LOG(LogStreaming, Verbose, TEXT("FAsyncPackage::CreateLinker for %s finished."), *Desc.PackagePath.GetDebugName());
+		TRACE_LOADTIME_ASYNC_PACKAGE_LINKER_ASSOCIATION(this, Linker);
 	}
 	return EAsyncPackageState::Complete;
 }
@@ -6896,6 +6891,7 @@ EAsyncPackageState::Type FAsyncPackage::PostLoadObjects()
 	}
 
 	// PostLoad objects.
+	TRACE_LOADTIME_POSTLOAD_SCOPE;
 	while (PostLoadIndex < PackageObjLoaded.Num() && PostLoadIndex < PreLoadIndex && !IsTimeLimitExceeded())
 	{
 		UObject* Object = PackageObjLoaded[PostLoadIndex++];
@@ -6916,10 +6912,7 @@ EAsyncPackageState::Type FAsyncPackage::PostLoadObjects()
 				check(!GEventDrivenLoaderEnabled || !Object->HasAnyFlags(RF_NeedLoad));
 
 				ThreadContext.CurrentlyPostLoadedObjectByALT = Object;
-				{
-					TRACE_LOADTIME_POSTLOAD_EXPORT_SCOPE(Object);
-					Object->ConditionalPostLoad();
-				}
+				Object->ConditionalPostLoad();
 				ThreadContext.CurrentlyPostLoadedObjectByALT = nullptr;
 
 				LastObjectWorkWasPerformedOn = Object;
@@ -6982,6 +6975,7 @@ EAsyncPackageState::Type FAsyncPackage::PostLoadDeferredObjects(double InTickSta
 	{
 		TGuardValue<bool> GuardIsRoutingPostLoad(PackageScope.ThreadContext.IsRoutingPostLoad, true);
 
+		TRACE_LOADTIME_POSTLOAD_SCOPE;
 		while (DeferredPostLoadIndex < DeferredPostLoadObjects.Num() &&
 			!AsyncLoadingThread.IsAsyncLoadingSuspendedInternal() &&
 			!::IsTimeLimitExceeded(InTickStartTime, bInUseTimeLimit, InOutTimeLimit, LastTypeOfWorkPerformed, LastObjectWorkWasPerformedOn))
@@ -7001,10 +6995,7 @@ EAsyncPackageState::Type FAsyncPackage::PostLoadDeferredObjects(double InTickSta
 			FScopeCycleCounterUObject ConstructorScope(Object, GET_STATID(STAT_FAsyncPackage_PostLoadObjectsGameThread));
 
 			PackageScope.ThreadContext.CurrentlyPostLoadedObjectByALT = Object;
-			{
-				TRACE_LOADTIME_POSTLOAD_EXPORT_SCOPE(Object);
-				Object->ConditionalPostLoad();
-			}
+			Object->ConditionalPostLoad();
 			PackageScope.ThreadContext.CurrentlyPostLoadedObjectByALT = nullptr;
 
 			if (ObjLoadedInPostLoad.Num())
diff --git a/Engine/Source/Runtime/CoreUObject/Private/Serialization/AsyncLoading2.cpp b/Engine/Source/Runtime/CoreUObject/Private/Serialization/AsyncLoading2.cpp
--- a/Engine/Source/Runtime/CoreUObject/Private/Serialization/AsyncLoading2.cpp
+++ b/Engine/Source/Runtime/CoreUObject/Private/Serialization/AsyncLoading2.cpp
@@ -4376,6 +4376,7 @@ void FAsyncPackage2::CreateLinker(const FLinkerInstancingContext* InstancingCont
 	check(Linker);
 	check(Linker->LinkerRoot == LinkerRoot);
 	check(!Linker->AsyncRoot);
+	TRACE_LOADTIME_ASYNC_PACKAGE_LINKER_ASSOCIATION(this, Linker);
 	Linker->AsyncRoot = this;
 	LinkerLoadState->Linker = Linker;
 #if ALT2_ENABLE_NEW_ARCHIVE_FOR_LINKERLOAD
@@ -4397,7 +4398,6 @@ void FAsyncPackage2::DetachLinker()
 void FAsyncPackage2::StartLoading(FIoBatch& IoBatch)
 {
 	TRACE_CPUPROFILER_EVENT_SCOPE(StartLoading);
-	TRACE_LOADTIME_BEGIN_LOAD_ASYNC_PACKAGE(this);
 
 	LoadStartTime = FPlatformTime::Seconds();
 
@@ -4556,6 +4556,7 @@ EEventLoadNodeExecutionResult FAsyncPackage2::ProcessLinkerLoadPackageSummary(FA
 	if (!LinkerLoadState->Linker->GetLoader()->IsError())
 	{
 		LinkerLoadState->Linker->bUseTimeLimit = false;
+		TRACE_LOADTIME_PROCESS_SUMMARY_SCOPE(LinkerLoadState->Linker);
 		LinkerResult = LinkerLoadState->Linker->ProcessPackageSummary(nullptr);
 	}
 #else
@@ -4900,6 +4901,8 @@ bool FAsyncPackage2::ProcessLinkerLoadPackageExports(FAsyncLoadingThreadState2&
 
 EEventLoadNodeExecutionResult FAsyncPackage2::PostLoadLinkerLoadPackageExports(FAsyncLoadingThreadState2& ThreadState)
 {
+	TRACE_LOADTIME_POSTLOAD_SCOPE;
+
 	FAsyncLoadingTickScope2 InAsyncLoadingTick(AsyncLoadingThread);
 
 	const int32 ObjectCount = ConstructedObjects.Num();
@@ -4944,6 +4947,7 @@ EEventLoadNodeExecutionResult FAsyncPackage2::Event_ProcessPackageSummary(FAsync
 	}
 	else
 	{
+		TRACE_LOADTIME_PROCESS_SUMMARY_SCOPE(Package);
 		check(Package->ExportBundleEntryIndex == 0);
 
 		ReadAsyncPackageHeader(Package->SerializationState, Package->HeaderData);
@@ -5959,6 +5963,7 @@ EEventLoadNodeExecutionResult FAsyncPackage2::Event_PostLoadExportBundle(FAsyncL
 		const TArrayView<FExportObject>& Exports = Package->Data.Exports;
 #endif
 
+		TRACE_LOADTIME_POSTLOAD_SCOPE;
 		while (Package->ExportBundleEntryIndex < int32(ExportBundle->EntryCount))
 		{
 			const FExportBundleEntry& BundleEntry = HeaderData->ExportBundleEntries[ExportBundle->FirstEntryIndex + Package->ExportBundleEntryIndex];
@@ -5990,10 +5995,7 @@ EEventLoadNodeExecutionResult FAsyncPackage2::Event_PostLoadExportBundle(FAsyncL
 					if (!bIsMultithreaded || (bAsyncPostLoadEnabled && CanPostLoadOnAsyncLoadingThread(Object)))
 					{
 						ThreadContext.CurrentlyPostLoadedObjectByALT = Object;
-						{
-							TRACE_LOADTIME_POSTLOAD_EXPORT_SCOPE(Object);
-							Object->ConditionalPostLoad();
-						}
+						Object->ConditionalPostLoad();
 						ThreadContext.CurrentlyPostLoadedObjectByALT = nullptr;
 					}
 				} while (false);
@@ -6087,6 +6089,7 @@ EEventLoadNodeExecutionResult FAsyncPackage2::Event_DeferredPostLoadExportBundle
 		const TArrayView<FExportObject>& Exports = Package->Data.Exports;
 #endif
 
+		TRACE_LOADTIME_POSTLOAD_SCOPE;
 		while (Package->ExportBundleEntryIndex < int32(ExportBundle->EntryCount))
 		{
 			const FExportBundleEntry& BundleEntry = HeaderData->ExportBundleEntries[ExportBundle->FirstEntryIndex + Package->ExportBundleEntryIndex];
@@ -6113,7 +6116,6 @@ EEventLoadNodeExecutionResult FAsyncPackage2::Event_DeferredPostLoadExportBundle
 					{
 						PackageScope.ThreadContext.CurrentlyPostLoadedObjectByALT = Object;
 						{
-							TRACE_LOADTIME_POSTLOAD_EXPORT_SCOPE(Object);
 							FScopeCycleCounterUObject ConstructorScope(Object, GET_STATID(STAT_FAsyncPackage_PostLoadObjectsGameThread));
 							Object->ConditionalPostLoad();
 						}
@@ -6447,7 +6449,6 @@ EAsyncPackageState::Type FAsyncLoadingThread2::ProcessLoadedPackagesFromGameThre
 			--LoadingPackagesCounter;
 
 			TRACE_COUNTER_SET(AsyncLoadingLoadingPackages, LoadingPackagesCounter);
-			TRACE_LOADTIME_END_LOAD_ASYNC_PACKAGE(Package);
 
 			PackagesReadyForCallback.Add(Package);
 		}
diff --git a/Engine/Source/Runtime/CoreUObject/Private/Serialization/LoadTimeTrace.cpp b/Engine/Source/Runtime/CoreUObject/Private/Serialization/LoadTimeTrace.cpp
--- a/Engine/Source/Runtime/CoreUObject/Private/Serialization/LoadTimeTrace.cpp
+++ b/Engine/Source/Runtime/CoreUObject/Private/Serialization/LoadTimeTrace.cpp
@@ -34,6 +34,15 @@ UE_TRACE_EVENT_BEGIN(LoadTime, PackageSummary)
 	UE_TRACE_EVENT_FIELD(UE::Trace::WideString, Name)
 UE_TRACE_EVENT_END()
 
+UE_TRACE_EVENT_BEGIN(LoadTime, BeginProcessSummary)
+	UE_TRACE_EVENT_FIELD(uint64, Cycle)
+	UE_TRACE_EVENT_FIELD(const void*, AsyncPackage)
+UE_TRACE_EVENT_END()
+
+UE_TRACE_EVENT_BEGIN(LoadTime, EndProcessSummary)
+	UE_TRACE_EVENT_FIELD(uint64, Cycle)
+UE_TRACE_EVENT_END()
+
 UE_TRACE_EVENT_BEGIN(LoadTime, BeginCreateExport)
 	UE_TRACE_EVENT_FIELD(uint64, Cycle)
 	UE_TRACE_EVENT_FIELD(const void*, AsyncPackage)
@@ -55,12 +64,18 @@ UE_TRACE_EVENT_BEGIN(LoadTime, EndSerializeExport)
 	UE_TRACE_EVENT_FIELD(uint64, Cycle)
 UE_TRACE_EVENT_END()
 
-UE_TRACE_EVENT_BEGIN(LoadTime, BeginPostLoadExport)
+UE_TRACE_EVENT_BEGIN(LoadTime, BeginPostLoad)
+UE_TRACE_EVENT_END()
+
+UE_TRACE_EVENT_BEGIN(LoadTime, EndPostLoad)
+UE_TRACE_EVENT_END()
+
+UE_TRACE_EVENT_BEGIN(LoadTime, BeginPostLoadObject)
 	UE_TRACE_EVENT_FIELD(uint64, Cycle)
 	UE_TRACE_EVENT_FIELD(const UObject*, Object)
 UE_TRACE_EVENT_END()
 
-UE_TRACE_EVENT_BEGIN(LoadTime, EndPostLoadExport)
+UE_TRACE_EVENT_BEGIN(LoadTime, EndPostLoadObject)
 	UE_TRACE_EVENT_FIELD(uint64, Cycle)
 UE_TRACE_EVENT_END()
 
@@ -78,18 +93,16 @@ UE_TRACE_EVENT_BEGIN(LoadTime, NewAsyncPackage)
 	UE_TRACE_EVENT_FIELD(const void*, AsyncPackage)
 UE_TRACE_EVENT_END()
 
-UE_TRACE_EVENT_BEGIN(LoadTime, BeginLoadAsyncPackage)
-	UE_TRACE_EVENT_FIELD(uint64, Cycle)
+UE_TRACE_EVENT_BEGIN(LoadTime, DestroyAsyncPackage)
 	UE_TRACE_EVENT_FIELD(const void*, AsyncPackage)
 UE_TRACE_EVENT_END()
 
-UE_TRACE_EVENT_BEGIN(LoadTime, EndLoadAsyncPackage)
-	UE_TRACE_EVENT_FIELD(uint64, Cycle)
-	UE_TRACE_EVENT_FIELD(const void*, AsyncPackage)
+UE_TRACE_EVENT_BEGIN(LoadTime, NewLinker)
+	UE_TRACE_EVENT_FIELD(const void*, Linker)
 UE_TRACE_EVENT_END()
 
-UE_TRACE_EVENT_BEGIN(LoadTime, DestroyAsyncPackage)
-	UE_TRACE_EVENT_FIELD(const void*, AsyncPackage)
+UE_TRACE_EVENT_BEGIN(LoadTime, DestroyLinker)
+	UE_TRACE_EVENT_FIELD(const void*, Linker)
 UE_TRACE_EVENT_END()
 
 UE_TRACE_EVENT_BEGIN(LoadTime, AsyncPackageRequestAssociation)
@@ -97,6 +110,11 @@ UE_TRACE_EVENT_BEGIN(LoadTime, AsyncPackageRequestAssociation)
 	UE_TRACE_EVENT_FIELD(uint64, RequestId)
 UE_TRACE_EVENT_END()
 
+UE_TRACE_EVENT_BEGIN(LoadTime, AsyncPackageLinkerAssociation)
+	UE_TRACE_EVENT_FIELD(const void*, AsyncPackage)
+	UE_TRACE_EVENT_FIELD(const void*, Linker)
+UE_TRACE_EVENT_END()
+
 UE_TRACE_EVENT_BEGIN(LoadTime, AsyncPackageImportDependency)
 	UE_TRACE_EVENT_FIELD(const void*, AsyncPackage)
 	UE_TRACE_EVENT_FIELD(const void*, ImportedAsyncPackage)
@@ -147,24 +165,22 @@ void FLoadTimeProfilerTracePrivate::OutputNewAsyncPackage(const void* AsyncPacka
 		<< NewAsyncPackage.AsyncPackage(AsyncPackage);
 }
 
-void FLoadTimeProfilerTracePrivate::OutputBeginLoadAsyncPackage(const void* AsyncPackage)
+void FLoadTimeProfilerTracePrivate::OutputDestroyAsyncPackage(const void* AsyncPackage)
 {
-	UE_TRACE_LOG(LoadTime, BeginLoadAsyncPackage, LoadTimeChannel)
-		<< BeginLoadAsyncPackage.Cycle(FPlatformTime::Cycles64())
-		<< BeginLoadAsyncPackage.AsyncPackage(AsyncPackage);
+	UE_TRACE_LOG(LoadTime, DestroyAsyncPackage, LoadTimeChannel)
+		<< DestroyAsyncPackage.AsyncPackage(AsyncPackage);
 }
 
-void FLoadTimeProfilerTracePrivate::OutputEndLoadAsyncPackage(const void* AsyncPackage)
+void FLoadTimeProfilerTracePrivate::OutputNewLinker(const void* Linker)
 {
-	UE_TRACE_LOG(LoadTime, EndLoadAsyncPackage, LoadTimeChannel)
-		<< EndLoadAsyncPackage.Cycle(FPlatformTime::Cycles64())
-		<< EndLoadAsyncPackage.AsyncPackage(AsyncPackage);
+	UE_TRACE_LOG(LoadTime, NewLinker, LoadTimeChannel)
+		<< NewLinker.Linker(Linker);
 }
 
-void FLoadTimeProfilerTracePrivate::OutputDestroyAsyncPackage(const void* AsyncPackage)
+void FLoadTimeProfilerTracePrivate::OutputDestroyLinker(const void* Linker)
 {
-	UE_TRACE_LOG(LoadTime, DestroyAsyncPackage, LoadTimeChannel)
-		<< DestroyAsyncPackage.AsyncPackage(AsyncPackage);
+	UE_TRACE_LOG(LoadTime, DestroyLinker, LoadTimeChannel)
+		<< DestroyLinker.Linker(Linker);
 }
 
 void FLoadTimeProfilerTracePrivate::OutputPackageSummary(const void* AsyncPackage, const FName& PackageName, uint32 TotalHeaderSize, uint32 ImportCount, uint32 ExportCount)
@@ -187,6 +203,13 @@ void FLoadTimeProfilerTracePrivate::OutputAsyncPackageRequestAssociation(const v
 		<< AsyncPackageRequestAssociation.RequestId(RequestId);
 }
 
+void FLoadTimeProfilerTracePrivate::OutputAsyncPackageLinkerAssociation(const void* AsyncPackage, const void* Linker)
+{
+	UE_TRACE_LOG(LoadTime, AsyncPackageLinkerAssociation, LoadTimeChannel)
+		<< AsyncPackageLinkerAssociation.AsyncPackage(AsyncPackage)
+		<< AsyncPackageLinkerAssociation.Linker(Linker);
+}
+
 void FLoadTimeProfilerTracePrivate::OutputAsyncPackageImportDependency(const void* Package, const void* ImportedPackage)
 {
 	UE_TRACE_LOG(LoadTime, AsyncPackageImportDependency, LoadTimeChannel)
@@ -211,6 +234,29 @@ void FLoadTimeProfilerTracePrivate::OutputClassInfo(const UClass* Class, const T
 		<< ClassInfo.Name(Name, NameLen);
 }
 
+void FLoadTimeProfilerTracePrivate::OutputBeginProcessSummary(const void* AsyncPackage)
+{
+	UE_TRACE_LOG(LoadTime, BeginProcessSummary, LoadTimeChannel)
+		<< BeginProcessSummary.Cycle(FPlatformTime::Cycles64())
+		<< BeginProcessSummary.AsyncPackage(AsyncPackage);
+}
+
+void FLoadTimeProfilerTracePrivate::OutputEndProcessSummary()
+{
+	UE_TRACE_LOG(LoadTime, EndProcessSummary, LoadTimeChannel)
+		<< EndProcessSummary.Cycle(FPlatformTime::Cycles64());
+}
+
+FLoadTimeProfilerTracePrivate::FProcessSummaryScope::FProcessSummaryScope(const void* AsyncPackage)
+{
+	FLoadTimeProfilerTracePrivate::OutputBeginProcessSummary(AsyncPackage);
+}
+
+FLoadTimeProfilerTracePrivate::FProcessSummaryScope::~FProcessSummaryScope()
+{
+	FLoadTimeProfilerTracePrivate::OutputEndProcessSummary();
+}
+
 FLoadTimeProfilerTracePrivate::FCreateExportScope::FCreateExportScope(const void* AsyncPackage, const UObject* const* InObject)
 	: Object(InObject)
 {
@@ -241,17 +287,27 @@ FLoadTimeProfilerTracePrivate::FSerializeExportScope::~FSerializeExportScope()
 		<< EndSerializeExport.Cycle(FPlatformTime::Cycles64());
 }
 
-FLoadTimeProfilerTracePrivate::FPostLoadExportScope::FPostLoadExportScope(const UObject* Object)
+FLoadTimeProfilerTracePrivate::FPostLoadScope::FPostLoadScope()
+{
+	UE_TRACE_LOG(LoadTime, BeginPostLoad, LoadTimeChannel);
+}
+
+FLoadTimeProfilerTracePrivate::FPostLoadScope::~FPostLoadScope()
+{
+	UE_TRACE_LOG(LoadTime, EndPostLoad, LoadTimeChannel);
+}
+
+FLoadTimeProfilerTracePrivate::FPostLoadObjectScope::FPostLoadObjectScope(const UObject* Object)
 {
-	UE_TRACE_LOG(LoadTime, BeginPostLoadExport, LoadTimeChannel)
-		<< BeginPostLoadExport.Cycle(FPlatformTime::Cycles64())
-		<< BeginPostLoadExport.Object(Object);
+	UE_TRACE_LOG(LoadTime, BeginPostLoadObject, LoadTimeChannel)
+		<< BeginPostLoadObject.Cycle(FPlatformTime::Cycles64())
+		<< BeginPostLoadObject.Object(Object);
 }
 
-FLoadTimeProfilerTracePrivate::FPostLoadExportScope::~FPostLoadExportScope()
+FLoadTimeProfilerTracePrivate::FPostLoadObjectScope::~FPostLoadObjectScope()
 {
-	UE_TRACE_LOG(LoadTime, EndPostLoadExport, LoadTimeChannel)
-		<< EndPostLoadExport.Cycle(FPlatformTime::Cycles64());
+	UE_TRACE_LOG(LoadTime, EndPostLoadObject, LoadTimeChannel)
+		<< EndPostLoadObject.Cycle(FPlatformTime::Cycles64());
 }
 
 #endif
diff --git a/Engine/Source/Runtime/CoreUObject/Private/Serialization/LoadTimeTracePrivate.h b/Engine/Source/Runtime/CoreUObject/Private/Serialization/LoadTimeTracePrivate.h
--- a/Engine/Source/Runtime/CoreUObject/Private/Serialization/LoadTimeTracePrivate.h
+++ b/Engine/Source/Runtime/CoreUObject/Private/Serialization/LoadTimeTracePrivate.h
@@ -15,16 +15,25 @@ struct FLoadTimeProfilerTracePrivate
 	static void OutputSuspendAsyncLoading();
 	static void OutputResumeAsyncLoading();
 	static void OutputNewAsyncPackage(const void* AsyncPackage);
-	static void OutputBeginLoadAsyncPackage(const void* AsyncPackage);
-	static void OutputEndLoadAsyncPackage(const void* AsyncPackage);
 	static void OutputDestroyAsyncPackage(const void* AsyncPackage);
+	static void OutputNewLinker(const void* Linker);
+	static void OutputDestroyLinker(const void* Linker);
 	static void OutputBeginRequest(uint64 RequestId);
 	static void OutputEndRequest(uint64 RequestId);
 	static void OutputPackageSummary(const void* AsyncPackage, const FName& PackageName, uint32 TotalHeaderSize, uint32 ImportCount, uint32 ExportCount);
 	static void OutputAsyncPackageImportDependency(const void* Package, const void* ImportedPackage);
 	static void OutputAsyncPackageRequestAssociation(const void* AsyncPackage, uint64 RequestId);
+	static void OutputAsyncPackageLinkerAssociation(const void* AsyncPackage, const void* Linker);
 	static void OutputClassInfo(const UClass* Class, const FName& Name);
 	static void OutputClassInfo(const UClass* Class, const TCHAR* Name);
+	static void OutputBeginProcessSummary(const void* AsyncPackage);
+	static void OutputEndProcessSummary();
+
+	struct FProcessSummaryScope
+	{
+		FProcessSummaryScope(const void* AsyncPackage);
+		~FProcessSummaryScope();
+	};
 
 	struct FCreateExportScope
 	{
@@ -41,12 +50,17 @@ struct FLoadTimeProfilerTracePrivate
 		~FSerializeExportScope();
 	};
 
-	struct FPostLoadExportScope
+	struct FPostLoadScope
 	{
-		FPostLoadExportScope(const UObject* Object);
-		~FPostLoadExportScope();
+		FPostLoadScope();
+		~FPostLoadScope();
 	};
 
+	struct FPostLoadObjectScope
+	{
+		FPostLoadObjectScope(const UObject* Object);
+		~FPostLoadObjectScope();
+	};
 };
 
 #define TRACE_LOADTIME_START_ASYNC_LOADING() \
@@ -67,35 +81,47 @@ struct FLoadTimeProfilerTracePrivate
 #define TRACE_LOADTIME_NEW_ASYNC_PACKAGE(AsyncPackage) \
 	FLoadTimeProfilerTracePrivate::OutputNewAsyncPackage(AsyncPackage)
 
-#define TRACE_LOADTIME_BEGIN_LOAD_ASYNC_PACKAGE(AsyncPackage) \
-	FLoadTimeProfilerTracePrivate::OutputBeginLoadAsyncPackage(AsyncPackage)
-
-#define TRACE_LOADTIME_END_LOAD_ASYNC_PACKAGE(AsyncPackage) \
-	FLoadTimeProfilerTracePrivate::OutputEndLoadAsyncPackage(AsyncPackage)
-
 #define TRACE_LOADTIME_DESTROY_ASYNC_PACKAGE(AsyncPackage) \
 	FLoadTimeProfilerTracePrivate::OutputDestroyAsyncPackage(AsyncPackage);
 
+#define TRACE_LOADTIME_NEW_LINKER(Linker) \
+	FLoadTimeProfilerTracePrivate::OutputNewLinker(Linker)
+
+#define TRACE_LOADTIME_DESTROY_LINKER(Linker) \
+	FLoadTimeProfilerTracePrivate::OutputDestroyLinker(Linker);
+
 #define TRACE_LOADTIME_PACKAGE_SUMMARY(AsyncPackage, PackageName, TotalHeaderSize, ImportCount, ExportCount) \
 	FLoadTimeProfilerTracePrivate::OutputPackageSummary(AsyncPackage, PackageName, TotalHeaderSize, ImportCount, ExportCount);
 
 #define TRACE_LOADTIME_ASYNC_PACKAGE_REQUEST_ASSOCIATION(AsyncPackage, RequestId) \
 	FLoadTimeProfilerTracePrivate::OutputAsyncPackageRequestAssociation(AsyncPackage, RequestId);
 
-#define TRACE_LOADTIME_ASYNC_PACKAGE_LINKER_ASSOCIATION(AsyncPackage, Linker) \
+#define TRACE_LOADTIME_ASYNC_PACKAGE_LINKER_ASSOCIATION(AsyncPackage, LinkerLinker) \
 	FLoadTimeProfilerTracePrivate::OutputAsyncPackageLinkerAssociation(AsyncPackage, Linker);
 
 #define TRACE_LOADTIME_ASYNC_PACKAGE_IMPORT_DEPENDENCY(AsyncPackage, ImportedAsyncPackage) \
 	FLoadTimeProfilerTracePrivate::OutputAsyncPackageImportDependency(AsyncPackage, ImportedAsyncPackage);
 
+#define TRACE_LOADTIME_PROCESS_SUMMARY_SCOPE(AsyncPackage) \
+	FLoadTimeProfilerTracePrivate::FProcessSummaryScope __LoadTimeTraceProcessSummaryScope(AsyncPackage);
+
+#define TRACE_LOADTIME_BEGIN_PROCESS_SUMMARY(AsyncPackage) \
+	FLoadTimeProfilerTracePrivate::OutputBeginProcessSummary(AsyncPackage);
+
+#define TRACE_LOADTIME_END_PROCESS_SUMMARY \
+	FLoadTimeProfilerTracePrivate::OutputEndProcessSummary();
+
 #define TRACE_LOADTIME_CREATE_EXPORT_SCOPE(AsyncPackage, Object) \
 	FLoadTimeProfilerTracePrivate::FCreateExportScope __LoadTimeTraceCreateExportScope(AsyncPackage, Object);
 
 #define TRACE_LOADTIME_SERIALIZE_EXPORT_SCOPE(Object, SerialSize) \
 	FLoadTimeProfilerTracePrivate::FSerializeExportScope __LoadTimeTraceSerializeExportScope(Object, SerialSize);
 
-#define TRACE_LOADTIME_POSTLOAD_EXPORT_SCOPE(Object) \
-	FLoadTimeProfilerTracePrivate::FPostLoadExportScope __LoadTimeTracePostLoadExportScope(Object);
+#define TRACE_LOADTIME_POSTLOAD_SCOPE \
+	FLoadTimeProfilerTracePrivate::FPostLoadScope __LoadTimeTracePostLoadScope;
+
+#define TRACE_LOADTIME_POSTLOAD_OBJECT_SCOPE(Object) \
+	FLoadTimeProfilerTracePrivate::FPostLoadObjectScope __LoadTimeTracePostLoadObjectScope(Object);
 
 #define TRACE_LOADTIME_CLASS_INFO(Class, Name) \
 	FLoadTimeProfilerTracePrivate::OutputClassInfo(Class, Name);
@@ -108,15 +134,20 @@ struct FLoadTimeProfilerTracePrivate
 #define TRACE_LOADTIME_BEGIN_REQUEST(...)
 #define TRACE_LOADTIME_END_REQUEST(...)
 #define TRACE_LOADTIME_NEW_ASYNC_PACKAGE(...)
-#define TRACE_LOADTIME_BEGIN_LOAD_ASYNC_PACKAGE(...)
-#define TRACE_LOADTIME_END_LOAD_ASYNC_PACKAGE(...)
 #define TRACE_LOADTIME_DESTROY_ASYNC_PACKAGE(...)
+#define TRACE_LOADTIME_NEW_LINKER(...)
+#define TRACE_LOADTIME_DESTROY_LINKER(...)
 #define TRACE_LOADTIME_PACKAGE_SUMMARY(...)
 #define TRACE_LOADTIME_ASYNC_PACKAGE_REQUEST_ASSOCIATION(...)
+#define TRACE_LOADTIME_ASYNC_PACKAGE_LINKER_ASSOCIATION(...)
 #define TRACE_LOADTIME_ASYNC_PACKAGE_IMPORT_DEPENDENCY(...)
+#define TRACE_LOADTIME_PROCESS_SUMMARY_SCOPE(...)
+#define TRACE_LOADTIME_BEGIN_PROCESS_SUMMARY(...)
+#define TRACE_LOADTIME_END_PROCESS_SUMMARY
 #define TRACE_LOADTIME_CREATE_EXPORT_SCOPE(...)
 #define TRACE_LOADTIME_SERIALIZE_EXPORT_SCOPE(...)
-#define TRACE_LOADTIME_POSTLOAD_EXPORT_SCOPE(...)
+#define TRACE_LOADTIME_POSTLOAD_SCOPE
+#define TRACE_LOADTIME_POSTLOAD_OBJECT_SCOPE(...)
 #define TRACE_LOADTIME_CLASS_INFO(...)
 
 #endif
diff --git a/Engine/Source/Runtime/CoreUObject/Private/UObject/LinkerLoad.cpp b/Engine/Source/Runtime/CoreUObject/Private/UObject/LinkerLoad.cpp
--- a/Engine/Source/Runtime/CoreUObject/Private/UObject/LinkerLoad.cpp
+++ b/Engine/Source/Runtime/CoreUObject/Private/UObject/LinkerLoad.cpp
@@ -703,6 +703,8 @@ FUObjectSerializeContext* FLinkerLoad::GetSerializeContext()
 
 FLinkerLoad::ELinkerStatus FLinkerLoad::ProcessPackageSummary(TMap<TPair<FName, FPackageIndex>, FPackageIndex>* ObjectNameWithOuterToExportMap)
 {
+	TRACE_LOADTIME_BEGIN_PROCESS_SUMMARY(this);
+
 	ELinkerStatus Status = LINKER_Loaded;
 	{
 		SCOPED_LOADTIMER(LinkerLoad_SerializePackageFileSummary);
@@ -815,6 +817,8 @@ FLinkerLoad::ELinkerStatus FLinkerLoad::ProcessPackageSummary(TMap<TPair<FName,
 		Status = SerializePreloadDependencies();
 	}
 
+	TRACE_LOADTIME_END_PROCESS_SUMMARY;
+
 	// Finalize creation process.
 	if( Status == LINKER_Loaded )
 	{
@@ -972,10 +976,14 @@ FLinkerLoad::FLinkerLoad(UPackage* InParent, const FPackagePath& InPackagePath,
 	{
 		InstancingContext.AddPackageMapping(PackageNameToLoad, LinkerRoot->GetFName());
 	}
+
+	TRACE_LOADTIME_NEW_LINKER(this);
 }
 
 FLinkerLoad::~FLinkerLoad()
 {
+	TRACE_LOADTIME_DESTROY_LINKER(this);
+
 #if !UE_BUILD_SHIPPING && !UE_BUILD_TEST
 	FLinkerManager::Get().RemoveLiveLinker(this);
 #endif
@@ -2551,6 +2559,11 @@ FLinkerLoad::ELinkerStatus FLinkerLoad::FinalizeCreation(TMap<TPair<FName, FPack
 		}
 
 
+		if (LinkerRoot)
+		{
+			TRACE_LOADTIME_PACKAGE_SUMMARY(this, LinkerRoot->GetFName(), Summary.TotalHeaderSize, Summary.ImportCount, Summary.ExportCount);
+		}
+
 		// Avoid duplicate work in the case of async linker creation.
 		bHasFinishedInitialization = true;
 
@@ -4329,6 +4342,7 @@ void FLinkerLoad::Preload( UObject* Object )
 
 				{
 					SCOPE_CYCLE_COUNTER(STAT_LinkerSerialize);
+					TRACE_LOADTIME_SERIALIZE_EXPORT_SCOPE(Object, Export.SerialSize);
 #if USE_CIRCULAR_DEPENDENCY_LOAD_DEFERRING
 					// communicate with FLinkerPlaceholderBase, what object is currently serializing in
 					FScopedPlaceholderContainerTracker SerializingObjTracker(Object);
@@ -5141,7 +5155,10 @@ UObject* FLinkerLoad::CreateExport( int32 Index )
 			Params.ExternalPackage->SetPackageFlags(ThisParent->GetPackage()->GetPackageFlags() & PKG_PlayInEditor);
 		}
 
-		Export.Object = StaticConstructObject_Internal(Params);
+		{
+			TRACE_LOADTIME_CREATE_EXPORT_SCOPE(this, &Export.Object);
+			Export.Object = StaticConstructObject_Internal(Params);
+		}
 
 		if (FPlatformProperties::RequiresCookedData())
 		{
diff --git a/Engine/Source/Runtime/CoreUObject/Private/UObject/Obj.cpp b/Engine/Source/Runtime/CoreUObject/Private/UObject/Obj.cpp
--- a/Engine/Source/Runtime/CoreUObject/Private/UObject/Obj.cpp
+++ b/Engine/Source/Runtime/CoreUObject/Private/UObject/Obj.cpp
@@ -65,6 +65,7 @@
 #include "Misc/PackageAccessTrackingOps.h"
 #include "ProfilingDebugging/AssetMetadataTrace.h"
 #include "Containers/VersePath.h"
+#include "Serialization/LoadTimeTracePrivate.h"
 
 DEFINE_LOG_CATEGORY(LogObj);
 
@@ -1113,7 +1114,8 @@ void UObject::ConditionalPostLoad()
 				LLM_SCOPED_TAG_WITH_OBJECT_IN_SET(Package, ELLMTagSet::Assets);
 				LLM_SCOPED_TAG_WITH_OBJECT_IN_SET(GetClass(), ELLMTagSet::AssetClasses);
 				UE_TRACE_METADATA_SCOPE_ASSET(Package, GetClass());
-
+				TRACE_LOADTIME_POSTLOAD_OBJECT_SCOPE(this);
+				
 				PostLoad();
 
 				LLM_PUSH_STATS_FOR_ASSET_TAGS();
diff --git a/Engine/Source/Runtime/CoreUObject/Private/UObject/UObjectGlobals.cpp b/Engine/Source/Runtime/CoreUObject/Private/UObject/UObjectGlobals.cpp
--- a/Engine/Source/Runtime/CoreUObject/Private/UObject/UObjectGlobals.cpp
+++ b/Engine/Source/Runtime/CoreUObject/Private/UObject/UObjectGlobals.cpp
@@ -1535,6 +1535,8 @@ UPackage* LoadPackageInternal(UPackage* InOuter, const FPackagePath& PackagePath
 		FCoreDelegates::OnSyncLoadPackage.Broadcast(PackagePath.GetPackageNameOrFallback());
 	}
 	
+	TRACE_LOADTIME_POSTLOAD_SCOPE;
+
 	// Set up a load context
 	TRefCountPtr<FUObjectSerializeContext> LoadContext = FUObjectThreadContext::Get().GetSerializeContext();
 
@@ -1543,6 +1545,18 @@ UPackage* LoadPackageInternal(UPackage* InOuter, const FPackagePath& PackagePath
 	// Try to load.
 	BeginLoad(LoadContext, *PackagePath.GetDebugName());
 
+	if (!ImportLinker)
+	{
+		TRACE_LOADTIME_BEGIN_REQUEST(0);
+	}
+	ON_SCOPE_EXIT
+	{
+		if (!ImportLinker)
+		{
+			TRACE_LOADTIME_END_REQUEST(0);
+		}
+	};
+
 	bool bFullyLoadSkipped = false;
 
 	if (SlowTask)
@@ -1574,6 +1588,14 @@ UPackage* LoadPackageInternal(UPackage* InOuter, const FPackagePath& PackagePath
 		{
 			FUObjectSerializeContext* InOutLoadContext = LoadContext;
 			Linker = GetPackageLinker(InOuter, PackagePath, LoadFlags, nullptr, InReaderOverride, &InOutLoadContext, ImportLinker, InstancingContext);
+			if (ImportLinker)
+			{
+				TRACE_LOADTIME_ASYNC_PACKAGE_IMPORT_DEPENDENCY(ImportLinker, Linker);
+			}
+			else
+			{
+				TRACE_LOADTIME_ASYNC_PACKAGE_REQUEST_ASSOCIATION(Linker, 0);
+			}
 			if (InOutLoadContext != LoadContext && InOutLoadContext)
 			{
 				// The linker already existed and was associated with another context
@@ -1853,6 +1875,7 @@ UPackage* LoadPackage(UPackage* InOuter, const FPackagePath& PackagePath, uint32
 											FDynamicStats::CreateMemoryStatId<FStatGroup_STATGROUP_LLMAssets>(FName(*(FString(TEXT("Package ")) + PackagePath.GetPackageNameOrFallback()))).GetName() :
 											NAME_None,
 										 ELLMTagSet::Assets, ELLMTracker::Default);
+	TRACE_LOADTIME_REQUEST_GROUP_SCOPE(TEXT("SyncLoad - %s"), *PackagePath.GetDebugName());
 	return LoadPackageInternal(InOuter, PackagePath, LoadFlags, /*ImportLinker =*/ nullptr, InReaderOverride, InstancingContext, DiffPackagePath);
 }
 
