[Insights]
  - Memory Insights: Fixed the case where the free event is first event in a new SbTree column after adding a fake alloc.
  - Memory Insights: Added number of heap allocs in the "Analysis Completed" log message.

diff --git a/Engine/Source/Developer/TraceServices/Private/Analyzers/AllocationsAnalysis.cpp b/Engine/Source/Developer/TraceServices/Private/Analyzers/AllocationsAnalysis.cpp
--- a/Engine/Source/Developer/TraceServices/Private/Analyzers/AllocationsAnalysis.cpp
+++ b/Engine/Source/Developer/TraceServices/Private/Analyzers/AllocationsAnalysis.cpp
@@ -319,7 +319,7 @@ bool FAllocationsAnalyzer::OnEvent(uint16 RouteId, EStyle Style, const FOnEventC
 #endif
 				}
 			}
-			else // EStyle::LeaveScope
+			else if (ensure(Style == EStyle::LeaveScope))
 			{
 				if (RouteId == RouteId_MemScopeTag) // "MemoryScope"
 				{
diff --git a/Engine/Source/Developer/TraceServices/Private/Model/AllocationsProvider.cpp b/Engine/Source/Developer/TraceServices/Private/Model/AllocationsProvider.cpp
--- a/Engine/Source/Developer/TraceServices/Private/Model/AllocationsProvider.cpp
+++ b/Engine/Source/Developer/TraceServices/Private/Model/AllocationsProvider.cpp
@@ -1444,6 +1444,7 @@ void FAllocationsProvider::EditFree(double Time, uint32 CallstackId, uint64 Addr
 		constexpr uint64 FakeAllocSize = 0;
 		constexpr uint32 FakeAllocAlignment = 0;
 		EditAlloc(Time, CallstackId, Address, FakeAllocSize, FakeAllocAlignment, RootHeap);
+		SbTree[RootHeap]->SetTimeForEvent(EventIndex[RootHeap], Time); // for the case where the free event is first event in a new SbTree column after adding the fake alloc
 		AllocationPtr = LiveAllocs[RootHeap]->Remove(Address); // we take ownership of AllocationPtr
 	}
 	if (AllocationPtr)
@@ -1587,6 +1588,8 @@ void FAllocationsProvider::EditMarkAllocationAsHeap(double Time, uint64 Address,
 			UE_LOG(LogTraceServices, Error, TEXT("[MemAlloc] HeapMarkAlloc: Could not find address 0x%llX (Heap=%u, Flags=%u, Time=%f)!"), Address, Heap, uint32(Flags), Time);
 		}
 	}
+
+	++HeapCount;
 }
 
 ////////////////////////////////////////////////////////////////////////////////////////////////////
@@ -1827,7 +1830,7 @@ void FAllocationsProvider::EditOnAnalysisCompleted(double Time)
 	{
 		TotalEventCount += EventIndex[RootHeap];
 	}
-	UE_LOG(LogTraceServices, Log, TEXT("[MemAlloc] Analysis Completed (%llu events, %llu allocs, %llu frees)"), TotalEventCount, AllocCount, FreeCount);
+	UE_LOG(LogTraceServices, Log, TEXT("[MemAlloc] Analysis Completed (%llu events, %llu allocs, %llu frees, %llu heaps)"), TotalEventCount, AllocCount, FreeCount, HeapCount);
 }
 
 ////////////////////////////////////////////////////////////////////////////////////////////////////
diff --git a/Engine/Source/Developer/TraceServices/Private/Model/AllocationsProvider.h b/Engine/Source/Developer/TraceServices/Private/Model/AllocationsProvider.h
--- a/Engine/Source/Developer/TraceServices/Private/Model/AllocationsProvider.h
+++ b/Engine/Source/Developer/TraceServices/Private/Model/AllocationsProvider.h
@@ -433,6 +433,7 @@ private:
 
 	uint64 AllocCount = 0;
 	uint64 FreeCount = 0;
+	uint64 HeapCount = 0;
 
 	uint64 MiscErrors = 0;
 	uint64 HeapErrors = 0;
