Asset Loading Insights: Filter out rows with missing packages from the event aggregation table to match the other tables

diff --git a/Engine/Source/Developer/TraceServices/Private/Model/LoadTimeProfiler.cpp b/Engine/Source/Developer/TraceServices/Private/Model/LoadTimeProfiler.cpp
--- a/Engine/Source/Developer/TraceServices/Private/Model/LoadTimeProfiler.cpp
+++ b/Engine/Source/Developer/TraceServices/Private/Model/LoadTimeProfiler.cpp
@@ -266,8 +266,13 @@ ITable<FLoadTimeProfilerAggregatedStats>* FLoadTimeProfilerProvider::CreateEvent
 	TTable<FLoadTimeProfilerAggregatedStats>* Table = new TTable<FLoadTimeProfilerAggregatedStats>(AggregatedStatsTableLayout);
 	for (const auto& KV : Aggregation)
 	{
+		const FPackageInfo* Package = KV.Key;
+		if (!Package)
+		{
+			continue;
+		}
 		FLoadTimeProfilerAggregatedStats& Row = Table->AddRow();
-		Row.Name = KV.Key ? KV.Key->Name : TEXT("[unknown]");
+		Row.Name = Package->Name;
 		const FAggregatedTimingStats& Stats = KV.Value;
 		Row.Count = Stats.InstanceCount;
 		Row.Total = Stats.TotalExclusiveTime;
