[Insights] Added detection for pending events (with metadata) older than regular cpu timing events. Time needs to increase monotonically for cpu timing events. So, when this happens, we further ignore all pending events on this thread. The issue can occur in late connect trace sessions with trace protocol <= 6 (i.e. the scoped events have relative timestamps).

diff --git a/Engine/Source/Developer/TraceServices/Private/Analyzers/CpuProfilerTraceAnalysis.cpp b/Engine/Source/Developer/TraceServices/Private/Analyzers/CpuProfilerTraceAnalysis.cpp
--- a/Engine/Source/Developer/TraceServices/Private/Analyzers/CpuProfilerTraceAnalysis.cpp
+++ b/Engine/Source/Developer/TraceServices/Private/Analyzers/CpuProfilerTraceAnalysis.cpp
@@ -16,8 +16,6 @@
 #define CPUPROFILER_DEBUG_BEGIN_EVENT(Time, Event) { ++TotalScopeCount; }
 #define CPUPROFILER_DEBUG_END_EVENT(Time)
 
-#define CPUPROFILER_SAFE_DISPATCH_PENDING_EVENTS 1
-
 namespace TraceServices
 {
 
@@ -167,7 +165,6 @@ bool FCpuProfilerAnalyzer::OnEvent(uint16 RouteId, EStyle Style, const FOnEventC
 			break;
 		}
 
-#if CPUPROFILER_SAFE_DISPATCH_PENDING_EVENTS
 		int32 RemainingPending = ThreadState.PendingEvents.Num();
 		if (RemainingPending > 0)
 		{
@@ -177,7 +174,6 @@ bool FCpuProfilerAnalyzer::OnEvent(uint16 RouteId, EStyle Style, const FOnEventC
 			check(RemainingPending == 0);
 			ThreadState.PendingEvents.Reset();
 		}
-#endif
 
 		ensure(ThreadState.LastCycle != 0 || ThreadState.ScopeStack.Num() == 0);
 		if (ThreadState.LastCycle != 0)
@@ -333,7 +329,6 @@ uint64 FCpuProfilerAnalyzer::ProcessBuffer(const FEventTime& EventTime, FThreadS
 	}
 	check(BufferPtr == BufferEnd);
 
-#if CPUPROFILER_SAFE_DISPATCH_PENDING_EVENTS
 	if (RemainingPending == 0)
 	{
 		//CPUPROFILER_DEBUG_LOGF(TEXT("[%u] MetaEvents: %d added\n"), ThreadState.ThreadId, ThreadState.PendingEvents.Num());
@@ -345,14 +340,6 @@ uint64 FCpuProfilerAnalyzer::ProcessBuffer(const FEventTime& EventTime, FThreadS
 		CPUPROFILER_DEBUG_LOGF(TEXT("[%u] MetaEvents: %d added, %d still pending\n"), ThreadState.ThreadId, NumEventsToRemove, RemainingPending);
 		ThreadState.PendingEvents.RemoveAt(0, NumEventsToRemove);
 	}
-#else
-	if (RemainingPending > 0)
-	{
-		DispatchPendingEvents(LastCycle, ~0ull, EventTime, ThreadState, PendingCursor, RemainingPending);
-		check(RemainingPending == 0);
-		ThreadState.PendingEvents.Reset();
-	}
-#endif
 
 	ThreadState.LastCycle = LastCycle;
 	return LastCycle;
@@ -542,7 +529,6 @@ uint64 FCpuProfilerAnalyzer::ProcessBufferV2(const FEventTime& EventTime, FThrea
 	}
 	check(BufferPtr == BufferEnd);
 
-#if CPUPROFILER_SAFE_DISPATCH_PENDING_EVENTS
 	if (RemainingPending == 0)
 	{
 		//CPUPROFILER_DEBUG_LOGF(TEXT("[%u] MetaEvents: %d added\n"), ThreadState.ThreadId, ThreadState.PendingEvents.Num());
@@ -554,14 +540,6 @@ uint64 FCpuProfilerAnalyzer::ProcessBufferV2(const FEventTime& EventTime, FThrea
 		CPUPROFILER_DEBUG_LOGF(TEXT("[%u] MetaEvents: %d added, %d still pending\n"), ThreadState.ThreadId, NumEventsToRemove, RemainingPending);
 		ThreadState.PendingEvents.RemoveAt(0, NumEventsToRemove);
 	}
-#else
-	if (RemainingPending > 0)
-	{
-		DispatchPendingEvents(LastCycle, ~0ull, EventTime, ThreadState, PendingCursor, RemainingPending);
-		check(RemainingPending == 0);
-		ThreadState.PendingEvents.Reset();
-	}
-#endif
 
 	ThreadState.LastCycle = LastCycle;
 	return LastCycle;
@@ -577,6 +555,13 @@ void FCpuProfilerAnalyzer::DispatchPendingEvents(
 	const FPendingEvent*& PendingCursor,
 	int32& RemainingPending)
 {
+	if (ThreadState.bShouldIgnorePendingEvents)
+	{
+		PendingCursor += RemainingPending;
+		RemainingPending = 0;
+		return;
+	}
+
 	for (; RemainingPending > 0; RemainingPending--, PendingCursor++)
 	{
 		bool bEnter = true;
@@ -592,24 +577,21 @@ void FCpuProfilerAnalyzer::DispatchPendingEvents(
 			break;
 		}
 
-#if CPUPROFILER_SAFE_DISPATCH_PENDING_EVENTS
-		if (PendingCycle >= LastCycle)
-		{
-			// Update LastCycle in order to verify time (of following pending events) increases monotonically.
-			LastCycle = PendingCycle;
-		}
-		else
-		{
-			// Time needs to increase monotonically.
-			// We are not allowing events to "go back in time".
-			PendingCycle = LastCycle;
-		}
-#else
 		if (PendingCycle < LastCycle)
 		{
-			PendingCycle = LastCycle;
+			// Time needs to increase monotonically.
+			// We are not allowing pending events (with metadata) older than regular cpu timing events.
+			// When this happens we further ignore all pending events on this thread.
+			// The issue can occur in late connect trace sessions with trace protocol <= 6 (i.e. the scoped events have relative timestamps).
+			ThreadState.bShouldIgnorePendingEvents = true;
+			PendingCursor += RemainingPending;
+			RemainingPending = 0;
+			UE_LOG(LogTraceServices, Error, TEXT("[CpuProfiler] Detected non-monotonically increasing timestamp. Further CPU timing events with metadata are ignored on thread %u."), ThreadState.ThreadId);
+			break;
 		}
-#endif
+
+		// Update LastCycle in order to verify time (of following pending events) increases monotonically.
+		LastCycle = PendingCycle;
 
 		double PendingTime = EventTime.AsSeconds(PendingCycle);
 
@@ -644,6 +626,11 @@ void FCpuProfilerAnalyzer::OnCpuScopeEnter(const FOnEventContext& Context)
 	uint32 ThreadId = Context.ThreadInfo.GetId();
 	FThreadState& ThreadState = GetThreadState(ThreadId);
 
+	if (ThreadState.bShouldIgnorePendingEvents)
+	{
+		return;
+	}
+
 	uint32 SpecId = Context.EventData.GetTypeInfo().GetId();
 	SpecId = ~SpecId; // to keep out of the way of normal spec IDs.
 
@@ -681,6 +668,11 @@ void FCpuProfilerAnalyzer::OnCpuScopeLeave(const FOnEventContext& Context)
 	uint32 ThreadId = Context.ThreadInfo.GetId();
 	FThreadState& ThreadState = GetThreadState(ThreadId);
 
+	if (ThreadState.bShouldIgnorePendingEvents)
+	{
+		return;
+	}
+
 	uint64 Cycle = Context.EventTime.AsCycle64();
 	ThreadState.PendingEvents.Add({~Cycle});
 }
@@ -800,8 +792,6 @@ FCpuProfilerAnalyzer::FThreadState& FCpuProfilerAnalyzer::GetThreadState(uint32
 
 } // namespace TraceServices
 
-#undef CPUPROFILER_SAFE_DISPATCH_PENDING_EVENTS
-
 #undef CPUPROFILER_DEBUG_LOGF
 #undef CPUPROFILER_DEBUG_BEGIN_EVENT
 #undef CPUPROFILER_DEBUG_END_EVENT
diff --git a/Engine/Source/Developer/TraceServices/Private/Analyzers/CpuProfilerTraceAnalysis.h b/Engine/Source/Developer/TraceServices/Private/Analyzers/CpuProfilerTraceAnalysis.h
--- a/Engine/Source/Developer/TraceServices/Private/Analyzers/CpuProfilerTraceAnalysis.h
+++ b/Engine/Source/Developer/TraceServices/Private/Analyzers/CpuProfilerTraceAnalysis.h
@@ -45,6 +45,7 @@ private:
 		TArray<FPendingEvent> PendingEvents;
 		IEditableTimeline<FTimingProfilerEvent>* Timeline = nullptr;
 		uint64 LastCycle = 0;
+		bool bShouldIgnorePendingEvents = false; // becomes true when we detect first pending event with incorrect timestamp (i.e < LastCycle)
 	};
 
 	void OnCpuScopeEnter(const FOnEventContext& Context);
