[Insights]
  - Memory Insights: Added Enumerate API for live allocations in a specified address range.
  - Memory Insights: Moved the long living allocations map into a new FAllocsMap class wrapper. Also refactored some parts in AllocationsProvider in preparation for more changes.

diff --git a/Engine/Source/Developer/TraceServices/Private/Model/AllocationsProvider.cpp b/Engine/Source/Developer/TraceServices/Private/Model/AllocationsProvider.cpp
--- a/Engine/Source/Developer/TraceServices/Private/Model/AllocationsProvider.cpp
+++ b/Engine/Source/Developer/TraceServices/Private/Model/AllocationsProvider.cpp
@@ -559,7 +559,7 @@
 
 ////////////////////////////////////////////////////////////////////////////////////////////////////
 
-FAllocationItem* FShortLivingAllocs::FindRef(uint64 Address)
+FAllocationItem* FShortLivingAllocs::FindRef(uint64 Address) const
 {
 #if INSIGHTS_SLA_USE_ADDRESS_MAP
 	FNode** NodePtr = AddressMap.Find(Address);
@@ -583,6 +583,51 @@
 
 ////////////////////////////////////////////////////////////////////////////////////////////////////
 
+FAllocationItem* FShortLivingAllocs::FindRange(const uint64 Address) const
+{
+	//TODO: Can we do better than iterating all the nodes?
+	FNode* Node = LastAddedAllocNode;
+	while (Node != nullptr)
+	{
+		if (Node->Alloc->IsContained(Address))
+		{
+			return Node->Alloc;
+		}
+		Node = Node->Prev;
+	}
+	return nullptr;
+}
+
+////////////////////////////////////////////////////////////////////////////////////////////////////
+
+void FShortLivingAllocs::Enumerate(TFunctionRef<void(const FAllocationItem& Alloc)> Callback) const
+{
+	FNode* Node = LastAddedAllocNode;
+	while (Node != nullptr)
+	{
+		Callback(*Node->Alloc);
+		Node = Node->Prev;
+	}
+}
+
+////////////////////////////////////////////////////////////////////////////////////////////////////
+
+void FShortLivingAllocs::Enumerate(uint64 StartAddress, uint64 EndAddress, TFunctionRef<void(const FAllocationItem& Alloc)> Callback) const
+{
+	FNode* Node = LastAddedAllocNode;
+	while (Node != nullptr)
+	{
+		const FAllocationItem& Alloc = *Node->Alloc;
+		if (Alloc.Address >= StartAddress && Alloc.Address < EndAddress)
+		{
+			Callback(Alloc);
+		}
+		Node = Node->Prev;
+	}
+}
+
+////////////////////////////////////////////////////////////////////////////////////////////////////
+
 FAllocationItem* FShortLivingAllocs::AddAndRemoveOldest(FAllocationItem* Alloc)
 {
 	if (FirstUnusedNode == nullptr)
@@ -705,32 +750,6 @@
 
 ////////////////////////////////////////////////////////////////////////////////////////////////////
 
-void FShortLivingAllocs::Enumerate(TFunctionRef<void(const FAllocationItem& Alloc)> Callback) const
-{
-	FNode* Node = LastAddedAllocNode;
-	while (Node != nullptr)
-	{
-		Callback(*(Node->Alloc));
-		Node = Node->Prev;
-	}
-}
-
 ////////////////////////////////////////////////////////////////////////////////////////////////////
-FAllocationItem* FShortLivingAllocs::FindRange(const uint64 Address) const
-{
-	//todo: Can we do better than iterating all the nodes?
-	FNode* Node = LastAddedAllocNode;
-	while (Node != nullptr)
-	{
-		if (Node->Alloc->IsContained(Address))
-		{
-			return Node->Alloc;
-		}
-		Node = Node->Prev;
-	}
-	return nullptr;
-}
-
-////////////////////////////////////////////////////////////////////////////////////////////////////
 // FHeapAllocs
 ////////////////////////////////////////////////////////////////////////////////////////////////////
@@ -770,9 +789,9 @@
 
 ////////////////////////////////////////////////////////////////////////////////////////////////////
 
-FAllocationItem* FHeapAllocs::FindRef(uint64 Address)
+FAllocationItem* FHeapAllocs::FindRef(uint64 Address) const
 {
-	FList* ListPtr = AddressMap.Find(Address);
+	const FList* ListPtr = AddressMap.Find(Address);
 	if (ListPtr)
 	{
 		// Search in reverse added order.
@@ -789,6 +808,56 @@
 
 ////////////////////////////////////////////////////////////////////////////////////////////////////
 
+FAllocationItem* FHeapAllocs::FindRange(uint64 Address) const
+{
+	for (const auto& KV : AddressMap)
+	{
+		// Search in reverse added order.
+		for (const FNode* Node = KV.Value.Last; Node != nullptr; Node = Node->Prev)
+		{
+			if (Node->Alloc->IsContained(Address))
+			{
+				return Node->Alloc;
+			}
+		}
+	}
+	return nullptr;
+}
+
+////////////////////////////////////////////////////////////////////////////////////////////////////
+
+void FHeapAllocs::Enumerate(TFunctionRef<void(const FAllocationItem& Alloc)> Callback) const
+{
+	for (const auto& KV : AddressMap)
+	{
+		// Iterate in same order as added.
+		for (const FNode* Node = KV.Value.First; Node != nullptr; Node = Node->Next)
+		{
+			Callback(*Node->Alloc);
+		}
+	}
+}
+
+////////////////////////////////////////////////////////////////////////////////////////////////////
+
+void FHeapAllocs::Enumerate(uint64 StartAddress, uint64 EndAddress, TFunctionRef<void(const FAllocationItem& Alloc)> Callback) const
+{
+	for (const auto& KV : AddressMap)
+	{
+		if (KV.Key >= StartAddress && KV.Key < EndAddress)
+		{
+			// Iterate in same order as added.
+			for (const FNode* Node = KV.Value.First; Node != nullptr; Node = Node->Next)
+			{
+				//check(Node->Alloc->Address == KV.Key);
+				Callback(*Node->Alloc);
+			}
+		}
+	}
+}
+
+////////////////////////////////////////////////////////////////////////////////////////////////////
+
 void FHeapAllocs::Add(FAllocationItem* Alloc)
 {
 	FNode* Node = new FNode();
@@ -843,37 +912,5 @@
 }
 
 ////////////////////////////////////////////////////////////////////////////////////////////////////
-
-void FHeapAllocs::Enumerate(TFunctionRef<void(const FAllocationItem& Alloc)> Callback) const
-{
-	for (const auto& KV : AddressMap)
-	{
-		// Iterate in same order as added.
-		for (const FNode* Node = KV.Value.First; Node != nullptr; Node = Node->Next)
-		{
-			Callback(*Node->Alloc);
-		}
-	}
-}
-
-////////////////////////////////////////////////////////////////////////////////////////////////////
-
-FAllocationItem* FHeapAllocs::FindRange(uint64 Address) const
-{
-	for (const auto& KV : AddressMap)
-	{
-		// Search in reverse added order.
-		for (const FNode* Node = KV.Value.Last; Node != nullptr; Node = Node->Prev)
-		{
-			if (Node->Alloc->IsContained(Address))
-			{
-				return Node->Alloc;
-			}
-		}
-	}
-	return nullptr;
-}
-
-////////////////////////////////////////////////////////////////////////////////////////////////////
 // FLiveAllocCollection
 ////////////////////////////////////////////////////////////////////////////////////////////////////
@@ -890,14 +927,7 @@
 FLiveAllocCollection::~FLiveAllocCollection()
 {
 	HeapAllocs.Reset();
-
-	for (const auto& KV : LongLivingAllocs)
-	{
-		const FAllocationItem* Allocation = KV.Value;
-		delete Allocation;
-	}
 	LongLivingAllocs.Reset();
-
 	ShortLivingAllocs.Reset();
 
 #if INSIGHTS_USE_LAST_ALLOC
@@ -911,7 +941,7 @@
 
 ////////////////////////////////////////////////////////////////////////////////////////////////////
 
-FAllocationItem* FLiveAllocCollection::FindRef(uint64 Address)
+FAllocationItem* FLiveAllocCollection::FindRef(uint64 Address) const
 {
 #if INSIGHTS_USE_LAST_ALLOC
 	if (LastAlloc && LastAlloc->Address == Address)
@@ -941,7 +971,7 @@
 
 ////////////////////////////////////////////////////////////////////////////////////////////////////
 
-FAllocationItem* FLiveAllocCollection::FindHeapRef(uint64 Address)
+FAllocationItem* FLiveAllocCollection::FindHeapRef(uint64 Address) const
 {
 	FAllocationItem* FoundHeapAlloc = HeapAllocs.FindRef(Address);
 	if (FoundHeapAlloc)
@@ -955,7 +985,7 @@
 
 ////////////////////////////////////////////////////////////////////////////////////////////////////
 
-FAllocationItem* FLiveAllocCollection::FindByAddressRange(uint64 Address)
+FAllocationItem* FLiveAllocCollection::FindByAddressRange(uint64 Address) const
 {
 #if INSIGHTS_USE_LAST_ALLOC
 	if (LastAlloc && LastAlloc->IsContained(Address))
@@ -973,12 +1003,11 @@
 	}
 #endif
 
-	for (const auto& Alloc : LongLivingAllocs)
+	FAllocationItem* FoundLongLivingAlloc = LongLivingAllocs.FindRange(Address);
+	if (FoundLongLivingAlloc)
 	{
-		if (Alloc.Value->IsContained(Address))
-		{
-			return Alloc.Value;
-		}
+		INSIGHTS_SLOW_CHECK(FoundLongLivingAlloc->IsContained(Address));
+		return FoundLongLivingAlloc;
 	}
 
 	return nullptr;
@@ -986,7 +1015,7 @@
 
 ////////////////////////////////////////////////////////////////////////////////////////////////////
 
-FAllocationItem* FLiveAllocCollection::FindHeapByAddressRange(uint64 Address)
+FAllocationItem* FLiveAllocCollection::FindHeapByAddressRange(uint64 Address) const
 {
 	FAllocationItem* FoundHeapAlloc = HeapAllocs.FindRange(Address);
 	if (FoundHeapAlloc)
@@ -1000,6 +1029,46 @@
 
 ////////////////////////////////////////////////////////////////////////////////////////////////////
 
+void FLiveAllocCollection::Enumerate(TFunctionRef<void(const FAllocationItem& Alloc)> Callback) const
+{
+	HeapAllocs.Enumerate(Callback);
+
+#if INSIGHTS_USE_LAST_ALLOC
+	if (LastAlloc)
+	{
+		Callback(*LastAlloc);
+	}
+#endif
+
+#if INSIGHTS_USE_SHORT_LIVING_ALLOCS
+	ShortLivingAllocs.Enumerate(Callback);
+#endif
+
+	LongLivingAllocs.Enumerate(Callback);
+}
+
+////////////////////////////////////////////////////////////////////////////////////////////////////
+
+void FLiveAllocCollection::Enumerate(uint64 StartAddress, uint64 EndAddress, TFunctionRef<void(const FAllocationItem& Alloc)> Callback) const
+{
+	HeapAllocs.Enumerate(StartAddress, EndAddress, Callback);
+
+#if INSIGHTS_USE_LAST_ALLOC
+	if (LastAlloc && LastAlloc->Address >= StartAddress && LastAlloc->Address < EndAddress)
+	{
+		Callback(*LastAlloc);
+	}
+#endif
+
+#if INSIGHTS_USE_SHORT_LIVING_ALLOCS
+	ShortLivingAllocs.Enumerate(StartAddress, EndAddress, Callback);
+#endif
+
+	LongLivingAllocs.Enumerate(StartAddress, EndAddress, Callback);
+}
+
+////////////////////////////////////////////////////////////////////////////////////////////////////
+
 FAllocationItem* FLiveAllocCollection::AddNew(uint64 Address)
 {
 	FAllocationItem* Alloc = new FAllocationItem();
@@ -1043,7 +1112,7 @@
 	}
 #endif
 
-	LongLivingAllocs.Add(Alloc->Address, Alloc);
+	LongLivingAllocs.Add(Alloc);
 }
 
 ////////////////////////////////////////////////////////////////////////////////////////////////////
@@ -1096,8 +1165,8 @@
 	}
 #endif
 
-	FAllocationItem* RemovedLongLivingAlloc;
-	if (LongLivingAllocs.RemoveAndCopyValue(Address, RemovedLongLivingAlloc))
+	FAllocationItem* RemovedLongLivingAlloc = LongLivingAllocs.Remove(Address);
+	if (RemovedLongLivingAlloc)
 	{
 		INSIGHTS_SLOW_CHECK(RemovedLongLivingAlloc->Address == Address);
 		INSIGHTS_SLOW_CHECK(TotalAllocCount > 0);
@@ -1125,29 +1194,5 @@
 }
 
 ////////////////////////////////////////////////////////////////////////////////////////////////////
-
-void FLiveAllocCollection::Enumerate(TFunctionRef<void(const FAllocationItem& Alloc)> Callback) const
-{
-	HeapAllocs.Enumerate(Callback);
-
-#if INSIGHTS_USE_LAST_ALLOC
-	if (LastAlloc)
-	{
-		Callback(*LastAlloc);
-	}
-#endif
-
-#if INSIGHTS_USE_SHORT_LIVING_ALLOCS
-	ShortLivingAllocs.Enumerate(Callback);
-#endif
-
-	for (const auto& KV : LongLivingAllocs)
-	{
-		const FAllocationItem* Allocation = KV.Value;
-		Callback(*Allocation);
-	}
-}
-
-////////////////////////////////////////////////////////////////////////////////////////////////////
 // FAllocationsProvider
 ////////////////////////////////////////////////////////////////////////////////////////////////////
@@ -1211,58 +1256,6 @@
 
 ////////////////////////////////////////////////////////////////////////////////////////////////////
 
-void FAllocationsProvider::AddHeapSpec(HeapId Id, HeapId ParentId, const FStringView& Name, EMemoryTraceHeapFlags Flags)
-{
-	if (Id < MaxRootHeaps)
-	{
-		if (Id == 0 && SbTree[Id] != nullptr)
-		{
-			// "System" root heap is already created. See class constructor.
-			check(LiveAllocs[Id] != nullptr);
-		}
-		else
-		{
-			check(SbTree[Id] == nullptr && LiveAllocs[Id] == nullptr);
-
-			constexpr uint32 ColumnShift = 17; // 1<<17 = 128K
-			SbTree[Id] = new FSbTree(Session.GetLinearAllocator(), ColumnShift);
-			LiveAllocs[Id] = new FLiveAllocCollection();
-		}
-
-		FHeapSpec& RootHeapSpec = HeapSpecs[Id];
-		RootHeapSpec.Name = Session.StoreString(Name);
-		RootHeapSpec.Parent = nullptr;
-		RootHeapSpec.Id = Id;
-		RootHeapSpec.Flags = Flags;
-	}
-	else
-	{
-		FHeapSpec& ParentSpec = HeapSpecs[ParentId];
-		if (ParentSpec.Name == nullptr)
-		{
-			UE_LOG(LogTraceServices, Warning, TEXT("[MemAlloc] Parent heap id (%u) used before it was announced for heap %s."), ParentId, *FString(Name));
-		}
-		FHeapSpec& NewSpec = HeapSpecs[Id];
-		NewSpec.Name = Session.StoreString(Name);
-		NewSpec.Parent = &ParentSpec;
-		NewSpec.Id = Id;
-		NewSpec.Flags = Flags;
-
-		ParentSpec.Children.Add(&NewSpec);
-	}
-}
-
-////////////////////////////////////////////////////////////////////////////////////////////////////
-
-void FAllocationsProvider::EditHeapSpec(HeapId Id, HeapId ParentId, const FStringView& Name, EMemoryTraceHeapFlags Flags)
-{
-	Lock.WriteAccessCheck();
-
-	AddHeapSpec(Id, ParentId, Name, Flags);
-}
-
-////////////////////////////////////////////////////////////////////////////////////////////////////
-
 void FAllocationsProvider::EditAlloc(double Time, uint32 CallstackId, uint64 Address, uint64 InSize, uint32 InAlignment, HeapId RootHeap)
 {
 	Lock.WriteAccessCheck();
@@ -1495,56 +1488,58 @@
 
 ////////////////////////////////////////////////////////////////////////////////////////////////////
 
-void FAllocationsProvider::EditUnmarkAllocationAsHeap(double Time, uint64 Address, HeapId Heap)
+void FAllocationsProvider::AddHeapSpec(HeapId Id, HeapId ParentId, const FStringView& Name, EMemoryTraceHeapFlags Flags)
 {
-	Lock.WriteAccessCheck();
+	if (Id < MaxRootHeaps)
+	{
+		if (Id == 0 && SbTree[Id] != nullptr)
+		{
+			// "System" root heap is already created. See class constructor.
+			check(LiveAllocs[Id] != nullptr);
+		}
+		else
+		{
+			check(SbTree[Id] == nullptr && LiveAllocs[Id] == nullptr);
 
-#if INSIGHTS_DEBUG_WATCH
-	for (int32 AddrIndex = 0; AddrIndex < UE_ARRAY_COUNT(GWatchAddresses); ++AddrIndex)
-	{
-		if (GWatchAddresses[AddrIndex] == Address)
-		{
-			UE_LOG(LogTraceServices, Warning, TEXT("[MemAlloc][%u] HeapUnmarkAlloc 0x%llX : Heap=%u, Time=%f"), CurrentTraceThreadId, Address, Heap, Time);
-			INSIGHTS_DEBUG_WATCH_FOUND;
+			constexpr uint32 ColumnShift = 17; // 1<<17 = 128K
+			SbTree[Id] = new FSbTree(Session.GetLinearAllocator(), ColumnShift);
+			LiveAllocs[Id] = new FLiveAllocCollection();
 		}
+
+		FHeapSpec& RootHeapSpec = HeapSpecs[Id];
+		RootHeapSpec.Name = Session.StoreString(Name);
+		RootHeapSpec.Parent = nullptr;
+		RootHeapSpec.Id = Id;
+		RootHeapSpec.Flags = Flags;
 	}
-#endif // INSIGHTS_DEBUG_WATCH
-
-	HeapId RootHeap = FindRootHeap(Heap);
-
-	// Remove the heap allocation from the Live allocs.
-	FAllocationItem* Alloc = LiveAllocs[RootHeap]->RemoveHeap(Address); // we take ownership of Alloc
-	if (Alloc)
-	{
-		const uint64 Size = Alloc->GetSize();
-		const uint32 Alignment = Alloc->GetAlignment();
-		const uint32 CallstackId = Alloc->CallstackId;
-		const uint32 FreeCallstackId = 0; // unknown
-
-		// Re-add this allocation to the Live allocs.
-		LiveAllocs[RootHeap]->Add(Alloc); // the Live allocs takes ownership of Alloc
-
-		// We cannot just unmark the allocation as heap, there is no timestamp support, instead fake a "free"
-		// event and an "alloc" event. Make sure the new allocation retains the tag from the original.
-		CurrentTracker = 1;
-		EditPushTagFromPtr(CurrentSystemThreadId, CurrentTracker, Address);
-		EditFree(Time, FreeCallstackId, Address, RootHeap);
-		EditAlloc(Time, CallstackId, Address, Size, Alignment, RootHeap);
-		EditPopTagFromPtr(CurrentSystemThreadId, CurrentTracker);
-		CurrentTracker = 0;
-	}
 	else
 	{
-		++HeapErrors;
-		if (HeapErrors <= MaxLogMessagesPerErrorType)
+		FHeapSpec& ParentSpec = HeapSpecs[ParentId];
+		if (ParentSpec.Name == nullptr)
 		{
-			UE_LOG(LogTraceServices, Error, TEXT("[MemAlloc] HeapUnmarkAlloc: Could not find address 0x%llX (Heap=%u, Time=%f)!"), Address, Heap, Time);
+			UE_LOG(LogTraceServices, Warning, TEXT("[MemAlloc] Parent heap id (%u) used before it was announced for heap %s."), ParentId, *FString(Name));
 		}
+		FHeapSpec& NewSpec = HeapSpecs[Id];
+		NewSpec.Name = Session.StoreString(Name);
+		NewSpec.Parent = &ParentSpec;
+		NewSpec.Id = Id;
+		NewSpec.Flags = Flags;
+
+		ParentSpec.Children.Add(&NewSpec);
 	}
 }
 
 ////////////////////////////////////////////////////////////////////////////////////////////////////
 
+void FAllocationsProvider::EditHeapSpec(HeapId Id, HeapId ParentId, const FStringView& Name, EMemoryTraceHeapFlags Flags)
+{
+	Lock.WriteAccessCheck();
+
+	AddHeapSpec(Id, ParentId, Name, Flags);
+}
+
+////////////////////////////////////////////////////////////////////////////////////////////////////
+
 void FAllocationsProvider::EditMarkAllocationAsHeap(double Time, uint64 Address, HeapId Heap, EMemoryTraceHeapAllocationFlags Flags)
 {
 	Lock.WriteAccessCheck();
@@ -1560,10 +1555,10 @@
 	}
 #endif // INSIGHTS_DEBUG_WATCH
 
-	HeapId RootHeap = FindRootHeap(Heap);
+	const HeapId RootHeap = FindRootHeap(Heap);
 
 	// Remove the allocation from the Live allocs.
-	FAllocationItem* Alloc = LiveAllocs[RootHeap]->Remove(Address);
+	FAllocationItem* Alloc = LiveAllocs[RootHeap]->Remove(Address); // we take ownership of Alloc
 	if (Alloc)
 	{
 		// Mark allocation as a "heap" allocation.
@@ -1572,7 +1567,7 @@
 		Alloc->RootHeap = static_cast<uint8>(RootHeap);
 
 		// Re-add it to the Live allocs as a heap allocation.
-		LiveAllocs[RootHeap]->AddHeap(Alloc);
+		LiveAllocs[RootHeap]->AddHeap(Alloc); // the Live allocs takes ownership of Alloc
 
 		// Update stats. Remove this allocation from the total.
 		TotalAllocatedMemory -= Alloc->GetSize();
@@ -1594,6 +1589,56 @@
 
 ////////////////////////////////////////////////////////////////////////////////////////////////////
 
+void FAllocationsProvider::EditUnmarkAllocationAsHeap(double Time, uint64 Address, HeapId Heap)
+{
+	Lock.WriteAccessCheck();
+
+#if INSIGHTS_DEBUG_WATCH
+	for (int32 AddrIndex = 0; AddrIndex < UE_ARRAY_COUNT(GWatchAddresses); ++AddrIndex)
+	{
+		if (GWatchAddresses[AddrIndex] == Address)
+		{
+			UE_LOG(LogTraceServices, Warning, TEXT("[MemAlloc][%u] HeapUnmarkAlloc 0x%llX : Heap=%u, Time=%f"), CurrentTraceThreadId, Address, Heap, Time);
+			INSIGHTS_DEBUG_WATCH_FOUND;
+		}
+	}
+#endif // INSIGHTS_DEBUG_WATCH
+
+	const HeapId RootHeap = FindRootHeap(Heap);
+
+	// Remove the heap allocation from the Live allocs.
+	FAllocationItem* Alloc = LiveAllocs[RootHeap]->RemoveHeap(Address); // we take ownership of Alloc
+	if (Alloc)
+	{
+		const uint64 Size = Alloc->GetSize();
+		const uint32 Alignment = Alloc->GetAlignment();
+		const uint32 AllocCallstackId = Alloc->CallstackId;
+		const uint32 FreeCallstackId = 0; // unknown
+
+		// Re-add this allocation to the Live allocs.
+		LiveAllocs[RootHeap]->Add(Alloc); // the Live allocs takes ownership of Alloc
+
+		// We cannot just unmark the allocation as heap, there is no timestamp support, instead fake a "free"
+		// event and an "alloc" event. Make sure the new allocation retains the tag from the original.
+		CurrentTracker = 1;
+		EditPushTagFromPtr(CurrentSystemThreadId, CurrentTracker, Address);
+		EditFree(Time, FreeCallstackId, Address, RootHeap);
+		EditAlloc(Time, AllocCallstackId, Address, Size, Alignment, RootHeap);
+		EditPopTagFromPtr(CurrentSystemThreadId, CurrentTracker);
+		CurrentTracker = 0;
+	}
+	else
+	{
+		++HeapErrors;
+		if (HeapErrors <= MaxLogMessagesPerErrorType)
+		{
+			UE_LOG(LogTraceServices, Error, TEXT("[MemAlloc] HeapUnmarkAlloc: Could not find address 0x%llX (Heap=%u, Time=%f)!"), Address, Heap, Time);
+		}
+	}
+}
+
+////////////////////////////////////////////////////////////////////////////////////////////////////
+
 void FAllocationsProvider::UpdateHistogramByAllocSize(uint64 Size)
 {
 	if (Size > MaxAllocSize)
diff --git a/Engine/Source/Developer/TraceServices/Private/Model/AllocationsProvider.h b/Engine/Source/Developer/TraceServices/Private/Model/AllocationsProvider.h
--- a/Engine/Source/Developer/TraceServices/Private/Model/AllocationsProvider.h
+++ b/Engine/Source/Developer/TraceServices/Private/Model/AllocationsProvider.h
@@ -158,7 +158,16 @@ public:
 	bool IsFull() const { return AllocCount == MaxAllocCount; }
 	int32 Num() const { return AllocCount; }
 
-	FORCEINLINE FAllocationItem* FindRef(uint64 Address);
+	// Finds the allocation with the specified address.
+	// Returns the found allocation or nullptr if not found.
+	FORCEINLINE FAllocationItem* FindRef(uint64 Address) const;
+
+	// Finds an allocation containing the specified address.
+	// Returns the found allocation or nullptr if not found.
+	FORCEINLINE FAllocationItem* FindRange(uint64 Address) const;
+
+	FORCEINLINE void Enumerate(TFunctionRef<void(const FAllocationItem& Alloc)> Callback) const;
+	FORCEINLINE void Enumerate(uint64 StartAddress, uint64 EndAddress, TFunctionRef<void(const FAllocationItem& Alloc)> Callback) const;
 
 	// The collection keeps ownership of FAllocationItem* until Remove is called or until the oldest allocation is removed.
 	// Returns the removed oldest allocation if collection is already full; nullptr otherwise.
@@ -168,13 +177,8 @@ public:
 	// The caller takes ownership of FAllocationItem*. Returns nullptr if Address is not found.
 	FORCEINLINE FAllocationItem* Remove(uint64 Address);
 
-	void Enumerate(TFunctionRef<void(const FAllocationItem& Alloc)> Callback) const;
-
-	// Finds an allocation that contains the address specified or null if not found.
-	FAllocationItem* FindRange(uint64 Address) const;
-
 private:
-	//TMap<uint64, FNode*> AddressMap; // map for short living allocations: Address -> FNode*
+	//TMap<uint64, FNode*> AddressMap; // map for short living allocations: Address -> FNode* // #if INSIGHTS_SLA_USE_ADDRESS_MAP
 	FNode* AllNodes = nullptr; // preallocated array of nodes (MaxAllocCount nodes)
 	FNode* LastAddedAllocNode = nullptr; // the last added alloc; double linked list: Prev -> .. -> OldestAlloc
 	FNode* OldestAllocNode = nullptr; // the oldest alloc; double linked list: Next -> .. -> LastAddedAlloc
@@ -184,6 +188,120 @@ private:
 
 ////////////////////////////////////////////////////////////////////////////////////////////////////
 
+class FAllocsMap
+{
+public:
+	FAllocsMap()
+	{
+	}
+
+	~FAllocsMap()
+	{
+		Reset();
+	}
+
+	void Reset()
+	{
+		for (const auto& KV : AllocsMap)
+		{
+			const FAllocationItem* Allocation = KV.Value;
+			delete Allocation;
+		}
+		AllocsMap.Reset();
+	}
+
+	int32 Num() const
+	{
+		return AllocsMap.Num();
+	}
+
+	// Finds the allocation with the specified address.
+	// Returns the found allocation or nullptr if not found.
+	FAllocationItem* FindRef(uint64 Address) const
+	{
+		return AllocsMap.FindRef(Address);
+	}
+
+	// Finds an allocation containing the specified address.
+	// Returns the found allocation or nullptr if not found.
+	FAllocationItem* FindRange(uint64 Address) const
+	{
+		for (const auto& Alloc : AllocsMap)
+		{
+			if (Alloc.Value->IsContained(Address))
+			{
+				return Alloc.Value;
+			}
+		}
+		return nullptr;
+	}
+
+	void Enumerate(TFunctionRef<void(const FAllocationItem& Alloc)> Callback) const
+	{
+		for (const auto& KV : AllocsMap)
+		{
+			Callback(*KV.Value);
+		}
+	}
+
+	void Enumerate(uint64 StartAddress, uint64 EndAddress, TFunctionRef<void(const FAllocationItem& Alloc)> Callback) const
+	{
+		for (const auto& KV : AllocsMap)
+		{
+			const FAllocationItem& Allocation = *KV.Value;
+			if (Allocation.Address >= StartAddress && Allocation.Address < EndAddress)
+			{
+				Callback(Allocation);
+			}
+		}
+	}
+
+	// The collection keeps ownership of FAllocationItem* until Remove is called.
+	void Add(FAllocationItem* Alloc)
+	{
+		AllocsMap.Add(Alloc->Address, Alloc);
+	}
+
+	// The caller takes ownership of FAllocationItem*. Returns nullptr if Address is not found.
+	FAllocationItem* Remove(uint64 Address)
+	{
+		FAllocationItem* RemovedLongLivingAlloc;
+		if (AllocsMap.RemoveAndCopyValue(Address, RemovedLongLivingAlloc))
+		{
+			return RemovedLongLivingAlloc;
+		}
+		return nullptr;
+	}
+
+private:
+	template<typename ValueType>
+	struct TAddressKeyFuncs : BaseKeyFuncs<TPair<uint64, ValueType>, uint64, false>
+	{
+		typedef typename TTypeTraits<uint64>::ConstPointerType KeyInitType;
+		typedef const TPairInitializer<typename TTypeTraits<uint64>::ConstInitType, typename TTypeTraits<ValueType>::ConstInitType>& ElementInitType;
+
+		static FORCEINLINE uint64 GetSetKey(ElementInitType Element)
+		{
+			return Element.Key;
+		}
+
+		static FORCEINLINE bool Matches(uint64 A, uint64 B)
+		{
+			return A == B;
+		}
+
+		static FORCEINLINE uint32 GetKeyHash(uint64 Key)
+		{
+			return uint32(Key >> 6);
+		}
+	};
+
+private:
+	TMap<uint64, FAllocationItem*, FDefaultSetAllocator, TAddressKeyFuncs<FAllocationItem*>> AllocsMap;
+};
+
+////////////////////////////////////////////////////////////////////////////////////////////////////
+
 class FHeapAllocs
 {
 private:
@@ -210,7 +328,14 @@ public:
 
 	// Finds the last heap allocation with specified address.
 	// Returns the found allocation or nullptr if not found.
-	FORCEINLINE FAllocationItem* FindRef(uint64 Address);
+	FORCEINLINE FAllocationItem* FindRef(uint64 Address) const;
+
+	// Finds a heap allocation containing the specified address.
+	// Returns the found allocation or nullptr if not found.
+	FORCEINLINE FAllocationItem* FindRange(uint64 Address) const;
+
+	FORCEINLINE void Enumerate(TFunctionRef<void(const FAllocationItem& Alloc)> Callback) const;
+	FORCEINLINE void Enumerate(uint64 StartAddress, uint64 EndAddress, TFunctionRef<void(const FAllocationItem& Alloc)> Callback) const;
 
 	// The collection keeps ownership of FAllocationItem* until Remove is called.
 	FORCEINLINE void Add(FAllocationItem* Alloc);
@@ -218,11 +343,6 @@ public:
 	// The caller takes ownership of FAllocationItem*. Returns nullptr if Address is not found.
 	FORCEINLINE FAllocationItem* Remove(uint64 Address);
 
-	void Enumerate(TFunctionRef<void(const FAllocationItem& Alloc)> Callback) const;
-
-	// Finds an allocation that contains the address specified or null if not found.
-	FAllocationItem* FindRange(uint64 Address) const;
-
 private:
 	TMap<uint64, FList> AddressMap; // map heap allocs: Address -> list of heap allocs
 	int32 AllocCount = 0;
@@ -241,19 +361,25 @@ public:
 
 	// Finds the allocation with specified address.
 	// Returns the found allocation or nullptr if not found.
-	FORCEINLINE FAllocationItem* FindRef(uint64 Address);
+	FORCEINLINE FAllocationItem* FindRef(uint64 Address) const;
 
 	// Finds the heap allocation with specified address.
 	// Returns the found allocation or nullptr if not found.
-	FORCEINLINE FAllocationItem* FindHeapRef(uint64 Address);
+	FORCEINLINE FAllocationItem* FindHeapRef(uint64 Address) const;
 
 	// Finds an allocation containing the address.
 	// Returns the found allocation or nullptr if not found.
-	FORCEINLINE FAllocationItem* FindByAddressRange(uint64 Address);
+	FORCEINLINE FAllocationItem* FindByAddressRange(uint64 Address) const;
 
 	// Finds the heap allocation containing the address.
 	// Returns the found allocation or nullptr if not found.
-	FORCEINLINE FAllocationItem* FindHeapByAddressRange(uint64 Address);
+	FORCEINLINE FAllocationItem* FindHeapByAddressRange(uint64 Address) const;
+
+	// Enumerates all allocations (including heap allocations).
+	FORCEINLINE void Enumerate(TFunctionRef<void(const FAllocationItem& Alloc)> Callback) const;
+
+	// Enumerates allocations in a sprecified address range (including heap allocations).
+	FORCEINLINE void Enumerate(uint64 StartAddress, uint64 EndAddress, TFunctionRef<void(const FAllocationItem& Alloc)> Callback) const;
 
 	// Adds a new allocation with specified address.
 	// The collection keeps ownership of FAllocationItem* until Remove is called.
@@ -283,37 +409,11 @@ public:
 	// Returns the removed heap allocation or nullptr if not found.
 	FORCEINLINE FAllocationItem* RemoveHeap(uint64 Address);
 
-	// Enumerates all allocations (including heap allocations).
-	void Enumerate(TFunctionRef<void(const FAllocationItem& Alloc)> Callback) const;
-
-private:
-	template<typename ValueType>
-	struct TAddressKeyFuncs : BaseKeyFuncs<TPair<uint64, ValueType>, uint64, false>
-	{
-		typedef typename TTypeTraits<uint64>::ConstPointerType KeyInitType;
-		typedef const TPairInitializer<typename TTypeTraits<uint64>::ConstInitType, typename TTypeTraits<ValueType>::ConstInitType>& ElementInitType;
-
-		static FORCEINLINE uint64 GetSetKey(ElementInitType Element)
-		{
-			return Element.Key;
-		}
-
-		static FORCEINLINE bool Matches(uint64 A, uint64 B)
-		{
-			return A == B;
-		}
-
-		static FORCEINLINE uint32 GetKeyHash(uint64 Key)
-		{
-			return uint32(Key >> 6);
-		}
-	};
-
 private:
 	FAllocationItem* LastAlloc = nullptr; // last allocation
-	FShortLivingAllocs ShortLivingAllocs; // short living allocs
-	TMap<uint64, FAllocationItem*, FDefaultSetAllocator, TAddressKeyFuncs<FAllocationItem*>> LongLivingAllocs; // long living allocations
-	FHeapAllocs HeapAllocs; // heap allocs
+	FShortLivingAllocs ShortLivingAllocs; // short living allocations
+	FAllocsMap LongLivingAllocs; // long living allocations
+	FHeapAllocs HeapAllocs; // heap allocations
 
 	uint32 TotalAllocCount = 0;
 	uint32 MaxAllocCount = 0; // debug stats
@@ -374,10 +474,10 @@ public:
 
 	void EditInit(double Time, uint8 MinAlignment);
 
-	void EditHeapSpec(HeapId Id, HeapId ParentId, const FStringView& Name, EMemoryTraceHeapFlags Flags);
-
 	void EditAlloc(double Time, uint32 CallstackId, uint64 Address, uint64 Size, uint32 Alignment, HeapId RootHeap);
 	void EditFree(double Time, uint32 CallstackId, uint64 Address, HeapId RootHeap);
+
+	void EditHeapSpec(HeapId Id, HeapId ParentId, const FStringView& Name, EMemoryTraceHeapFlags Flags);
 	void EditMarkAllocationAsHeap(double Time, uint64 Address, HeapId Heap, EMemoryTraceHeapAllocationFlags Flags);
 	void EditUnmarkAllocationAsHeap(double Time, uint64 Address, HeapId Heap);
 
