[TraceServices]
* Fixed unix non unity compile error by adding UnrealString.h

diff --git a/Engine/Source/Developer/TraceServices/Private/Model/Tables.h b/Engine/Source/Developer/TraceServices/Private/Model/Tables.h
--- a/Engine/Source/Developer/TraceServices/Private/Model/Tables.h
+++ b/Engine/Source/Developer/TraceServices/Private/Model/Tables.h
@@ -2,6 +2,7 @@
 
 #pragma once
 
+#include "Containers/UnrealString.h"
 #include "TraceServices/Containers/Tables.h"
 #include "Templates/Function.h"
 #include "Common/PagedArray.h"
