Iris
- Changed check to ensure in order to avoid crashing when parsing certain trace-files.

diff --git a/Engine/Source/Developer/TraceServices/Private/Analyzers/NetTraceAnalyzer.cpp b/Engine/Source/Developer/TraceServices/Private/Analyzers/NetTraceAnalyzer.cpp
--- a/Engine/Source/Developer/TraceServices/Private/Analyzers/NetTraceAnalyzer.cpp
+++ b/Engine/Source/Developer/TraceServices/Private/Analyzers/NetTraceAnalyzer.cpp
@@ -607,7 +607,7 @@ void FNetTraceAnalyzer::HandleConnectionCreatedEvent(const FOnEventContext& Cont
 	const uint16 ConnectionId = EventData.GetValue<uint16>("ConnectionId");
 
 	TSharedRef<FNetTraceGameInstanceState> GameInstanceState = GetOrCreateActiveGameInstanceState(GameInstanceId);
-	check(!GameInstanceState->ActiveConnections.Contains(ConnectionId));
+	ensureAlwaysMsgf(!GameInstanceState->ActiveConnections.Contains(ConnectionId), TEXT("Got ConnectionCreatedEvent for already existing connection GameInstanceId: %u ConnectionId: %u"), GameInstanceId, ConnectionId);
 
 	// Add to both active connections and to persistent connections
  	FNetProfilerConnectionInternal& Connection = NetProfilerProvider.CreateConnection(GameInstanceState->GameInstanceIndex);
