[Insights]
  - MemoryInsights: Moved FAllocationItem and FAllocMap (renamed from FAllocsMap) into their own header files.
  - MemoryInsights: Added debug logs to track edit API calls in FAllocationsProvider (can be enabled using API INSIGHTS_API_LOGF).
  - MemoryInsights: Fixed few includes.

diff --git a/Engine/Source/Developer/TraceServices/Private/Model/AllocMap.cpp b/Engine/Source/Developer/TraceServices/Private/Model/AllocMap.cpp
new file mode 100644
--- /dev/null
+++ b/Engine/Source/Developer/TraceServices/Private/Model/AllocMap.cpp
@@ -0,0 +1,16 @@
+// Copyright Epic Games, Inc. All Rights Reserved.
+
+#include "AllocMap.h"
+
+////////////////////////////////////////////////////////////////////////////////////////////////////
+
+namespace TraceServices
+{
+
+////////////////////////////////////////////////////////////////////////////////////////////////////
+
+//...
+
+////////////////////////////////////////////////////////////////////////////////////////////////////
+
+} // namespace TraceServices
diff --git a/Engine/Source/Developer/TraceServices/Private/Model/AllocMap.h b/Engine/Source/Developer/TraceServices/Private/Model/AllocMap.h
new file mode 100644
--- /dev/null
+++ b/Engine/Source/Developer/TraceServices/Private/Model/AllocMap.h
@@ -0,0 +1,138 @@
+// Copyright Epic Games, Inc. All Rights Reserved.
+
+#pragma once
+
+#include "Containers/Map.h"
+#include "Model/AllocationItem.h"
+
+namespace TraceServices
+{
+
+class FAllocMap
+{
+public:
+	typedef int32 SizeType;
+
+public:
+	FORCEINLINE FAllocMap()
+	{
+	}
+
+	FORCEINLINE ~FAllocMap()
+	{
+		Empty();
+	}
+
+	FORCEINLINE void Empty(SizeType ExpectedNumElements = 0)
+	{
+		DestructAllocs();
+		AllocsMap.Empty(ExpectedNumElements);
+	}
+
+	FORCEINLINE void Reset()
+	{
+		DestructAllocs();
+		AllocsMap.Reset();
+	}
+
+	FORCEINLINE SizeType Num() const
+	{
+		return AllocsMap.Num();
+	}
+
+	// Finds the allocation with the specified address.
+	// Returns the found allocation or nullptr if not found.
+	FORCEINLINE FAllocationItem* FindRef(uint64 Address) const
+	{
+		return AllocsMap.FindRef(Address);
+	}
+
+	// Finds an allocation containing the specified address.
+	// Returns the found allocation or nullptr if not found.
+	FORCEINLINE FAllocationItem* FindRange(uint64 Address) const
+	{
+		for (const auto& Alloc : AllocsMap)
+		{
+			if (Alloc.Value->IsContained(Address))
+			{
+				return Alloc.Value;
+			}
+		}
+		return nullptr;
+	}
+
+	FORCEINLINE void Enumerate(TFunctionRef<void(const FAllocationItem& Alloc)> Callback) const
+	{
+		for (const auto& KV : AllocsMap)
+		{
+			Callback(*KV.Value);
+		}
+	}
+
+	FORCEINLINE void Enumerate(uint64 StartAddress, uint64 EndAddress, TFunctionRef<void(const FAllocationItem& Alloc)> Callback) const
+	{
+		for (const auto& KV : AllocsMap)
+		{
+			const FAllocationItem& Allocation = *KV.Value;
+			if (Allocation.Address >= StartAddress && Allocation.Address < EndAddress)
+			{
+				Callback(Allocation);
+			}
+		}
+	}
+
+	// The collection keeps ownership of FAllocationItem* until Remove is called.
+	FORCEINLINE void Add(FAllocationItem* Alloc)
+	{
+		AllocsMap.Add(Alloc->Address, Alloc);
+	}
+
+	// The caller takes ownership of FAllocationItem*. Returns nullptr if Address is not found.
+	FORCEINLINE FAllocationItem* Remove(uint64 Address)
+	{
+		FAllocationItem* RemovedLongLivingAlloc;
+		if (AllocsMap.RemoveAndCopyValue(Address, RemovedLongLivingAlloc))
+		{
+			return RemovedLongLivingAlloc;
+		}
+		return nullptr;
+	}
+
+private:
+	FORCEINLINE void DestructAllocs()
+	{
+		for (const auto& KV : AllocsMap)
+		{
+			const FAllocationItem* Allocation = KV.Value;
+			delete Allocation;
+		}
+	}
+
+private:
+	template<typename ValueType>
+	struct TAddressKeyFuncs : BaseKeyFuncs<TPair<uint64, ValueType>, uint64, false>
+	{
+		typedef typename TTypeTraits<uint64>::ConstPointerType KeyInitType;
+		typedef const TPairInitializer<typename TTypeTraits<uint64>::ConstInitType, typename TTypeTraits<ValueType>::ConstInitType>& ElementInitType;
+
+		static FORCEINLINE uint64 GetSetKey(ElementInitType Element)
+		{
+			return Element.Key;
+		}
+
+		static FORCEINLINE bool Matches(uint64 A, uint64 B)
+		{
+			return A == B;
+		}
+
+		static FORCEINLINE uint32 GetKeyHash(uint64 Key)
+		{
+			return uint32(Key >> 6);
+		}
+	};
+
+private:
+	TMap<uint64, FAllocationItem*, FDefaultSetAllocator, TAddressKeyFuncs<FAllocationItem* >> AllocsMap;
+};
+
+} // namespace TraceServices
diff --git a/Engine/Source/Developer/TraceServices/Private/Model/AllocationItem.h b/Engine/Source/Developer/TraceServices/Private/Model/AllocationItem.h
new file mode 100644
--- /dev/null
+++ b/Engine/Source/Developer/TraceServices/Private/Model/AllocationItem.h
@@ -0,0 +1,46 @@
+// Copyright Epic Games, Inc. All Rights Reserved.
+
+#pragma once
+
+#include "CoreTypes.h"
+#include "Misc/EnumClassFlags.h"
+#include "ProfilingDebugging/MemoryTrace.h" // for EMemoryTraceHeapAllocationFlags
+#include "TraceServices/Model/AllocationsProvider.h" // for TagIdType
+
+namespace TraceServices
+{
+
+struct FAllocationItem
+{
+	static constexpr uint32 AlignmentBits = 8;
+	static constexpr uint32 AlignmentShift = 56;
+	static constexpr uint64 SizeMask = (1ULL << AlignmentShift) - 1;
+
+	static uint64 UnpackSize(uint64 SizeAndAlignment) { return (SizeAndAlignment & SizeMask); }
+	static uint32 UnpackAlignment(uint64 SizeAndAlignment) { return static_cast<uint32>(SizeAndAlignment >> AlignmentShift); }
+	static uint64 PackSizeAndAlignment(uint64 Size, uint8 Alignment) { return (Size | (static_cast<uint64>(Alignment) << AlignmentShift)); }
+
+	FORCEINLINE bool IsContained(uint64 InAddress) const { return Address >= InAddress && Address < (Address + GetSize()); }
+	FORCEINLINE uint64 GetEndAddress() const { return Address + UnpackSize(SizeAndAlignment); }
+	FORCEINLINE uint64 GetSize() const { return UnpackSize(SizeAndAlignment); }
+	FORCEINLINE uint32 GetAlignment() const { return UnpackAlignment(SizeAndAlignment); }
+	FORCEINLINE bool IsHeap() const { return EnumHasAnyFlags(Flags, EMemoryTraceHeapAllocationFlags::Heap); }
+
+	uint64 Address;
+	uint64 SizeAndAlignment; // (Alignment << AlignmentShift) | Size
+	uint32 StartEventIndex;
+	uint32 EndEventIndex;
+	double StartTime;
+	double EndTime;
+	uint32 ThreadId;
+	uint32 CallstackId;
+	uint32 FreeCallstackId;
+	uint32 MetadataId;
+	TagIdType Tag; // uint32
+	uint8 RootHeap;
+	EMemoryTraceHeapAllocationFlags Flags; // uint8
+};
+
+static_assert(sizeof(FAllocationItem) == 64, "struct FAllocationItem needs packing");
+
+} // namespace TraceServices
diff --git a/Engine/Source/Developer/TraceServices/Private/Model/AllocationsProvider.cpp b/Engine/Source/Developer/TraceServices/Private/Model/AllocationsProvider.cpp
--- a/Engine/Source/Developer/TraceServices/Private/Model/AllocationsProvider.cpp
+++ b/Engine/Source/Developer/TraceServices/Private/Model/AllocationsProvider.cpp
@@ -47,6 +47,13 @@
 
 #define INSIGHTS_DEBUG_METADATA 0
 
+#if 0
+#define INSIGHTS_API_LOGF(Format, ...) { UE_LOG(LogTraceServices, Log, TEXT("[MemAlloc][API] ") Format, __VA_ARGS__); }
+//#define INSIGHTS_API_LOGF(Format, ...) { FPlatformMisc::LowLevelOutputDebugStringf(TEXT("[MemAlloc][API] ") Format TEXT("\n"), __VA_ARGS__); }
+#else
+#define INSIGHTS_API_LOGF(Format, ...)
+#endif
+
 ////////////////////////////////////////////////////////////////////////////////////////////////////
 
 namespace TraceServices
@@ -748,8 +755,6 @@ FAllocationItem* FShortLivingAllocs::Remove(uint64 Address)
 	return Node->Alloc;
 }
 
-////////////////////////////////////////////////////////////////////////////////////////////////////
-
 ////////////////////////////////////////////////////////////////////////////////////////////////////
 // FHeapAllocs
 ////////////////////////////////////////////////////////////////////////////////////////////////////
@@ -1243,6 +1248,8 @@ void FAllocationsProvider::EditInit(double InTime, uint8 InMinAlignment)
 		return;
 	}
 
+	INSIGHTS_API_LOGF(TEXT("Init(Time=%f, MinAlignment=%u)"), InTime, uint32(InMinAlignment));
+
 	InitTime = InTime;
 	MinAlignment = InMinAlignment;
 
@@ -1265,6 +1272,8 @@ void FAllocationsProvider::EditAlloc(double Time, uint32 CallstackId, uint64 Add
 		return;
 	}
 
+	INSIGHTS_API_LOGF(TEXT("Alloc(Time=%f, Address=0x%llX, Size=%llu, Alignment=%u, RootHeap=%u, CallstackId=%u)"), Time, Address, InSize, InAlignment, RootHeap, CallstackId);
+
 	if (Address == 0)
 	{
 		UE_LOG(LogTraceServices, Warning, TEXT("[MemAlloc] Alloc at address 0 : Size=%llu, RootHeap=%u, Time=%f, CallstackId=%u"), InSize, RootHeap, Time, CallstackId);
@@ -1359,7 +1368,7 @@ void FAllocationsProvider::EditAlloc(double Time, uint32 CallstackId, uint64 Add
 		++AllocErrors;
 		if (AllocErrors <= MaxLogMessagesPerErrorType)
 		{
-			UE_LOG(LogTraceServices, Error, TEXT("[MemAlloc] Invalid ALLOC event (Address=0x%llX, Size=%llu, Tag=%u, RootHeap=%u, Time=%f)!"), Address, InSize, Tag, RootHeap, Time);
+			UE_LOG(LogTraceServices, Error, TEXT("[MemAlloc] Invalid ALLOC event (Address=0x%llX, Size=%llu, Tag=%u, RootHeap=%u, Time=%f, CallstackId=%u)!"), Address, InSize, Tag, RootHeap, Time, CallstackId);
 		}
 	}
 #endif
@@ -1387,6 +1396,8 @@ void FAllocationsProvider::EditFree(double Time, uint32 CallstackId, uint64 Addr
 		return;
 	}
 
+	INSIGHTS_API_LOGF(TEXT("Free(Time=%f, Address=0x%llX, RootHeap=%u, CallstackId=%u)"), Time, Address, RootHeap, CallstackId);
+
 	if (Address == 0)
 	{
 		UE_LOG(LogTraceServices, Warning, TEXT("[MemAlloc] Free for address 0 : RootHeap=%u, Time=%f, CallstackId=%u"), RootHeap, Time, CallstackId);
@@ -1431,7 +1442,7 @@ void FAllocationsProvider::EditFree(double Time, uint32 CallstackId, uint64 Addr
 		++FreeErrors;
 		if (FreeErrors <= MaxLogMessagesPerErrorType)
 		{
-			UE_LOG(LogTraceServices, Error, TEXT("[MemAlloc] Invalid FREE event (Address=0x%llX, RootHeap=%u, Time=%f)! A fake alloc will be created with size 0."), Address, RootHeap, Time);
+			UE_LOG(LogTraceServices, Error, TEXT("[MemAlloc] Invalid FREE event (Address=0x%llX, RootHeap=%u, Time=%f, CallstackId=%u)! A fake alloc will be created with size 0."), Address, RootHeap, Time, CallstackId);
 		}
 		// Fake the missing alloc.
 		constexpr uint64 FakeAllocSize = 0;
@@ -1470,7 +1481,7 @@ void FAllocationsProvider::EditFree(double Time, uint32 CallstackId, uint64 Addr
 		++FreeErrors;
 		if (FreeErrors <= MaxLogMessagesPerErrorType)
 		{
-			UE_LOG(LogTraceServices, Error, TEXT("[MemAlloc] Invalid FREE event (Address=0x%llX, RootHeap=%u, Time=%f)!"), Address, RootHeap, Time);
+			UE_LOG(LogTraceServices, Error, TEXT("[MemAlloc] Invalid FREE event (Address=0x%llX, RootHeap=%u, Time=%f, CallstackId=%u)!"), Address, RootHeap, Time, CallstackId);
 		}
 	}
 
@@ -1488,6 +1499,17 @@ void FAllocationsProvider::EditFree(double Time, uint32 CallstackId, uint64 Addr
 
 ////////////////////////////////////////////////////////////////////////////////////////////////////
 
+void FAllocationsProvider::EditHeapSpec(HeapId Id, HeapId ParentId, const FStringView& Name, EMemoryTraceHeapFlags Flags)
+{
+	Lock.WriteAccessCheck();
+
+	INSIGHTS_API_LOGF(TEXT("HeapSpec(Id=%u, ParentId=%u, Name=\"%*s\", Flags=0x%X)"), Id, ParentId, Name.Len(), Name.GetData(), uint32(Flags));
+
+	AddHeapSpec(Id, ParentId, Name, Flags);
+}
+
+////////////////////////////////////////////////////////////////////////////////////////////////////
+
 void FAllocationsProvider::AddHeapSpec(HeapId Id, HeapId ParentId, const FStringView& Name, EMemoryTraceHeapFlags Flags)
 {
 	if (Id < MaxRootHeaps)
@@ -1531,25 +1553,20 @@ void FAllocationsProvider::AddHeapSpec(HeapId Id, HeapId ParentId, const FString
 
 ////////////////////////////////////////////////////////////////////////////////////////////////////
 
-void FAllocationsProvider::EditHeapSpec(HeapId Id, HeapId ParentId, const FStringView& Name, EMemoryTraceHeapFlags Flags)
-{
-	Lock.WriteAccessCheck();
-
-	AddHeapSpec(Id, ParentId, Name, Flags);
-}
-
-////////////////////////////////////////////////////////////////////////////////////////////////////
-
 void FAllocationsProvider::EditMarkAllocationAsHeap(double Time, uint64 Address, HeapId Heap, EMemoryTraceHeapAllocationFlags Flags)
 {
+	uint32 CallstackId = 0; // TODO
+
 	Lock.WriteAccessCheck();
 
+	INSIGHTS_API_LOGF(TEXT("MarkAllocAsHeap(Time=%f, Address=0x%llX, Heap=%u, Flags=0x%X, CallstackId=%u)"), Time, Address, Heap, uint32(Flags), CallstackId);
+
 #if INSIGHTS_DEBUG_WATCH
 	for (int32 AddrIndex = 0; AddrIndex < UE_ARRAY_COUNT(GWatchAddresses); ++AddrIndex)
 	{
 		if (GWatchAddresses[AddrIndex] == Address)
 		{
-			UE_LOG(LogTraceServices, Warning, TEXT("[MemAlloc][%u] HeapMarkAlloc 0x%llX : Heap=%u, Flags=%u, Time=%f"), CurrentTraceThreadId, Address, Heap, uint32(Flags), Time);
+			UE_LOG(LogTraceServices, Warning, TEXT("[MemAlloc][%u] HeapMarkAlloc 0x%llX : Heap=%u, Flags=%u, Time=%f, CallstackId=%u"), CurrentTraceThreadId, Address, Heap, uint32(Flags), Time, CallstackId);
 			INSIGHTS_DEBUG_WATCH_FOUND;
 		}
 	}
@@ -1580,7 +1597,7 @@ void FAllocationsProvider::EditMarkAllocationAsHeap(double Time, uint64 Address,
 		++HeapErrors;
 		if (HeapErrors <= MaxLogMessagesPerErrorType)
 		{
-			UE_LOG(LogTraceServices, Error, TEXT("[MemAlloc] HeapMarkAlloc: Could not find address 0x%llX (Heap=%u, Flags=%u, Time=%f)!"), Address, Heap, uint32(Flags), Time);
+			UE_LOG(LogTraceServices, Error, TEXT("[MemAlloc] HeapMarkAlloc: Could not find address 0x%llX (Heap=%u, Flags=%u, Time=%f, CallstackId=%u)!"), Address, Heap, uint32(Flags), Time, CallstackId);
 		}
 	}
 
@@ -1591,14 +1608,18 @@ void FAllocationsProvider::EditMarkAllocationAsHeap(double Time, uint64 Address,
 
 void FAllocationsProvider::EditUnmarkAllocationAsHeap(double Time, uint64 Address, HeapId Heap)
 {
+	uint32 CallstackId = 0; // TODO
+
 	Lock.WriteAccessCheck();
 
+	INSIGHTS_API_LOGF(TEXT("UnmarkAllocAsHeap(Time=%f, Address=0x%llX, Heap=%u, CallstackId=%u)"), Time, Address, Heap, CallstackId);
+
 #if INSIGHTS_DEBUG_WATCH
 	for (int32 AddrIndex = 0; AddrIndex < UE_ARRAY_COUNT(GWatchAddresses); ++AddrIndex)
 	{
 		if (GWatchAddresses[AddrIndex] == Address)
 		{
-			UE_LOG(LogTraceServices, Warning, TEXT("[MemAlloc][%u] HeapUnmarkAlloc 0x%llX : Heap=%u, Time=%f"), CurrentTraceThreadId, Address, Heap, Time);
+			UE_LOG(LogTraceServices, Warning, TEXT("[MemAlloc][%u] HeapUnmarkAlloc 0x%llX : Heap=%u, Time=%f, CallstackId=%u"), CurrentTraceThreadId, Address, Heap, Time, CallstackId);
 			INSIGHTS_DEBUG_WATCH_FOUND;
 		}
 	}
@@ -1613,7 +1634,6 @@ void FAllocationsProvider::EditUnmarkAllocationAsHeap(double Time, uint64 Addres
 		const uint64 Size = Alloc->GetSize();
 		const uint32 Alignment = Alloc->GetAlignment();
 		const uint32 AllocCallstackId = Alloc->CallstackId;
-		const uint32 FreeCallstackId = 0; // unknown
 
 		// Re-add this allocation to the Live allocs.
 		LiveAllocs[RootHeap]->Add(Alloc); // the Live allocs takes ownership of Alloc
@@ -1622,7 +1642,7 @@ void FAllocationsProvider::EditUnmarkAllocationAsHeap(double Time, uint64 Addres
 		// event and an "alloc" event. Make sure the new allocation retains the tag from the original.
 		CurrentTracker = 1;
 		EditPushTagFromPtr(CurrentSystemThreadId, CurrentTracker, Address);
-		EditFree(Time, FreeCallstackId, Address, RootHeap);
+		EditFree(Time, CallstackId, Address, RootHeap);
 		EditAlloc(Time, AllocCallstackId, Address, Size, Alignment, RootHeap);
 		EditPopTagFromPtr(CurrentSystemThreadId, CurrentTracker);
 		CurrentTracker = 0;
@@ -1632,7 +1652,7 @@ void FAllocationsProvider::EditUnmarkAllocationAsHeap(double Time, uint64 Addres
 		++HeapErrors;
 		if (HeapErrors <= MaxLogMessagesPerErrorType)
 		{
-			UE_LOG(LogTraceServices, Error, TEXT("[MemAlloc] HeapUnmarkAlloc: Could not find address 0x%llX (Heap=%u, Time=%f)!"), Address, Heap, Time);
+			UE_LOG(LogTraceServices, Error, TEXT("[MemAlloc] HeapUnmarkAlloc: Could not find address 0x%llX (Heap=%u, Time=%f, CallstackId=%u)!"), Address, Heap, Time, CallstackId);
 		}
 	}
 }
@@ -2240,3 +2260,4 @@ const IAllocationsProvider* ReadAllocationsProvider(const IAnalysisSession& Sess
 #undef INSIGHTS_USE_LAST_ALLOC
 #undef INSIGHTS_VALIDATE_ALLOC_EVENTS
 #undef INSIGHTS_DEBUG_METADATA
+#undef INSIGHTS_API_LOGF
diff --git a/Engine/Source/Developer/TraceServices/Private/Model/AllocationsProvider.h b/Engine/Source/Developer/TraceServices/Private/Model/AllocationsProvider.h
--- a/Engine/Source/Developer/TraceServices/Private/Model/AllocationsProvider.h
+++ b/Engine/Source/Developer/TraceServices/Private/Model/AllocationsProvider.h
@@ -6,7 +6,10 @@
 #include "Containers/Map.h"
 #include "HAL/CriticalSection.h"
 #include "ProfilingDebugging/MemoryTrace.h" // for EMemoryTraceHeapFlags and EMemoryTraceHeapAllocationFlags
+
 #include "TraceServices/Model/AllocationsProvider.h"
+#include "AllocationItem.h"
+#include "AllocMap.h"
 
 namespace TraceServices
 {
@@ -102,40 +105,6 @@ private:
 
 ////////////////////////////////////////////////////////////////////////////////////////////////////
 
-struct FAllocationItem
-{
-	static constexpr uint32 AlignmentBits = 8;
-	static constexpr uint32 AlignmentShift = 56;
-	static constexpr uint64 SizeMask = (1ULL << AlignmentShift) - 1;
-
-	static uint64 UnpackSize(uint64 SizeAndAlignment) { return (SizeAndAlignment & SizeMask); }
-	static uint32 UnpackAlignment(uint64 SizeAndAlignment) { return static_cast<uint32>(SizeAndAlignment >> AlignmentShift); }
-	static uint64 PackSizeAndAlignment(uint64 Size, uint8 Alignment) { return (Size | (static_cast<uint64>(Alignment) << AlignmentShift)); }
-
-	FORCEINLINE bool IsContained(uint64 InAddress) const { return Address >= InAddress && Address < (Address + GetSize()); }
-	uint64 GetSize() const { return UnpackSize(SizeAndAlignment); }
-	uint32 GetAlignment() const { return UnpackAlignment(SizeAndAlignment); }
-	bool IsHeap() const { return EnumHasAnyFlags(Flags, EMemoryTraceHeapAllocationFlags::Heap); }
-
-	uint64 Address;
-	uint64 SizeAndAlignment; // (Alignment << AlignmentShift) | Size
-	uint32 StartEventIndex;
-	uint32 EndEventIndex;
-	double StartTime;
-	double EndTime;
-	uint32 ThreadId;
-	uint32 CallstackId;
-	uint32 FreeCallstackId;
-	uint32 MetadataId;
-	TagIdType Tag;
-	uint8 RootHeap;
-	EMemoryTraceHeapAllocationFlags Flags;
-};
-
-static_assert(sizeof(FAllocationItem) == 64, "struct FAllocationItem needs packing");
-
-////////////////////////////////////////////////////////////////////////////////////////////////////
-
 class FShortLivingAllocs
 {
 private:
@@ -188,120 +157,6 @@ private:
 
 ////////////////////////////////////////////////////////////////////////////////////////////////////
 
-class FAllocsMap
-{
-public:
-	FAllocsMap()
-	{
-	}
-
-	~FAllocsMap()
-	{
-		Reset();
-	}
-
-	void Reset()
-	{
-		for (const auto& KV : AllocsMap)
-		{
-			const FAllocationItem* Allocation = KV.Value;
-			delete Allocation;
-		}
-		AllocsMap.Reset();
-	}
-
-	int32 Num() const
-	{
-		return AllocsMap.Num();
-	}
-
-	// Finds the allocation with the specified address.
-	// Returns the found allocation or nullptr if not found.
-	FAllocationItem* FindRef(uint64 Address) const
-	{
-		return AllocsMap.FindRef(Address);
-	}
-
-	// Finds an allocation containing the specified address.
-	// Returns the found allocation or nullptr if not found.
-	FAllocationItem* FindRange(uint64 Address) const
-	{
-		for (const auto& Alloc : AllocsMap)
-		{
-			if (Alloc.Value->IsContained(Address))
-			{
-				return Alloc.Value;
-			}
-		}
-		return nullptr;
-	}
-
-	void Enumerate(TFunctionRef<void(const FAllocationItem& Alloc)> Callback) const
-	{
-		for (const auto& KV : AllocsMap)
-		{
-			Callback(*KV.Value);
-		}
-	}
-
-	void Enumerate(uint64 StartAddress, uint64 EndAddress, TFunctionRef<void(const FAllocationItem& Alloc)> Callback) const
-	{
-		for (const auto& KV : AllocsMap)
-		{
-			const FAllocationItem& Allocation = *KV.Value;
-			if (Allocation.Address >= StartAddress && Allocation.Address < EndAddress)
-			{
-				Callback(Allocation);
-			}
-		}
-	}
-
-	// The collection keeps ownership of FAllocationItem* until Remove is called.
-	void Add(FAllocationItem* Alloc)
-	{
-		AllocsMap.Add(Alloc->Address, Alloc);
-	}
-
-	// The caller takes ownership of FAllocationItem*. Returns nullptr if Address is not found.
-	FAllocationItem* Remove(uint64 Address)
-	{
-		FAllocationItem* RemovedLongLivingAlloc;
-		if (AllocsMap.RemoveAndCopyValue(Address, RemovedLongLivingAlloc))
-		{
-			return RemovedLongLivingAlloc;
-		}
-		return nullptr;
-	}
-
-private:
-	template<typename ValueType>
-	struct TAddressKeyFuncs : BaseKeyFuncs<TPair<uint64, ValueType>, uint64, false>
-	{
-		typedef typename TTypeTraits<uint64>::ConstPointerType KeyInitType;
-		typedef const TPairInitializer<typename TTypeTraits<uint64>::ConstInitType, typename TTypeTraits<ValueType>::ConstInitType>& ElementInitType;
-
-		static FORCEINLINE uint64 GetSetKey(ElementInitType Element)
-		{
-			return Element.Key;
-		}
-
-		static FORCEINLINE bool Matches(uint64 A, uint64 B)
-		{
-			return A == B;
-		}
-
-		static FORCEINLINE uint32 GetKeyHash(uint64 Key)
-		{
-			return uint32(Key >> 6);
-		}
-	};
-
-private:
-	TMap<uint64, FAllocationItem*, FDefaultSetAllocator, TAddressKeyFuncs<FAllocationItem*>> AllocsMap;
-};
-
-////////////////////////////////////////////////////////////////////////////////////////////////////
-
 class FHeapAllocs
 {
 private:
@@ -412,7 +267,7 @@ public:
 private:
 	FAllocationItem* LastAlloc = nullptr; // last allocation
 	FShortLivingAllocs ShortLivingAllocs; // short living allocations
-	FAllocsMap LongLivingAllocs; // long living allocations
+	FAllocMap LongLivingAllocs; // long living allocations
 	FHeapAllocs HeapAllocs; // heap allocations
 
 	uint32 TotalAllocCount = 0;
diff --git a/Engine/Source/Developer/TraceServices/Private/Model/AllocationsQuery.cpp b/Engine/Source/Developer/TraceServices/Private/Model/AllocationsQuery.cpp
--- a/Engine/Source/Developer/TraceServices/Private/Model/AllocationsQuery.cpp
+++ b/Engine/Source/Developer/TraceServices/Private/Model/AllocationsQuery.cpp
@@ -4,6 +4,8 @@
 
 #include "SbTree.h"
 #include "Common/Utils.h"
+#include "Model/AllocationItem.h"
+#include "Model/AllocationsProvider.h"
 #include "TraceServices/Containers/Allocators.h"
 
 #include <limits>
diff --git a/Engine/Source/Developer/TraceServices/Private/Model/AllocationsQuery.h b/Engine/Source/Developer/TraceServices/Private/Model/AllocationsQuery.h
--- a/Engine/Source/Developer/TraceServices/Private/Model/AllocationsQuery.h
+++ b/Engine/Source/Developer/TraceServices/Private/Model/AllocationsQuery.h
@@ -2,17 +2,19 @@
 
 #pragma once
 
-#include "AllocationsProvider.h"
 #include "CallstacksProvider.h"
 #include "Async/TaskGraphInterfaces.h"
 #include "Containers/Queue.h"
 #include "HAL/PlatformAtomics.h"
+#include "TraceServices/Model/AllocationsProvider.h"
 
 #include <atomic>
 
 namespace TraceServices
 {
 
+struct FAllocationItem;
+class FAllocationsProvider;
 class ILinearAllocator;
 
 struct FAllocationsImpl
diff --git a/Engine/Source/Developer/TraceServices/Private/Model/SbTree.cpp b/Engine/Source/Developer/TraceServices/Private/Model/SbTree.cpp
--- a/Engine/Source/Developer/TraceServices/Private/Model/SbTree.cpp
+++ b/Engine/Source/Developer/TraceServices/Private/Model/SbTree.cpp
@@ -3,8 +3,8 @@
 #include "SbTree.h"
 
 #include "Misc/FileHelper.h"
+#include "Model/AllocationItem.h"
 #include "TraceServices/Containers/Allocators.h"
-#include "Model/AllocationsProvider.h"
 
 #include <limits>
 
diff --git a/Engine/Source/Developer/TraceServices/Private/Model/SbTree.h b/Engine/Source/Developer/TraceServices/Private/Model/SbTree.h
--- a/Engine/Source/Developer/TraceServices/Private/Model/SbTree.h
+++ b/Engine/Source/Developer/TraceServices/Private/Model/SbTree.h
@@ -2,12 +2,13 @@
 
 #pragma once
 
-#include "AllocationsProvider.h"
 //#include "Common/PagedArray.h"
+#include "TraceServices/Model/AllocationsProvider.h"
 
 namespace TraceServices
 {
 
+struct FAllocationItem;
 class ILinearAllocator;
 
 ////////////////////////////////////////////////////////////////////////////////////////////////////
diff --git a/Engine/Source/Developer/TraceServices/Private/Model/TableImportTask.h b/Engine/Source/Developer/TraceServices/Private/Model/TableImportTask.h
--- a/Engine/Source/Developer/TraceServices/Private/Model/TableImportTask.h
+++ b/Engine/Source/Developer/TraceServices/Private/Model/TableImportTask.h
@@ -2,9 +2,10 @@
 
 #pragma once
 
-#include "AllocationsProvider.h"
-#include "Common/StringStore.h"
 #include "CoreMinimal.h"
+
+// TraceServices
+#include "Common/StringStore.h"
 #include "Model/Tables.h"
 #include "TraceServices/Model/TableImport.h"
 
diff --git a/Engine/Source/Developer/TraceServices/Public/TraceServices/Model/TableImport.h b/Engine/Source/Developer/TraceServices/Public/TraceServices/Model/TableImport.h
--- a/Engine/Source/Developer/TraceServices/Public/TraceServices/Model/TableImport.h
+++ b/Engine/Source/Developer/TraceServices/Public/TraceServices/Model/TableImport.h
@@ -2,7 +2,6 @@
 
 #pragma once
 
-#include "AllocationsProvider.h"
 #include "CoreMinimal.h"
 #include "TraceServices/Containers/Tables.h"
 
@@ -10,6 +9,7 @@ class FTokenizedMessage;
 
 namespace TraceServices
 {
+
 class FImportTableRow;
 
 class ITableImportData
