[Insights] Fixed FSbTree::SetTimeForEvent for cases where an Alloc or Free event is generated from within analysis (like in case of zero size allocs generated for unmatched free events or for free events generated to free allocs from heaps).

diff --git a/Engine/Source/Developer/TraceServices/Private/Model/SbTree.cpp b/Engine/Source/Developer/TraceServices/Private/Model/SbTree.cpp
--- a/Engine/Source/Developer/TraceServices/Private/Model/SbTree.cpp
+++ b/Engine/Source/Developer/TraceServices/Private/Model/SbTree.cpp
@@ -353,8 +353,15 @@ void FSbTree::SetTimeForEvent(uint32 EventIndex, double Time)
 	// Detect the first event of each new column.
 	if ((EventIndex & ((1 << ColumnShift) - 1)) == 0)
 	{
-		check(static_cast<uint32>(ColumnStartTimes.Num()) == (EventIndex >> ColumnShift));
-		ColumnStartTimes.Add(Time);
+		if (static_cast<uint32>(ColumnStartTimes.Num()) == (EventIndex >> ColumnShift))
+		{
+			ColumnStartTimes.Add(Time);
+		}
+		else
+		{
+			check(static_cast<uint32>(ColumnStartTimes.Num()) == (EventIndex >> ColumnShift) + 1);
+			check(ColumnStartTimes.Last() == Time);
+		}
 	}
 }
 
