[Insights] MemoryInsights: Added warning when detecting unsupported memory trace versions (supported versions are 1 and 2).

diff --git a/Engine/Source/Developer/TraceServices/Private/Analyzers/AllocationsAnalysis.cpp b/Engine/Source/Developer/TraceServices/Private/Analyzers/AllocationsAnalysis.cpp
--- a/Engine/Source/Developer/TraceServices/Private/Analyzers/AllocationsAnalysis.cpp
+++ b/Engine/Source/Developer/TraceServices/Private/Analyzers/AllocationsAnalysis.cpp
@@ -7,6 +7,7 @@
 
 // TraceServices
 #include "Common/ProviderLock.h"
+#include "Common/Utils.h"
 #include "Model/AllocationsProvider.h"
 #include "Model/MetadataProvider.h"
 #include "TraceServices/Model/AnalysisSession.h"
@@ -17,6 +18,14 @@ namespace TraceServices
 
 #define INSIGHTS_MEM_TRACE_METADATA_TEST 0
 
+namespace AllocationsAnalyzer::Private
+{
+	// version 1: UE 5.0
+	// version 2: UE 5.2
+	constexpr int32 MinSupportedVersion = 1;
+	constexpr int32 MaxSupportedVersion = 2;
+}
+
 ////////////////////////////////////////////////////////////////////////////////////////////////////
 
 FAllocationsAnalyzer::FAllocationsAnalyzer(IAnalysisSession& InSession, FAllocationsProvider& InAllocationsProvider, FMetadataProvider& InMetadataProvider)
@@ -97,12 +106,12 @@ bool FAllocationsAnalyzer::OnEvent(uint16 RouteId, EStyle Style, const FOnEventC
 	{
 		case RouteId_Init:
 		{
-			const uint8 Version = EventData.GetValue<uint8>("Version");
+			const uint32 Version = EventData.GetValue<uint32>("Version", 0);
 
-			constexpr uint8 MinSupportedVersion = 1; // UE 5.0
-			constexpr uint8 MaxSupportedVersion = 1; // UE 5.0
+			using namespace AllocationsAnalyzer::Private;
 			if (Version < MinSupportedVersion || Version > MaxSupportedVersion)
 			{
+				UE_LOG(LogTraceServices, Error, TEXT("[MemAlloc] Version %u for Memory trace events is not supported by the current analyzer. Supported versions: [%u .. %u]"), Version, MinSupportedVersion, MaxSupportedVersion);
 				break;
 			}
 
@@ -150,9 +159,9 @@ bool FAllocationsAnalyzer::OnEvent(uint16 RouteId, EStyle Style, const FOnEventC
 			const double Time = GetCurrentTime();
 
 			// CallstackId is optional. If the field is not present CallstackId will be 0 (i.e. "no callstack").
-			uint32 CallstackId = EventData.GetValue<uint32>("CallstackId", 0);
+			const uint32 CallstackId = EventData.GetValue<uint32>("CallstackId", 0);
 
-			uint64 Address = EventData.GetValue<uint64>("Address");
+			const uint64 Address = EventData.GetValue<uint64>("Address");
 
 			RootHeap = EventData.GetValue<uint8>("RootHeap", static_cast<uint8>(RootHeap));
 
@@ -190,9 +199,9 @@ bool FAllocationsAnalyzer::OnEvent(uint16 RouteId, EStyle Style, const FOnEventC
 			const double Time = GetCurrentTime();
 
 			// CallstackId is optional. If the field is not present CallstackId will be 0 (i.e. "no callstack").
-			uint32 CallstackId = EventData.GetValue<uint32>("CallstackId", 0);
+			const uint32 CallstackId = EventData.GetValue<uint32>("CallstackId", 0);
 
-			uint64 Address = EventData.GetValue<uint64>("Address");
+			const uint64 Address = EventData.GetValue<uint64>("Address");
 
 			RootHeap = EventData.GetValue<uint8>("RootHeap", static_cast<uint8>(RootHeap));
 
@@ -213,6 +222,10 @@ bool FAllocationsAnalyzer::OnEvent(uint16 RouteId, EStyle Style, const FOnEventC
 		case RouteId_HeapMarkAlloc:
 		{
 			const double Time = GetCurrentTime();
+
+			// CallstackId is optional. If the field is not present CallstackId will be 0 (i.e. "no callstack").
+			const uint32 CallstackId = EventData.GetValue<uint32>("CallstackId", 0);
+
 			const uint64 Address = EventData.GetValue<uint64>("Address");
 			const HeapId Heap = EventData.GetValue<uint16>("Heap", 0);
 			const EMemoryTraceHeapAllocationFlags Flags = EventData.GetValue<EMemoryTraceHeapAllocationFlags>("Flags");
@@ -222,13 +235,17 @@ bool FAllocationsAnalyzer::OnEvent(uint16 RouteId, EStyle Style, const FOnEventC
 
 			FProviderEditScopeLock _(AllocationsProvider);
 			AllocationsProvider.SetCurrentThreadId(TraceThreadId, SystemThreadId);
-			AllocationsProvider.EditMarkAllocationAsHeap(Time, Address, Heap, Flags);
+			AllocationsProvider.EditMarkAllocationAsHeap(Time, /*CallstackId,*/ Address, Heap, Flags);
 			break;
 		}
 
 		case RouteId_HeapUnmarkAlloc:
 		{
 			const double Time = GetCurrentTime();
+
+			// CallstackId is optional. If the field is not present CallstackId will be 0 (i.e. "no callstack").
+			const uint32 CallstackId = EventData.GetValue<uint32>("CallstackId", 0);
+
 			const uint64 Address = EventData.GetValue<uint64>("Address");
 			const HeapId Heap = EventData.GetValue<uint16>("Heap", 0);
 
@@ -237,7 +254,7 @@ bool FAllocationsAnalyzer::OnEvent(uint16 RouteId, EStyle Style, const FOnEventC
 
 			FProviderEditScopeLock _(AllocationsProvider);
 			AllocationsProvider.SetCurrentThreadId(TraceThreadId, SystemThreadId);
-			AllocationsProvider.EditUnmarkAllocationAsHeap(Time, Address, Heap);
+			AllocationsProvider.EditUnmarkAllocationAsHeap(Time, /*CallstackId,*/ Address, Heap);
 			break;
 		}
 
