CIS fix nonunity build errors

#rnx

diff --git a/Engine/Source/Developer/TraceServices/Private/Model/TableMergeTask.cpp b/Engine/Source/Developer/TraceServices/Private/Model/TableMergeTask.cpp
--- a/Engine/Source/Developer/TraceServices/Private/Model/TableMergeTask.cpp
+++ b/Engine/Source/Developer/TraceServices/Private/Model/TableMergeTask.cpp
@@ -4,6 +4,7 @@
 
 #include "Logging/TokenizedMessage.h"
 #include "Tasks/Task.h"
+#include "Async/TaskGraphInterfaces.h"
 
 namespace TraceServices
 {
diff --git a/Engine/Source/Developer/TraceServices/Private/Model/TableMergeTask.h b/Engine/Source/Developer/TraceServices/Private/Model/TableMergeTask.h
--- a/Engine/Source/Developer/TraceServices/Private/Model/TableMergeTask.h
+++ b/Engine/Source/Developer/TraceServices/Private/Model/TableMergeTask.h
@@ -5,6 +5,8 @@
 #include "Tables.h"
 #include "TraceServices/Model/TableMerge.h"
 
+#include "Misc/TVariant.h"
+
 namespace TraceServices
 {
 class FImportTableRow;
@@ -24,7 +26,7 @@ private:
 	struct FOnlyB { FOnlyB(uint64 Index) : ColIndex(Index) {} uint64 ColIndex; };
 	struct FMerged { FMerged(uint64 Index) : ColIndex(Index) {} uint64 ColIndex; };
 	enum class Sign { Positive, Negative };
-	
+
 	bool BuildCLayout(const ITableLayout& LayoutA, const ITableLayout& LayoutB, TTableLayout<FImportTableRow>& LayoutC);
 
 	template<typename Variant, Sign Operation>
