Automatically remap known renamed binaries.

diff --git a/Engine/Source/Developer/TraceServices/Private/Model/SymslibResolver.cpp b/Engine/Source/Developer/TraceServices/Private/Model/SymslibResolver.cpp
--- a/Engine/Source/Developer/TraceServices/Private/Model/SymslibResolver.cpp
+++ b/Engine/Source/Developer/TraceServices/Private/Model/SymslibResolver.cpp
@@ -163,18 +163,7 @@ namespace
 		IPlatformFile* PlatformFile = &FPlatformFileManager::Get().GetPlatformFile();
 		const TCHAR* EngineFolders[] = {TEXT("Engine"), TEXT("Plugins")};
 		FString Result;
-
-		// If this file name has been renamed try mapping it to known
-		if (FilePath.Compare(TEXT("eboot.bin")) == 0)
-		{
-			Result = FPaths::Combine(SearchPath, Project, TEXT("Binaries"), Platform, Project, TEXT(".self"));
-			if (PlatformFile->FileExists(*Result))
-			{
-				return Result;
-			}
-		}
-
-		const FString File = FPaths::GetCleanFilename(FilePath);
+		FString File = FPaths::GetCleanFilename(FilePath);
 
 		// Look in some well known directories
 		for (const TCHAR* Folder : EngineFolders)
@@ -232,9 +221,15 @@ namespace
 
 	static bool LoadBinary(const TCHAR* Path, SYMS_Arena* Arena, SYMS_ParseBundle& Bundle, TArray<FAutoMappedFile>& Files, const FString& SearchPath, const FString& Platform, const FString& Project)
 	{
-		const FString FilePath = Path;
+		FString FilePath = Path;
 		bool bFileFound = false;
 
+		// Remap known renamed binaries
+		if (FilePath.EndsWith(TEXT("eboot.bin")) && !Platform.IsEmpty() && !Project.IsEmpty())
+		{
+			FilePath = FPaths::Combine(FPaths::GetPath(FilePath), FString::Printf(TEXT("%sGame.self"), *Project));
+		}
+
 		// First lookup file in symbol path
 		FString BinaryPath = FindBinaryFileInPath(FilePath, SearchPath);
 		if (BinaryPath.IsEmpty() && !Platform.IsEmpty())
@@ -530,9 +525,8 @@ FSymslibResolver::FSymslibResolver(IAnalysisSession& InSession, IResolvedSymbolF
 #endif
 
 	// Try to get session information
-	// todo: Sadly this will not be available in the constructor or when the first modules are traced, since the
-	//		 session info is traced later in the process startup sequence.
-#if 0
+	// Depending on how module tracing and session info tracing is implemented on different platforms this may not
+	// be available at this point.
 	FAnalysisSessionReadScope _(Session);
 	const IDiagnosticsProvider* DiagnosticsProvider = ReadDiagnosticsProvider(Session);
 	if (DiagnosticsProvider && DiagnosticsProvider->IsSessionInfoAvailable())
@@ -541,7 +535,6 @@ FSymslibResolver::FSymslibResolver(IAnalysisSession& InSession, IResolvedSymbolF
 		Platform = Info.Platform;
 		Project = Info.AppName;
 	}
-#endif
 }
 
 FSymslibResolver::~FSymslibResolver()
