[Insights] Fixed conversion warnings in the TraceServices module.

diff --git a/Engine/Source/Developer/TraceServices/Private/Common/Utils.cpp b/Engine/Source/Developer/TraceServices/Private/Common/Utils.cpp
--- a/Engine/Source/Developer/TraceServices/Private/Common/Utils.cpp
+++ b/Engine/Source/Developer/TraceServices/Private/Common/Utils.cpp
@@ -18,8 +18,8 @@ void Table2Csv(const IUntypedTable& Table, const TCHAR* Filename)
 	check(OutputFile);
 	FString Header;
 	const ITableLayout& Layout = Table.GetLayout();
-	uint32 ColumnCount = static_cast<uint32>(Layout.GetColumnCount());
-	for (uint32 ColumnIndex = 0; ColumnIndex < ColumnCount; ++ColumnIndex)
+	const uint64 ColumnCount = Layout.GetColumnCount();
+	for (uint64 ColumnIndex = 0; ColumnIndex < ColumnCount; ++ColumnIndex)
 	{
 		Header += Layout.GetColumnName(ColumnIndex);
 		if (ColumnIndex < ColumnCount - 1)
@@ -37,7 +37,7 @@ void Table2Csv(const IUntypedTable& Table, const TCHAR* Filename)
 	for (; TableReader->IsValid(); TableReader->NextRow())
 	{
 		FString Line;
-		for (uint32 ColumnIndex = 0; ColumnIndex < ColumnCount; ++ColumnIndex)
+		for (uint64 ColumnIndex = 0; ColumnIndex < ColumnCount; ++ColumnIndex)
 		{
 			switch (Layout.GetColumnType(ColumnIndex))
 			{
diff --git a/Engine/Source/Developer/TraceServices/Private/Model/TableImportTask.h b/Engine/Source/Developer/TraceServices/Private/Model/TableImportTask.h
--- a/Engine/Source/Developer/TraceServices/Private/Model/TableImportTask.h
+++ b/Engine/Source/Developer/TraceServices/Private/Model/TableImportTask.h
@@ -19,11 +19,21 @@ struct FColumnValueContainer;
 class FImportTableRow
 {
 public:
-	FColumnValueContainer GetValue(uint32 Index) const { return Values[Index]; }
+	FColumnValueContainer GetValue(uint64 ColumnIndex) const
+	{
+		return Values[static_cast<int32>(ColumnIndex)];
+	}
 
 	template<typename T>
-	void SetValue(uint32 Index, T Value) { Values[Index] = Value; }
-	void SetNumValues(uint32 Size) { Values.AddUninitialized(Size); }
+	void SetValue(uint64 ColumnIndex, T Value)
+	{
+		Values[static_cast<int32>(ColumnIndex)] = Value;
+	}
+
+	void SetNumValues(uint64 NumColumns)
+	{
+		Values.AddUninitialized(static_cast<int32>(NumColumns));
+	}
 
 private:
 	TArray<FColumnValueContainer> Values;
diff --git a/Engine/Source/Developer/TraceServices/Private/Model/TableMergeTask.h b/Engine/Source/Developer/TraceServices/Private/Model/TableMergeTask.h
--- a/Engine/Source/Developer/TraceServices/Private/Model/TableMergeTask.h
+++ b/Engine/Source/Developer/TraceServices/Private/Model/TableMergeTask.h
@@ -47,26 +47,28 @@ private:
 template <typename Variant, FTableMergeTask::Sign Operation>
 void FTableMergeTask::FillResultUsing(IUntypedTableReader& SourceTableReader) const
 {
+	const uint64 ColumnsMapCount = static_cast<uint64>(ColumnsMap.Num());
+
 	for (SourceTableReader.SetRowIndex(0); SourceTableReader.IsValid(); SourceTableReader.NextRow())
 	{
 		FImportTableRow& Row = TableC->AddRow();
-		Row.SetNumValues(ColumnsMap.Num());
+		Row.SetNumValues(ColumnsMapCount);
 
-		for (uint64 DestIndex = 0; DestIndex < ColumnsMap.Num(); ++DestIndex)
+		for (uint64 DestIndex = 0; DestIndex < ColumnsMapCount; ++DestIndex)
 		{
 			const ETableColumnType ColumnType = TableC->GetLayout().GetColumnType(DestIndex);
 
 			uint64 SourceIndex;
 			int64 Multiplier = 1;
-			if (const auto* Column = ColumnsMap[DestIndex].TryGet<FBoth>())
+			if (const auto* Column = ColumnsMap[static_cast<int32>(DestIndex)].TryGet<FBoth>())
 			{
 				SourceIndex = Column->ColIndex;
 			}
-			else if (const auto* ColumnA = ColumnsMap[DestIndex].TryGet<Variant>())
+			else if (const auto* ColumnA = ColumnsMap[static_cast<int32>(DestIndex)].TryGet<Variant>())
 			{
 				SourceIndex = ColumnA->ColIndex;
 			}
-			else if (const auto* ColumnM = ColumnsMap[DestIndex].TryGet<FMerged>())
+			else if (const auto* ColumnM = ColumnsMap[static_cast<int32>(DestIndex)].TryGet<FMerged>())
 			{
 				SourceIndex = ColumnM->ColIndex;
 				Multiplier = Operation == Sign::Negative ? -1 : 1;
diff --git a/Engine/Source/Developer/TraceServices/Private/Model/Tables.h b/Engine/Source/Developer/TraceServices/Private/Model/Tables.h
--- a/Engine/Source/Developer/TraceServices/Private/Model/Tables.h
+++ b/Engine/Source/Developer/TraceServices/Private/Model/Tables.h
@@ -211,32 +211,32 @@ public:
 
 	uint64 GetColumnCount() const override
 	{
-		return Columns.Num();
+		return static_cast<uint64>(Columns.Num());
 	}
 
 	const TCHAR* GetColumnName(uint64 ColumnIndex) const override
 	{
-		return *Columns[(int32)ColumnIndex].Name;
+		return *Columns[static_cast<int32>(ColumnIndex)].Name;
 	}
 
 	ETableColumnType GetColumnType(uint64 ColumnIndex) const override
 	{
-		return Columns[(int32)ColumnIndex].Type;
+		return Columns[static_cast<int32>(ColumnIndex)].Type;
 	}
 
 	void SetColumnType(uint64 ColumnIndex, ETableColumnType ColumnType)
 	{
-		Columns[(int32)ColumnIndex].Type = ColumnType;
+		Columns[static_cast<int32>(ColumnIndex)].Type = ColumnType;
 	}
 
 	uint32 GetColumnDisplayHintFlags(uint64 ColumnIndex) const override
 	{
-		return Columns[(int32)ColumnIndex].DisplayHintFlags;
+		return Columns[static_cast<int32>(ColumnIndex)].DisplayHintFlags;
 	}
 
 	FColumnValueContainer GetColumnValue(const RowType& Row, uint64 ColumnIndex) const
 	{
-		return Columns[(int32)ColumnIndex].Projector(Row);
+		return Columns[static_cast<int32>(ColumnIndex)].Projector(Row);
 	}
 
 private:
