[Insights]
  - MemoryInsights: Added analysis support for NeverFrees heaps.
  - MemoryInsights: Added memory leak detection when a heap is freed (i.e. remaining live allocs) -- but disabled for now, as it is too slow without a better Live Allocs data structure.
  - MemoryInsights: Added limit for the number of log warnings generated.

diff --git a/Engine/Source/Developer/TraceServices/Private/Analyzers/AllocationsAnalysis.cpp b/Engine/Source/Developer/TraceServices/Private/Analyzers/AllocationsAnalysis.cpp
--- a/Engine/Source/Developer/TraceServices/Private/Analyzers/AllocationsAnalysis.cpp
+++ b/Engine/Source/Developer/TraceServices/Private/Analyzers/AllocationsAnalysis.cpp
@@ -235,7 +235,7 @@ bool FAllocationsAnalyzer::OnEvent(uint16 RouteId, EStyle Style, const FOnEventC
 
 			FProviderEditScopeLock _(AllocationsProvider);
 			AllocationsProvider.SetCurrentThreadId(TraceThreadId, SystemThreadId);
-			AllocationsProvider.EditMarkAllocationAsHeap(Time, /*CallstackId,*/ Address, Heap, Flags);
+			AllocationsProvider.EditMarkAllocationAsHeap(Time, CallstackId, Address, Heap, Flags);
 			break;
 		}
 
@@ -254,7 +254,7 @@ bool FAllocationsAnalyzer::OnEvent(uint16 RouteId, EStyle Style, const FOnEventC
 
 			FProviderEditScopeLock _(AllocationsProvider);
 			AllocationsProvider.SetCurrentThreadId(TraceThreadId, SystemThreadId);
-			AllocationsProvider.EditUnmarkAllocationAsHeap(Time, /*CallstackId,*/ Address, Heap);
+			AllocationsProvider.EditUnmarkAllocationAsHeap(Time, CallstackId, Address, Heap);
 			break;
 		}
 
diff --git a/Engine/Source/Developer/TraceServices/Private/Model/AllocationsProvider.cpp b/Engine/Source/Developer/TraceServices/Private/Model/AllocationsProvider.cpp
--- a/Engine/Source/Developer/TraceServices/Private/Model/AllocationsProvider.cpp
+++ b/Engine/Source/Developer/TraceServices/Private/Model/AllocationsProvider.cpp
@@ -59,6 +59,7 @@
 namespace TraceServices
 {
 
+constexpr uint32 MaxLogMessagesPerWarningType = 100;
 constexpr uint32 MaxLogMessagesPerErrorType = 100;
 
 #if INSIGHTS_DEBUG_WATCH
@@ -1276,14 +1277,22 @@ void FAllocationsProvider::EditAlloc(double Time, uint32 CallstackId, uint64 Add
 
 	if (Address == 0)
 	{
-		UE_LOG(LogTraceServices, Warning, TEXT("[MemAlloc] Alloc at address 0 : Size=%llu, RootHeap=%u, Time=%f, CallstackId=%u"), InSize, RootHeap, Time, CallstackId);
+		++AllocWarnings;
+		if (AllocWarnings <= MaxLogMessagesPerWarningType)
+		{
+			UE_LOG(LogTraceServices, Warning, TEXT("[MemAlloc] Alloc at address 0 : Size=%llu, RootHeap=%u, Time=%f, CallstackId=%u"), InSize, RootHeap, Time, CallstackId);
+		}
 		return;
 	}
 
 	check(RootHeap < MaxRootHeaps);
 	if (SbTree[RootHeap] == nullptr)
 	{
-		UE_LOG(LogTraceServices, Warning, TEXT("[MemAlloc] Invalid root heap id (%u). This heap has not yet been registered."), RootHeap);
+		++AllocWarnings;
+		if (AllocWarnings <= MaxLogMessagesPerWarningType)
+		{
+			UE_LOG(LogTraceServices, Warning, TEXT("[MemAlloc] Alloc with invalid root heap id (%u). This heap has not yet been registered."), RootHeap);
+		}
 		return;
 	}
 
@@ -1400,14 +1409,22 @@ void FAllocationsProvider::EditFree(double Time, uint32 CallstackId, uint64 Addr
 
 	if (Address == 0)
 	{
-		UE_LOG(LogTraceServices, Warning, TEXT("[MemAlloc] Free for address 0 : RootHeap=%u, Time=%f, CallstackId=%u"), RootHeap, Time, CallstackId);
+		++FreeWarnings;
+		if (FreeWarnings <= MaxLogMessagesPerWarningType)
+		{
+			UE_LOG(LogTraceServices, Warning, TEXT("[MemAlloc] Free for address 0 : RootHeap=%u, Time=%f, CallstackId=%u"), RootHeap, Time, CallstackId);
+		}
 		return;
 	}
 
 	check(RootHeap < MaxRootHeaps);
 	if (SbTree[RootHeap] == nullptr)
 	{
-		UE_LOG(LogTraceServices, Warning, TEXT("[MemAlloc] Invalid root heap id (%u). This heap has not yet been registered."), RootHeap);
+		++FreeWarnings;
+		if (FreeWarnings <= MaxLogMessagesPerWarningType)
+		{
+			UE_LOG(LogTraceServices, Warning, TEXT("[MemAlloc] Free with invalid root heap id (%u). This heap has not yet been registered."), RootHeap);
+		}
 		return;
 	}
 
@@ -1430,8 +1447,15 @@ void FAllocationsProvider::EditFree(double Time, uint32 CallstackId, uint64 Addr
 	if (!AllocationPtr)
 	{
 		// Try to free a heap allocation.
-		AllocationPtr = LiveAllocs[RootHeap]->RemoveHeap(Address); // we take ownership of AllocationPtr
-		INSIGHTS_SLOW_CHECK(!AllocationPtr || AllocationPtr->IsHeap());
+		AllocationPtr = LiveAllocs[RootHeap]->FindHeapRef(Address);
+		if (AllocationPtr)
+		{
+			INSIGHTS_SLOW_CHECK(AllocationPtr->IsHeap());
+			HeapId Heap = AllocationPtr->RootHeap;
+			EditUnmarkAllocationAsHeap(Time, CallstackId, Address, Heap);
+			SbTree[RootHeap]->SetTimeForEvent(EventIndex[RootHeap], Time); // for the case where the free event is first event in a new SbTree column after changing the heap alloc
+			AllocationPtr = LiveAllocs[RootHeap]->Remove(Address); // we take ownership of AllocationPtr
+		}
 	}
 	else
 	{
@@ -1539,7 +1563,11 @@ void FAllocationsProvider::AddHeapSpec(HeapId Id, HeapId ParentId, const FString
 		FHeapSpec& ParentSpec = HeapSpecs[ParentId];
 		if (ParentSpec.Name == nullptr)
 		{
-			UE_LOG(LogTraceServices, Warning, TEXT("[MemAlloc] Parent heap id (%u) used before it was announced for heap %s."), ParentId, *FString(Name));
+			++HeapWarnings;
+			if (HeapWarnings <= MaxLogMessagesPerWarningType)
+			{
+				UE_LOG(LogTraceServices, Warning, TEXT("[MemAlloc] Parent heap id (%u) used before it was announced for heap %s."), ParentId, *FString(Name));
+			}
 		}
 		FHeapSpec& NewSpec = HeapSpecs[Id];
 		NewSpec.Name = Session.StoreString(Name);
@@ -1553,10 +1581,8 @@ void FAllocationsProvider::AddHeapSpec(HeapId Id, HeapId ParentId, const FString
 
 ////////////////////////////////////////////////////////////////////////////////////////////////////
 
-void FAllocationsProvider::EditMarkAllocationAsHeap(double Time, uint64 Address, HeapId Heap, EMemoryTraceHeapAllocationFlags Flags)
+void FAllocationsProvider::EditMarkAllocationAsHeap(double Time, uint32 CallstackId, uint64 Address, HeapId Heap, EMemoryTraceHeapAllocationFlags Flags)
 {
-	uint32 CallstackId = 0; // TODO
-
 	Lock.WriteAccessCheck();
 
 	INSIGHTS_API_LOGF(TEXT("MarkAllocAsHeap(Time=%f, Address=0x%llX, Heap=%u, Flags=0x%X, CallstackId=%u)"), Time, Address, Heap, uint32(Flags), CallstackId);
@@ -1572,16 +1598,49 @@ void FAllocationsProvider::EditMarkAllocationAsHeap(double Time, uint64 Address,
 	}
 #endif // INSIGHTS_DEBUG_WATCH
 
+	if (Heap >= (uint32)HeapSpecs.Num())
+	{
+		++HeapErrors;
+		if (HeapErrors <= MaxLogMessagesPerErrorType)
+		{
+			UE_LOG(LogTraceServices, Error, TEXT("[MemAlloc] HeapMarkAlloc: Heap %u is not valid (Address=0x%llX, Time=%f, CallstackId=%u)!"), Heap, Address, Time, CallstackId);
+		}
+		return;
+	}
+	const FHeapSpec& HeapSpec = HeapSpecs[Heap];
+
 	const HeapId RootHeap = FindRootHeap(Heap);
+#if 0 // TODO
+	ensure(RootHeap != Heap);
+	if (Heap == RootHeap)
+	{
+		++HeapWarnings;
+		if (HeapWarnings <= MaxLogMessagesPerWarningType)
+		{
+			UE_LOG(LogTraceServices, Warning, TEXT("[MemAlloc] HeapMarkAlloc: Heap %u is a root heap (Address=0x%llX, Time=%f, CallstackId=%u)!"), Heap, Address, Time, CallstackId);
+		}
+	}
+#endif
 
 	// Remove the allocation from the Live allocs.
 	FAllocationItem* Alloc = LiveAllocs[RootHeap]->Remove(Address); // we take ownership of Alloc
 	if (Alloc)
 	{
+		check(Address == Alloc->Address);
+
+		if (EnumHasAnyFlags(Alloc->Flags, EMemoryTraceHeapAllocationFlags::Heap))
+		{
+			++HeapErrors;
+			if (HeapErrors <= MaxLogMessagesPerErrorType)
+			{
+				UE_LOG(LogTraceServices, Error, TEXT("[MemAlloc] HeapMarkAlloc: Alloc 0x%llX is already marked as heap (Heap=%u, Flags=%u, Time=%f, CallstackId=%u)!"), Address, Heap, uint32(Flags), Time, CallstackId);
+			}
+		}
+
 		// Mark allocation as a "heap" allocation.
-		check(EnumHasAnyFlags(Flags, EMemoryTraceHeapAllocationFlags::Heap));
-		Alloc->Flags = Flags;
-		Alloc->RootHeap = static_cast<uint8>(RootHeap);
+		ensure(EnumHasAnyFlags(Flags, EMemoryTraceHeapAllocationFlags::Heap));
+		Alloc->Flags = Flags | EMemoryTraceHeapAllocationFlags::Heap;
+		Alloc->RootHeap = static_cast<uint8>(Heap);
 
 		// Re-add it to the Live allocs as a heap allocation.
 		LiveAllocs[RootHeap]->AddHeap(Alloc); // the Live allocs takes ownership of Alloc
@@ -1597,7 +1656,7 @@ void FAllocationsProvider::EditMarkAllocationAsHeap(double Time, uint64 Address,
 		++HeapErrors;
 		if (HeapErrors <= MaxLogMessagesPerErrorType)
 		{
-			UE_LOG(LogTraceServices, Error, TEXT("[MemAlloc] HeapMarkAlloc: Could not find address 0x%llX (Heap=%u, Flags=%u, Time=%f, CallstackId=%u)!"), Address, Heap, uint32(Flags), Time, CallstackId);
+			UE_LOG(LogTraceServices, Error, TEXT("[MemAlloc] HeapMarkAlloc: Could not find address 0x%llX (Heap=%u \"%s\", Flags=%u, Time=%f, CallstackId=%u)!"), Address, Heap, HeapSpec.Name, uint32(Flags), Time, CallstackId);
 		}
 	}
 
@@ -1606,10 +1665,8 @@ void FAllocationsProvider::EditMarkAllocationAsHeap(double Time, uint64 Address,
 
 ////////////////////////////////////////////////////////////////////////////////////////////////////
 
-void FAllocationsProvider::EditUnmarkAllocationAsHeap(double Time, uint64 Address, HeapId Heap)
+void FAllocationsProvider::EditUnmarkAllocationAsHeap(double Time, uint32 CallstackId, uint64 Address, HeapId Heap)
 {
-	uint32 CallstackId = 0; // TODO
-
 	Lock.WriteAccessCheck();
 
 	INSIGHTS_API_LOGF(TEXT("UnmarkAllocAsHeap(Time=%f, Address=0x%llX, Heap=%u, CallstackId=%u)"), Time, Address, Heap, CallstackId);
@@ -1625,12 +1682,103 @@ void FAllocationsProvider::EditUnmarkAllocationAsHeap(double Time, uint64 Addres
 	}
 #endif // INSIGHTS_DEBUG_WATCH
 
+	if (Heap >= (uint32)HeapSpecs.Num())
+	{
+		++HeapErrors;
+		if (HeapErrors <= MaxLogMessagesPerErrorType)
+		{
+			UE_LOG(LogTraceServices, Error, TEXT("[MemAlloc] HeapUnmarkAlloc: Heap %u is not valid (Address=0x%llX, Time=%f, CallstackId=%u)!"), Heap, Address, Time, CallstackId);
+		}
+		return;
+	}
+	const FHeapSpec& HeapSpec = HeapSpecs[Heap];
+
 	const HeapId RootHeap = FindRootHeap(Heap);
+#if 0 // TODO
+	ensure(RootHeap != Heap);
+	if (Heap == RootHeap)
+	{
+		++HeapWarnings;
+		if (HeapWarnings <= MaxLogMessagesPerWarningType)
+		{
+			UE_LOG(LogTraceServices, Warning, TEXT("[MemAlloc] HeapUnmarkAlloc: Heap %u is a root heap (Address=0x%llX, Time=%f, CallstackId=%u)!"), Heap, Address, Time, CallstackId);
+		}
+	}
+#endif
 
 	// Remove the heap allocation from the Live allocs.
 	FAllocationItem* Alloc = LiveAllocs[RootHeap]->RemoveHeap(Address); // we take ownership of Alloc
 	if (Alloc)
 	{
+		check(Address == Alloc->Address);
+
+		if (!EnumHasAnyFlags(Alloc->Flags, EMemoryTraceHeapAllocationFlags::Heap))
+		{
+			++HeapErrors;
+			if (HeapErrors <= MaxLogMessagesPerErrorType)
+			{
+				UE_LOG(LogTraceServices, Error, TEXT("[MemAlloc] HeapUnmarkAlloc: Alloc 0x%llX is not marked as heap (Heap=%u, Time=%f, CallstackId=%u)!"), Address, Heap, Time, CallstackId);
+			}
+		}
+
+		//constexpr bool bSearchLiveAllocs = false; // never
+		const bool bSearchLiveAllocs = EnumHasAnyFlags(HeapSpec.Flags, EMemoryTraceHeapFlags::NeverFrees); // only for heaps with NeverFrees flag
+		//constexpr bool bSearchLiveAllocs = true; // for all heaps
+
+		if (bSearchLiveAllocs)
+		{
+			// Find all live allocs in this heap.
+			uint64 AllocatedSizeInHeap = 0;
+			struct FAllocInHeap
+			{
+				uint64 Address;
+				uint64 Size;
+			};
+			TArray<FAllocInHeap> AllocsInHeap;
+			const uint64 HeapSize = Alloc->GetSize();
+			const uint64 EndAddress = Alloc->Address + HeapSize;
+			LiveAllocs[RootHeap]->Enumerate(Address, EndAddress, [&AllocatedSizeInHeap, &AllocsInHeap](const FAllocationItem& ChildAlloc)
+			{
+				const uint64 ChildAllocSize = ChildAlloc.GetSize();
+				AllocatedSizeInHeap += ChildAllocSize;
+				AllocsInHeap.Add({ ChildAlloc.Address, ChildAllocSize });
+			});
+
+			if (AllocsInHeap.Num() > 0)
+			{
+				if (!EnumHasAnyFlags(HeapSpec.Flags, EMemoryTraceHeapFlags::NeverFrees))
+				{
+					// For heaps that do not have NeverFrees flag, we report live allocs as leaks.
+					UE_LOG(LogTraceServices, Warning, TEXT("[MemAlloc] HeapUnmarkAlloc: %d memory leaks (%llu bytes) detected for heap %u (\"%s\", Address=0x%llX, Size=%llu, Time=%f, CallstackId=%u)"),
+						AllocsInHeap.Num(), AllocatedSizeInHeap, Heap, HeapSpec.Name, Address, HeapSize, Time, CallstackId);
+					UE_LOG(LogTraceServices, Warning, TEXT("[MemAlloc] Top memory leaks:"));
+					AllocsInHeap.Sort([this](const FAllocInHeap& A, const FAllocInHeap& B) -> bool { return A.Size > B.Size; });
+					int NumAllocsInTop = 10; // only show first 10 allocs
+					for (const FAllocInHeap& AllocInHeap : AllocsInHeap)
+					{
+						UE_LOG(LogTraceServices, Warning, TEXT("[MemAlloc]     alloc 0x%llX (%u bytes)"),  AllocInHeap.Address, AllocInHeap.Size);
+						if (--NumAllocsInTop <= 0)
+						{
+							break;
+						}
+					}
+				}
+#if 0 // debug
+				else
+				{
+					UE_LOG(LogTraceServices, Log, TEXT("[MemAlloc] HeapUnmarkAlloc: %d memory allocs (%llu bytes) freed for heap %u (\"%s\", Address=0x%llX, Size=%llu, Time=%f, CallstackId=%u)"),
+						AllocsInHeap.Num(), AllocatedSizeInHeap, Heap, HeapSpec.Name, Address, HeapSize, Time, CallstackId);
+				}
+#endif
+
+				// Free automatically all allocs in this heap.
+				for (const FAllocInHeap& AllocInHeap : AllocsInHeap)
+				{
+					EditFree(Time, CallstackId, AllocInHeap.Address, RootHeap);
+				}
+			}
+		}
+
 		const uint64 Size = Alloc->GetSize();
 		const uint32 Alignment = Alloc->GetAlignment();
 		const uint32 AllocCallstackId = Alloc->CallstackId;
@@ -1652,7 +1800,7 @@ void FAllocationsProvider::EditUnmarkAllocationAsHeap(double Time, uint64 Addres
 		++HeapErrors;
 		if (HeapErrors <= MaxLogMessagesPerErrorType)
 		{
-			UE_LOG(LogTraceServices, Error, TEXT("[MemAlloc] HeapUnmarkAlloc: Could not find address 0x%llX (Heap=%u, Time=%f, CallstackId=%u)!"), Address, Heap, Time, CallstackId);
+			UE_LOG(LogTraceServices, Error, TEXT("[MemAlloc] HeapUnmarkAlloc: Could not find address 0x%llX (Heap=%u \"%s\", Time=%f, CallstackId=%u)!"), Address, Heap, HeapSpec.Name, Time, CallstackId);
 		}
 	}
 }
@@ -1859,8 +2007,9 @@ void FAllocationsProvider::EditOnAnalysisCompleted(double Time)
 	DebugPrint();
 #endif
 
-	for (const FSbTree* Tree : SbTree)
+	for (uint32 RootHeap = 0; RootHeap < MaxRootHeaps; ++RootHeap)
 	{
+		const FSbTree* Tree = SbTree[RootHeap];
 		if (Tree)
 		{
 			Tree->Validate();
@@ -1869,6 +2018,13 @@ void FAllocationsProvider::EditOnAnalysisCompleted(double Time)
 
 	//TODO: shrink live allocs buffers
 
+	if (AllocWarnings > 0 || FreeWarnings > 0 || HeapWarnings > 0 || MiscWarnings > 0)
+	{
+		UE_LOG(LogTraceServices, Warning, TEXT("[MemAlloc] %u warnings (%u ALLOC + %u FREE + %u HEAP + %u other)"),
+			AllocWarnings + FreeWarnings + HeapWarnings + MiscWarnings,
+			AllocWarnings, FreeWarnings, HeapWarnings, MiscWarnings);
+	}
+
 	if (AllocErrors > 0)
 	{
 		UE_LOG(LogTraceServices, Error, TEXT("[MemAlloc] ALLOC event errors: %u"), AllocErrors);
@@ -2177,8 +2333,9 @@ void FAllocationsProvider::EnumerateTags(TFunctionRef<void(const TCHAR*, const T
 
 void FAllocationsProvider::DebugPrint() const
 {
-	for (const FSbTree* Tree : SbTree)
+	for (uint32 RootHeap = 0; RootHeap < MaxRootHeaps; ++RootHeap)
 	{
+		const FSbTree* Tree = SbTree[RootHeap];
 		if (Tree)
 		{
 			Tree->DebugPrint();
diff --git a/Engine/Source/Developer/TraceServices/Private/Model/AllocationsProvider.h b/Engine/Source/Developer/TraceServices/Private/Model/AllocationsProvider.h
--- a/Engine/Source/Developer/TraceServices/Private/Model/AllocationsProvider.h
+++ b/Engine/Source/Developer/TraceServices/Private/Model/AllocationsProvider.h
@@ -333,8 +333,8 @@ public:
 	void EditFree(double Time, uint32 CallstackId, uint64 Address, HeapId RootHeap);
 
 	void EditHeapSpec(HeapId Id, HeapId ParentId, const FStringView& Name, EMemoryTraceHeapFlags Flags);
-	void EditMarkAllocationAsHeap(double Time, uint64 Address, HeapId Heap, EMemoryTraceHeapAllocationFlags Flags);
-	void EditUnmarkAllocationAsHeap(double Time, uint64 Address, HeapId Heap);
+	void EditMarkAllocationAsHeap(double Time, uint32 CallstackId, uint64 Address, HeapId Heap, EMemoryTraceHeapAllocationFlags Flags);
+	void EditUnmarkAllocationAsHeap(double Time, uint32 CallstackId, uint64 Address, HeapId Heap);
 
 	void EditAddTagSpec(TagIdType Tag, TagIdType ParentTag, const TCHAR* Display) { EditAccessCheck(); TagTracker.AddTagSpec(Tag, ParentTag, Display); }
 	void EditPushTag(uint32 ThreadId, uint8 Tracker, TagIdType Tag);
@@ -390,9 +390,13 @@ private:
 	uint64 FreeCount = 0;
 	uint64 HeapCount = 0;
 
+	uint64 MiscWarnings = 0;
 	uint64 MiscErrors = 0;
+	uint64 HeapWarnings = 0;
 	uint64 HeapErrors = 0;
+	uint64 AllocWarnings = 0;
 	uint64 AllocErrors = 0;
+	uint64 FreeWarnings = 0;
 	uint64 FreeErrors = 0;
 
 	uint64 MaxAllocSize = 0;
