[Insights] Improved debug API in AllocationsProvider.

diff --git a/Engine/Source/Developer/TraceServices/Private/Model/AllocationsProvider.cpp b/Engine/Source/Developer/TraceServices/Private/Model/AllocationsProvider.cpp
--- a/Engine/Source/Developer/TraceServices/Private/Model/AllocationsProvider.cpp
+++ b/Engine/Source/Developer/TraceServices/Private/Model/AllocationsProvider.cpp
@@ -21,12 +21,6 @@
 
 #define INSIGHTS_SLOW_CHECK(expr) //check(expr)
 
-#define INSIGHTS_DEBUG_WATCH 0
-// Action to be executed when a match is found.
-#define INSIGHTS_DEBUG_WATCH_FOUND break // just log
-//#define INSIGHTS_DEBUG_WATCH_FOUND return // log and ignore events
-//#define INSIGHTS_DEBUG_WATCH_FOUND UE_DEBUG_BREAK(); break
-
 // Initial reserved number for long living allocations.
 #define INSIGHTS_LLA_RESERVE 0
 //#define INSIGHTS_LLA_RESERVE (64 * 1024)
@@ -44,16 +38,72 @@
 #define INSIGHTS_USE_LAST_ALLOC 1
 
 #define INSIGHTS_VALIDATE_ALLOC_EVENTS 0
+#define INSIGHTS_REMAP_INVALID_ALLOC_EVENTS 0
+#define INSIGHTS_REMAP_INVALID_FREE_EVENTS 0
 
 #define INSIGHTS_DEBUG_METADATA 0
 
-#if 0
-#define INSIGHTS_API_LOGF(Format, ...) { UE_LOG(LogTraceServices, Log, TEXT("[MemAlloc][API] ") Format, __VA_ARGS__); }
-//#define INSIGHTS_API_LOGF(Format, ...) { FPlatformMisc::LowLevelOutputDebugStringf(TEXT("[MemAlloc][API] ") Format TEXT("\n"), __VA_ARGS__); }
+////////////////////////////////////////////////////////////////////////////////////////////////////
+
+#define INSIGHTS_DEBUG_WATCH 0
+
+#if INSIGHTS_DEBUG_WATCH
+namespace TraceServices
+{
+	static uint64 GWatchAddresses[] =
+	{
+		// add here addresses to watch
+		//0x0ull,
+	};
+}
+#endif //INSIGHTS_DEBUG_WATCH
+
+// Action to be executed when a match is found.
+#define INSIGHTS_DEBUG_WATCH_FOUND // just log the API call
+//#define INSIGHTS_DEBUG_WATCH_FOUND return // ignore events
+//#define INSIGHTS_DEBUG_WATCH_FOUND UE_DEBUG_BREAK(); break
+
+////////////////////////////////////////////////////////////////////////////////////////////////////
+
+#define INSIGHTS_LOGF(Format, ...) { UE_LOG(LogTraceServices, Log, TEXT("[MemAlloc]") Format, __VA_ARGS__); }
+//#define INSIGHTS_LOGF(Format, ...) { FPlatformMisc::LowLevelOutputDebugStringf(TEXT("[MemAlloc]") Format TEXT("\n"), __VA_ARGS__); }
+
+#if 0 || INSIGHTS_DEBUG_WATCH
+#define INSIGHTS_API_LOGF(Format, Time, ...) \
+	INSIGHTS_LOGF(TEXT("[API][%3u][%f] ") Format, CurrentTraceThreadId, Time, __VA_ARGS__);
+#else
+#define INSIGHTS_API_LOGF(Format, Time, ...)
+#endif
+
+#if 0 || INSIGHTS_DEBUG_WATCH
+#define INSIGHTS_INDIRECT_API_LOGF(ApiName, Address, Time) \
+	INSIGHTS_LOGF(TEXT("[API] --> ") ApiName TEXT(" 0x%llX"), Address);
 #else
-#define INSIGHTS_API_LOGF(Format, ...)
+#define INSIGHTS_INDIRECT_API_LOGF(ApiName, Address, Time)
 #endif
 
+#if INSIGHTS_DEBUG_WATCH
+	#define INSIGHTS_WATCH_API_LOGF(Address, Format, Time, ...) \
+	{\
+		if (IsAddressWatched(Address))\
+		{\
+			INSIGHTS_DEBUG_WATCH_FOUND;\
+			INSIGHTS_API_LOGF(Format, Time, __VA_ARGS__); \
+		}\
+	}
+	#define INSIGHTS_WATCH_INDIRECT_API_LOGF(ApiName, Address, Time) \
+	{\
+		if (IsAddressWatched(Address))\
+		{\
+			INSIGHTS_DEBUG_WATCH_FOUND;\
+			INSIGHTS_INDIRECT_API_LOGF(ApiName, Address, Time); \
+		}\
+	}
+#else
+	#define INSIGHTS_WATCH_API_LOGF(Address, Format, Time, ...) INSIGHTS_API_LOGF(Format, Time, __VA_ARGS__)
+	#define INSIGHTS_WATCH_INDIRECT_API_LOGF(ApiName, Address, Time) INSIGHTS_INDIRECT_API_LOGF(ApiName, Address, Time)
+#endif // INSIGHTS_DEBUG_WATCH
+
 ////////////////////////////////////////////////////////////////////////////////////////////////////
 
 namespace TraceServices
@@ -63,11 +113,17 @@ constexpr uint32 MaxLogMessagesPerWarningType = 100;
 constexpr uint32 MaxLogMessagesPerErrorType = 100;
 
 #if INSIGHTS_DEBUG_WATCH
-static uint64 GWatchAddresses[] =
+static bool IsAddressWatched(uint64 InAddress)
 {
-	// add here addresses to watch
-	0x0ull,
-};
+	for (int32 AddrIndex = 0; AddrIndex < UE_ARRAY_COUNT(GWatchAddresses); ++AddrIndex)
+	{
+		if (GWatchAddresses[AddrIndex] == InAddress)
+		{
+			return true;
+		}
+	}
+	return false;
+}
 #endif // INSIGHTS_DEBUG_WATCH
 
 ////////////////////////////////////////////////////////////////////////////////////////////////////
@@ -1249,7 +1305,7 @@ void FAllocationsProvider::EditInit(double InTime, uint8 InMinAlignment)
 		return;
 	}
 
-	INSIGHTS_API_LOGF(TEXT("Init(Time=%f, MinAlignment=%u)"), InTime, uint32(InMinAlignment));
+	INSIGHTS_API_LOGF(TEXT("Init : MinAlignment=%u)"), InTime, uint32(InMinAlignment));
 
 	InitTime = InTime;
 	MinAlignment = InMinAlignment;
@@ -1273,7 +1329,7 @@ void FAllocationsProvider::EditAlloc(double Time, uint32 CallstackId, uint64 Add
 		return;
 	}
 
-	INSIGHTS_API_LOGF(TEXT("Alloc(Time=%f, Address=0x%llX, Size=%llu, Alignment=%u, RootHeap=%u, CallstackId=%u)"), Time, Address, InSize, InAlignment, RootHeap, CallstackId);
+	INSIGHTS_WATCH_API_LOGF(Address, TEXT("Alloc 0x%llX : Size=%llu Alignment=%u RootHeap=%u CallstackId=%u"), Time, Address, InSize, InAlignment, RootHeap, CallstackId);
 
 	if (Address == 0)
 	{
@@ -1302,22 +1358,32 @@ void FAllocationsProvider::EditAlloc(double Time, uint32 CallstackId, uint64 Add
 
 	const TagIdType Tag = TagTracker.GetCurrentTag(CurrentSystemThreadId, CurrentTracker);
 
-#if INSIGHTS_DEBUG_WATCH
-	for (int32 AddrIndex = 0; AddrIndex < UE_ARRAY_COUNT(GWatchAddresses); ++AddrIndex)
+	FAllocationItem* AllocationPtr = nullptr;
+
+#if INSIGHTS_VALIDATE_ALLOC_EVENTS
+	AllocationPtr = LiveAllocs[RootHeap]->FindRef(Address);
+#if INSIGHTS_REMAP_INVALID_ALLOC_EVENTS
+	if (AllocationPtr)
 	{
-		if (GWatchAddresses[AddrIndex] == Address)
+		const uint64 OriginalAddress = Address;
+		check(OriginalAddress & (0x3FFull << 54) == 0);
+		for (uint64 Offset = 1; Offset <= 0x3FFull; ++Offset)
+		{
+			Address = OriginalAddress | (Offset << 54);
+			AllocationPtr = LiveAllocs[RootHeap]->FindRef(Address);
+			if (!AllocationPtr)
+			{
+				break;
+			}
+		}
+		++AllocWarnings;
+		if (AllocWarnings <= MaxLogMessagesPerWarningType)
 		{
-			UE_LOG(LogTraceServices, Warning, TEXT("[MemAlloc][%u] Alloc 0x%llX : Size=%llu, Tag=%u, RootHeap=%u, Time=%f, CallstackId=%u"), CurrentTraceThreadId, Address, InSize, Tag, RootHeap, Time, CallstackId);
-			INSIGHTS_DEBUG_WATCH_FOUND;
+			UE_LOG(LogTraceServices, Warning, TEXT("[MemAlloc] Invalid ALLOC event remapped (Address=0x%llX --> 0x%llX, Size=%llu, Tag=%u, RootHeap=%u, Time=%f, CallstackId=%u)!"), OriginalAddress, Address, InSize, Tag, RootHeap, Time, CallstackId);
 		}
 	}
-#endif // INSIGHTS_DEBUG_WATCH
-
-#if INSIGHTS_VALIDATE_ALLOC_EVENTS
-	FAllocationItem* AllocationPtr = LiveAllocs[RootHeap]->FindRef(Address);
-#else
-	FAllocationItem* AllocationPtr = nullptr;
-#endif
+#endif // INSIGHTS_REMAP_INVALID_ALLOC_EVENTS
+#endif // INSIGHTS_VALIDATE_ALLOC_EVENTS
 
 	if (!AllocationPtr)
 	{
@@ -1405,7 +1471,7 @@ void FAllocationsProvider::EditFree(double Time, uint32 CallstackId, uint64 Addr
 		return;
 	}
 
-	INSIGHTS_API_LOGF(TEXT("Free(Time=%f, Address=0x%llX, RootHeap=%u, CallstackId=%u)"), Time, Address, RootHeap, CallstackId);
+	INSIGHTS_WATCH_API_LOGF(Address, TEXT("Free 0x%llX : RootHeap=%u CallstackId=%u"), Time, Address, RootHeap, CallstackId);
 
 	if (Address == 0)
 	{
@@ -1432,17 +1498,6 @@ void FAllocationsProvider::EditFree(double Time, uint32 CallstackId, uint64 Addr
 
 	AdvanceTimelines(Time);
 
-#if INSIGHTS_DEBUG_WATCH
-	for (int32 AddrIndex = 0; AddrIndex < UE_ARRAY_COUNT(GWatchAddresses); ++AddrIndex)
-	{
-		if (GWatchAddresses[AddrIndex] == Address)
-		{
-			UE_LOG(LogTraceServices, Warning, TEXT("[MemAlloc][%u] Free 0x%llX : RootHeap=%u, Time=%f, CallstackId=%u"), CurrentTraceThreadId, Address, RootHeap, Time, CallstackId);
-			INSIGHTS_DEBUG_WATCH_FOUND;
-		}
-	}
-#endif // INSIGHTS_DEBUG_WATCH
-
 	FAllocationItem* AllocationPtr = LiveAllocs[RootHeap]->Remove(Address); // we take ownership of AllocationPtr
 	if (!AllocationPtr)
 	{
@@ -1452,6 +1507,7 @@ void FAllocationsProvider::EditFree(double Time, uint32 CallstackId, uint64 Addr
 		{
 			INSIGHTS_SLOW_CHECK(AllocationPtr->IsHeap());
 			HeapId Heap = AllocationPtr->RootHeap;
+			INSIGHTS_WATCH_INDIRECT_API_LOGF(TEXT("UnmarkAllocationAsHeap"), Address, Time);
 			EditUnmarkAllocationAsHeap(Time, CallstackId, Address, Heap);
 			SbTree[RootHeap]->SetTimeForEvent(EventIndex[RootHeap], Time); // for the case where the free event is first event in a new SbTree column after changing the heap alloc
 			AllocationPtr = LiveAllocs[RootHeap]->Remove(Address); // we take ownership of AllocationPtr
@@ -1461,6 +1517,29 @@ void FAllocationsProvider::EditFree(double Time, uint32 CallstackId, uint64 Addr
 	{
 		INSIGHTS_SLOW_CHECK(!AllocationPtr->IsHeap());
 	}
+
+#if INSIGHTS_REMAP_INVALID_FREE_EVENTS
+	if (!AllocationPtr)
+	{
+		const uint64 OriginalAddress = Address;
+		check(OriginalAddress & (0x3FFull << 54) == 0);
+		for (uint64 Offset = 1; Offset <= 0x3FFull; ++Offset)
+		{
+			Address = OriginalAddress | (Offset << 54);
+			AllocationPtr = LiveAllocs[RootHeap]->FindRef(Address);
+			if (AllocationPtr)
+			{
+				break;
+			}
+		}
+		++AllocWarnings;
+		if (AllocWarnings <= MaxLogMessagesPerWarningType)
+		{
+			UE_LOG(LogTraceServices, Warning, TEXT("[MemAlloc] Invalid FREE event remapped (Address=0x%llX --> 0x%llX, RootHeap=%u, Time=%f, CallstackId=%u)!"), OriginalAddress, Address, RootHeap, Time, CallstackId);
+		}
+	}
+#endif // INSIGHTS_REMAP_INVALID_FREE_EVENTS
+
 	if (!AllocationPtr)
 	{
 		++FreeErrors;
@@ -1471,6 +1550,7 @@ void FAllocationsProvider::EditFree(double Time, uint32 CallstackId, uint64 Addr
 		// Fake the missing alloc.
 		constexpr uint64 FakeAllocSize = 0;
 		constexpr uint32 FakeAllocAlignment = 0;
+		INSIGHTS_WATCH_INDIRECT_API_LOGF(TEXT("Alloc"), Address, Time);
 		EditAlloc(Time, CallstackId, Address, FakeAllocSize, FakeAllocAlignment, RootHeap);
 		SbTree[RootHeap]->SetTimeForEvent(EventIndex[RootHeap], Time); // for the case where the free event is first event in a new SbTree column after adding the fake alloc
 		AllocationPtr = LiveAllocs[RootHeap]->Remove(Address); // we take ownership of AllocationPtr
@@ -1527,7 +1607,7 @@ void FAllocationsProvider::EditHeapSpec(HeapId Id, HeapId ParentId, const FStrin
 {
 	Lock.WriteAccessCheck();
 
-	INSIGHTS_API_LOGF(TEXT("HeapSpec(Id=%u, ParentId=%u, Name=\"%*s\", Flags=0x%X)"), Id, ParentId, Name.Len(), Name.GetData(), uint32(Flags));
+	INSIGHTS_API_LOGF(TEXT("HeapSpec : Id=%u ParentId=%u Name=\"%*s\" Flags=0x%X"), -1.0, Id, ParentId, Name.Len(), Name.GetData(), uint32(Flags));
 
 	AddHeapSpec(Id, ParentId, Name, Flags);
 }
@@ -1585,18 +1665,7 @@ void FAllocationsProvider::EditMarkAllocationAsHeap(double Time, uint32 Callstac
 {
 	Lock.WriteAccessCheck();
 
-	INSIGHTS_API_LOGF(TEXT("MarkAllocAsHeap(Time=%f, Address=0x%llX, Heap=%u, Flags=0x%X, CallstackId=%u)"), Time, Address, Heap, uint32(Flags), CallstackId);
-
-#if INSIGHTS_DEBUG_WATCH
-	for (int32 AddrIndex = 0; AddrIndex < UE_ARRAY_COUNT(GWatchAddresses); ++AddrIndex)
-	{
-		if (GWatchAddresses[AddrIndex] == Address)
-		{
-			UE_LOG(LogTraceServices, Warning, TEXT("[MemAlloc][%u] HeapMarkAlloc 0x%llX : Heap=%u, Flags=%u, Time=%f, CallstackId=%u"), CurrentTraceThreadId, Address, Heap, uint32(Flags), Time, CallstackId);
-			INSIGHTS_DEBUG_WATCH_FOUND;
-		}
-	}
-#endif // INSIGHTS_DEBUG_WATCH
+	INSIGHTS_WATCH_API_LOGF(Address, TEXT("MarkAllocAsHeap 0x%llX : Heap=%u Flags=0x%X CallstackId=%u"), Time, Address, Heap, uint32(Flags), CallstackId);
 
 	if (Heap >= (uint32)HeapSpecs.Num())
 	{
@@ -1669,18 +1738,7 @@ void FAllocationsProvider::EditUnmarkAllocationAsHeap(double Time, uint32 Callst
 {
 	Lock.WriteAccessCheck();
 
-	INSIGHTS_API_LOGF(TEXT("UnmarkAllocAsHeap(Time=%f, Address=0x%llX, Heap=%u, CallstackId=%u)"), Time, Address, Heap, CallstackId);
-
-#if INSIGHTS_DEBUG_WATCH
-	for (int32 AddrIndex = 0; AddrIndex < UE_ARRAY_COUNT(GWatchAddresses); ++AddrIndex)
-	{
-		if (GWatchAddresses[AddrIndex] == Address)
-		{
-			UE_LOG(LogTraceServices, Warning, TEXT("[MemAlloc][%u] HeapUnmarkAlloc 0x%llX : Heap=%u, Time=%f, CallstackId=%u"), CurrentTraceThreadId, Address, Heap, Time, CallstackId);
-			INSIGHTS_DEBUG_WATCH_FOUND;
-		}
-	}
-#endif // INSIGHTS_DEBUG_WATCH
+	INSIGHTS_WATCH_API_LOGF(Address, TEXT("UnmarkAllocAsHeap 0x%llX : Heap=%u CallstackId=%u"), Time, Address, Heap, CallstackId);
 
 	if (Heap >= (uint32)HeapSpecs.Num())
 	{
@@ -1748,18 +1806,22 @@ void FAllocationsProvider::EditUnmarkAllocationAsHeap(double Time, uint32 Callst
 			{
 				if (!EnumHasAnyFlags(HeapSpec.Flags, EMemoryTraceHeapFlags::NeverFrees))
 				{
-					// For heaps that do not have NeverFrees flag, we report live allocs as leaks.
-					UE_LOG(LogTraceServices, Warning, TEXT("[MemAlloc] HeapUnmarkAlloc: %d memory leaks (%llu bytes) detected for heap %u (\"%s\", Address=0x%llX, Size=%llu, Time=%f, CallstackId=%u)"),
-						AllocsInHeap.Num(), AllocatedSizeInHeap, Heap, HeapSpec.Name, Address, HeapSize, Time, CallstackId);
-					UE_LOG(LogTraceServices, Warning, TEXT("[MemAlloc] Top memory leaks:"));
-					AllocsInHeap.Sort([this](const FAllocInHeap& A, const FAllocInHeap& B) -> bool { return A.Size > B.Size; });
-					int NumAllocsInTop = 10; // only show first 10 allocs
-					for (const FAllocInHeap& AllocInHeap : AllocsInHeap)
+					++HeapWarnings;
+					if (HeapWarnings <= MaxLogMessagesPerWarningType)
 					{
-						UE_LOG(LogTraceServices, Warning, TEXT("[MemAlloc]     alloc 0x%llX (%u bytes)"),  AllocInHeap.Address, AllocInHeap.Size);
-						if (--NumAllocsInTop <= 0)
+						// For heaps that do not have NeverFrees flag, we report live allocs as leaks.
+						UE_LOG(LogTraceServices, Warning, TEXT("[MemAlloc] HeapUnmarkAlloc: %d memory leaks (%llu bytes) detected for heap %u (\"%s\", Address=0x%llX, Size=%llu, Time=%f, CallstackId=%u)"),
+							AllocsInHeap.Num(), AllocatedSizeInHeap, Heap, HeapSpec.Name, Address, HeapSize, Time, CallstackId);
+						UE_LOG(LogTraceServices, Warning, TEXT("[MemAlloc] Top memory leaks:"));
+						AllocsInHeap.Sort([this](const FAllocInHeap& A, const FAllocInHeap& B) -> bool { return A.Size > B.Size; });
+						int NumAllocsInTop = 10; // only show first 10 allocs
+						for (const FAllocInHeap& AllocInHeap : AllocsInHeap)
 						{
-							break;
+							UE_LOG(LogTraceServices, Warning, TEXT("[MemAlloc]     alloc 0x%llX (%u bytes)"), AllocInHeap.Address, AllocInHeap.Size);
+							if (--NumAllocsInTop <= 0)
+							{
+								break;
+							}
 						}
 					}
 				}
@@ -1774,6 +1836,7 @@ void FAllocationsProvider::EditUnmarkAllocationAsHeap(double Time, uint32 Callst
 				// Free automatically all allocs in this heap.
 				for (const FAllocInHeap& AllocInHeap : AllocsInHeap)
 				{
+					INSIGHTS_WATCH_INDIRECT_API_LOGF(TEXT("Free"), AllocInHeap.Address, Time);
 					EditFree(Time, CallstackId, AllocInHeap.Address, RootHeap);
 				}
 			}
@@ -1790,7 +1853,9 @@ void FAllocationsProvider::EditUnmarkAllocationAsHeap(double Time, uint32 Callst
 		// event and an "alloc" event. Make sure the new allocation retains the tag from the original.
 		CurrentTracker = 1;
 		EditPushTagFromPtr(CurrentSystemThreadId, CurrentTracker, Address);
+		INSIGHTS_WATCH_INDIRECT_API_LOGF(TEXT("Free"), Address, Time);
 		EditFree(Time, CallstackId, Address, RootHeap);
+		INSIGHTS_WATCH_INDIRECT_API_LOGF(TEXT("Alloc"), Address, Time);
 		EditAlloc(Time, AllocCallstackId, Address, Size, Alignment, RootHeap);
 		EditPopTagFromPtr(CurrentSystemThreadId, CurrentTracker);
 		CurrentTracker = 0;
@@ -2409,12 +2474,18 @@ const IAllocationsProvider* ReadAllocationsProvider(const IAnalysisSession& Sess
 } // namespace TraceServices
 
 #undef INSIGHTS_SLOW_CHECK
-#undef INSIGHTS_DEBUG_WATCH
-#undef INSIGHTS_DEBUG_WATCH_FOUND
 #undef INSIGHTS_LLA_RESERVE
 #undef INSIGHTS_USE_SHORT_LIVING_ALLOCS
 #undef INSIGHTS_SLA_USE_ADDRESS_MAP
 #undef INSIGHTS_USE_LAST_ALLOC
 #undef INSIGHTS_VALIDATE_ALLOC_EVENTS
+#undef INSIGHTS_REMAP_INVALID_ALLOC_EVENTS
+#undef INSIGHTS_REMAP_INVALID_FREE_EVENTS
 #undef INSIGHTS_DEBUG_METADATA
+#undef INSIGHTS_DEBUG_WATCH
+#undef INSIGHTS_DEBUG_WATCH_FOUND
+#undef INSIGHTS_LOGF
 #undef INSIGHTS_API_LOGF
+#undef INSIGHTS_INDIRECT_API_LOGF
+#undef INSIGHTS_WATCH_API_LOGF
+#undef INSIGHTS_WATCH_INDIRECT_API_LOGF
