Removed two unused routes.

#rnx

diff --git a/Engine/Source/Developer/TraceServices/Private/Analyzers/MiscTraceAnalysis.cpp b/Engine/Source/Developer/TraceServices/Private/Analyzers/MiscTraceAnalysis.cpp
--- a/Engine/Source/Developer/TraceServices/Private/Analyzers/MiscTraceAnalysis.cpp
+++ b/Engine/Source/Developer/TraceServices/Private/Analyzers/MiscTraceAnalysis.cpp
@@ -45,8 +45,6 @@ void FMiscTraceAnalyzer::OnAnalysisBegin(const FOnAnalysisContext& Context)
 	Builder.RouteEvent(RouteId_SetThreadGroup, "Misc", "SetThreadGroup");
 	Builder.RouteEvent(RouteId_BeginThreadGroupScope, "Misc", "BeginThreadGroupScope");
 	Builder.RouteEvent(RouteId_EndThreadGroupScope, "Misc", "EndThreadGroupScope");
-	Builder.RouteEvent(RouteId_BookmarkSpec, "Misc", "BookmarkSpec");
-	Builder.RouteEvent(RouteId_Bookmark, "Misc", "Bookmark");
 	Builder.RouteEvent(RouteId_BeginFrame, "Misc", "BeginFrame");
 	Builder.RouteEvent(RouteId_EndFrame, "Misc", "EndFrame");
 	Builder.RouteEvent(RouteId_BeginGameFrame, "Misc", "BeginGameFrame");
diff --git a/Engine/Source/Developer/TraceServices/Private/Analyzers/MiscTraceAnalysis.h b/Engine/Source/Developer/TraceServices/Private/Analyzers/MiscTraceAnalysis.h
--- a/Engine/Source/Developer/TraceServices/Private/Analyzers/MiscTraceAnalysis.h
+++ b/Engine/Source/Developer/TraceServices/Private/Analyzers/MiscTraceAnalysis.h
@@ -39,8 +39,6 @@ private:
 		RouteId_SetThreadGroup,
 		RouteId_BeginThreadGroupScope,
 		RouteId_EndThreadGroupScope,
-		RouteId_BookmarkSpec,
-		RouteId_Bookmark,
 		RouteId_BeginFrame,
 		RouteId_EndFrame,
 		RouteId_BeginGameFrame,
