[Insights] Metadata: Added support for hierarchical Clear and Restore scopes. The Clear/Restore scopes replaces the current metadata stack with an empty/saved metadata stack and restores the previous metadata stack at the end of scope.

diff --git a/Engine/Source/Developer/TraceServices/Private/Model/MetadataProvider.cpp b/Engine/Source/Developer/TraceServices/Private/Model/MetadataProvider.cpp
--- a/Engine/Source/Developer/TraceServices/Private/Model/MetadataProvider.cpp
+++ b/Engine/Source/Developer/TraceServices/Private/Model/MetadataProvider.cpp
@@ -163,10 +163,6 @@ FMetadataProvider::FMetadataThread::FMetadataThread(uint32 InThreadId, ILinearAl
 	: ThreadId(InThreadId)
 	, CurrentStack()
 	, Metadata(InAllocator, 1024)
-	, bIsClearStackScope(false)
-	, bIsRestoreSavedStackScope(false)
-	, RestoreSavedStackId(0)
-	, RestoreSavedStackSize(0)
 {
 }
 
@@ -275,6 +271,36 @@ void FMetadataProvider::InternalClearStack(FMetadataThread& InMetadataThread)
 
 ////////////////////////////////////////////////////////////////////////////////////////////////////
 
+void FMetadataProvider::InternalPushSavedStack(FMetadataThread& InMetadataThread, FMetadataThread& InSavedMetadataThread, uint32 InSavedMetadataId)
+{
+	METADATA_PROVIDER_DEBUG_LOG(InMetadataThread.ThreadId, TEXT("  --> push metadata %u (from thread %u)"), InSavedMetadataId, InSavedMetadataThread.ThreadId);
+
+	auto Iterator = InSavedMetadataThread.Metadata.GetIteratorFromItem(InSavedMetadataId);
+	const FMetadataEntry* Entry = Iterator.GetCurrentItem();
+	while (true)
+	{
+		check(Entry->StoreIndex < MetadataStore.Num());
+		const FMetadataStoreEntry& StoreEntry = *MetadataStore.GetIteratorFromItem(Entry->StoreIndex);
+		const void* Data = (StoreEntry.Size > MaxInlinedMetadataSize) ? StoreEntry.Ptr : StoreEntry.Value;
+
+		if (InMetadataThread.CurrentStack.Num() == MaxMetadataStackSize)
+		{
+			METADATA_PROVIDER_ERROR(PushSavedStackErrors, TEXT("[Meta] Cannot push saved metadata %u (from thread %u) to thread %u. Stack size is too large."), InSavedMetadataId, InSavedMetadataThread.ThreadId, InMetadataThread.ThreadId);
+			break;
+		}
+
+		InternalPushStackEntry(InMetadataThread, StoreEntry.Type, Data, StoreEntry.Size);
+
+		if (Entry->StackSize == 1) // last one
+		{
+			break;
+		}
+		Entry = Iterator.PrevItem();
+	}
+}
+
+////////////////////////////////////////////////////////////////////////////////////////////////////
+
 void FMetadataProvider::PushScopedMetadata(uint32 InThreadId, uint16 InType, const void* InData, uint32 InSize)
 {
 	Lock.WriteAccessCheck();
@@ -339,23 +365,15 @@ void FMetadataProvider::BeginClearStackScope(uint32 InThreadId)
 
 	METADATA_PROVIDER_DEBUG_LOG(InThreadId, TEXT("BeginClearStackScope()"));
 
-	const uint32 ClearSavedStackId = 0;
-
 	FMetadataThread& MetadataThread = GetOrAddThread(InThreadId);
-	if (MetadataThread.bIsClearStackScope)
-	{
-		METADATA_PROVIDER_ERROR(ClearScopeErrors, TEXT("[Meta] Cannot begin clear stack (thread %u). Already in a clear scope."), InThreadId);
-		return;
-	}
-	MetadataThread.bIsClearStackScope = true;
 
 	// Save the current stack.
-	const uint32 MetadataId = PinAndGetId(InThreadId);
-	SavedStackMap.Add(ClearSavedStackId, { InThreadId, MetadataId });
+	const uint32 SavedMetadataId = PinAndGetId(InThreadId);
+	MetadataThread.SavedMetadataStack.Push(SavedMetadataId);
 
-	if (MetadataId != InvalidMetadataId)
+	if (SavedMetadataId != InvalidMetadataId)
 	{
-		METADATA_PROVIDER_DEBUG_LOG(InThreadId, TEXT("  --> id = %u"), MetadataId);
+		METADATA_PROVIDER_DEBUG_LOG(InThreadId, TEXT("  --> id = %u"), SavedMetadataId);
 	}
 
 	// Clear the current stack.
@@ -371,24 +389,21 @@ void FMetadataProvider::EndClearStackScope(uint32 InThreadId)
 
 	METADATA_PROVIDER_DEBUG_LOG(InThreadId, TEXT("EndClearStackScope()"));
 
-	const uint32 ClearSavedStackId = 0;
-
 	FMetadataThread& MetadataThread = GetOrAddThread(InThreadId);
-	if (!MetadataThread.bIsClearStackScope)
+	if (MetadataThread.SavedMetadataStack.IsEmpty())
 	{
-		METADATA_PROVIDER_ERROR(ClearScopeErrors, TEXT("[Meta] Cannot end clear stack (thread %u). Not in a clear scope."), InThreadId);
+		METADATA_PROVIDER_ERROR(ClearScopeErrors, TEXT("[Meta] Cannot end clear stack (thread %u). Clear scope stack is already empty."), InThreadId);
 		return;
 	}
-	MetadataThread.bIsClearStackScope = false;
 
-	FMetadataSavedStackInfo SavedStackInfo;
-	if (!SavedStackMap.RemoveAndCopyValue(ClearSavedStackId, SavedStackInfo))
+	if (MetadataThread.CurrentStack.Num() > 0)
 	{
-		METADATA_PROVIDER_ERROR(ClearScopeErrors, TEXT("[Meta] Cannot end clear stack (thread %u). Invalid saved stack."), InThreadId);
-		return;
+		UE_LOG(LogTraceServices, Warning, TEXT("[Meta] Invalid end clear stack event (thread %u). Current metadata stack is not empty."), InThreadId);
 	}
 
-	if (SavedStackInfo.MetadataId == InvalidMetadataId)
+	const uint32 SavedMetadataId = MetadataThread.SavedMetadataStack.Pop();
+
+	if (SavedMetadataId == InvalidMetadataId)
 	{
 		METADATA_PROVIDER_DEBUG_LOG(InThreadId, TEXT("  --> empty"));
 
@@ -396,39 +411,16 @@ void FMetadataProvider::EndClearStackScope(uint32 InThreadId)
 		return;
 	}
 
-	METADATA_PROVIDER_DEBUG_LOG(InThreadId, TEXT("  --> id = %u"), SavedStackInfo.MetadataId);
+	METADATA_PROVIDER_DEBUG_LOG(InThreadId, TEXT("  --> id = %u"), SavedMetadataId);
 
-	FMetadataThread* SavedMetadataThread = GetThread(SavedStackInfo.ThreadId);
-
-	if (!SavedMetadataThread || SavedStackInfo.MetadataId >= SavedMetadataThread->Metadata.Num())
+	if (SavedMetadataId >= MetadataThread.Metadata.Num())
 	{
 		METADATA_PROVIDER_ERROR(ClearScopeErrors, TEXT("[Meta] Cannot end clear stack (thread %u). Invalid saved metadata."), InThreadId);
 		return;
 	}
 
-	// Push the cleared saved stack.
-	auto Iterator = SavedMetadataThread->Metadata.GetIteratorFromItem(SavedStackInfo.MetadataId);
-	const FMetadataEntry* Entry = Iterator.GetCurrentItem();
-	while (true)
-	{
-		check(Entry->StoreIndex < MetadataStore.Num());
-		const FMetadataStoreEntry& StoreEntry = *MetadataStore.GetIteratorFromItem(Entry->StoreIndex);
-		const void* Data = (StoreEntry.Size > MaxInlinedMetadataSize) ? StoreEntry.Ptr : StoreEntry.Value;
-
-		if (MetadataThread.CurrentStack.Num() == MaxMetadataStackSize)
-		{
-			METADATA_PROVIDER_ERROR(ClearScopeErrors, TEXT("[Meta] Cannot end clear stack (thread %u). Stack size is too large."), InThreadId);
-			break;
-		}
-
-		InternalPushStackEntry(MetadataThread, StoreEntry.Type, Data, StoreEntry.Size);
-
-		if (Entry->StackSize == 1) // last one
-		{
-			break;
-		}
-		Entry = Iterator.PrevItem();
-	}
+	// Push the saved stack.
+	InternalPushSavedStack(MetadataThread, MetadataThread, SavedMetadataId);
 }
 
 ////////////////////////////////////////////////////////////////////////////////////////////////////
@@ -449,17 +441,21 @@ void FMetadataProvider::BeginRestoreSavedStackScope(uint32 InThreadId, uint32 In
 	Lock.WriteAccessCheck();
 	++EventCount;
 
+	METADATA_PROVIDER_DEBUG_LOG(InThreadId, TEXT("BeginRestoreSavedStackScope()"));
+
 	FMetadataThread& MetadataThread = GetOrAddThread(InThreadId);
 
-	if (MetadataThread.bIsRestoreSavedStackScope)
+	// Save the current stack.
+	const uint32 SavedMetadataId = PinAndGetId(InThreadId);
+	MetadataThread.SavedMetadataStack.Push(SavedMetadataId);
+
+	if (SavedMetadataId != InvalidMetadataId)
 	{
-		METADATA_PROVIDER_ERROR(RestoreScopeErrors, TEXT("[Meta] Cannot begin restore saved stack (thread %u, id %u). Already in a restore scope."), InThreadId, InSavedStackId);
-		return;
+		METADATA_PROVIDER_DEBUG_LOG(InThreadId, TEXT("  --> id = %u"), SavedMetadataId);
 	}
-	MetadataThread.bIsRestoreSavedStackScope = true;
 
-	MetadataThread.RestoreSavedStackId = InSavedStackId;
-	MetadataThread.RestoreSavedStackSize = 0;
+	// Clear the current stack.
+	InternalClearStack(MetadataThread);
 
 	const FMetadataSavedStackInfo* SavedStackInfo = SavedStackMap.Find(InSavedStackId);
 	if (!SavedStackInfo)
@@ -482,25 +478,8 @@ void FMetadataProvider::BeginRestoreSavedStackScope(uint32 InThreadId, uint32 In
 		return;
 	}
 
-	// Push the saved stack.
-	auto Iterator = SavedMetadataThread->Metadata.GetIteratorFromItem(SavedStackInfo->MetadataId);
-	const FMetadataEntry* Entry = Iterator.GetCurrentItem();
-	while (true)
-	{
-		check(Entry->StoreIndex < MetadataStore.Num());
-		const FMetadataStoreEntry& StoreEntry = *MetadataStore.GetIteratorFromItem(Entry->StoreIndex);
-		const void* Data = (StoreEntry.Size > MaxInlinedMetadataSize) ? StoreEntry.Ptr : StoreEntry.Value;
-
-		InternalPushStackEntry(MetadataThread, StoreEntry.Type, Data, StoreEntry.Size);
-
-		MetadataThread.RestoreSavedStackSize++;
-
-		if (Entry->StackSize == 1) // last one
-		{
-			break;
-		}
-		Entry = Iterator.PrevItem();
-	}
+	// Push the previously saved stack.
+	InternalPushSavedStack(MetadataThread, *SavedMetadataThread, SavedStackInfo->MetadataId);
 }
 
 ////////////////////////////////////////////////////////////////////////////////////////////////////
@@ -510,29 +489,38 @@ void FMetadataProvider::EndRestoreSavedStackScope(uint32 InThreadId)
 	Lock.WriteAccessCheck();
 	++EventCount;
 
-	FMetadataThread& MetadataThread = GetOrAddThread(InThreadId);
+	METADATA_PROVIDER_DEBUG_LOG(InThreadId, TEXT("EndRestoreSavedStackScope()"));
 
-	if (!MetadataThread.bIsRestoreSavedStackScope)
+	FMetadataThread& MetadataThread = GetOrAddThread(InThreadId);
+	if (MetadataThread.SavedMetadataStack.IsEmpty())
 	{
-		METADATA_PROVIDER_ERROR(RestoreScopeErrors, TEXT("[Meta] Cannot end restore saved stack (thread %u). Not in a restore scope."), InThreadId);
+		METADATA_PROVIDER_ERROR(RestoreScopeErrors, TEXT("[Meta] Cannot end restore saved stack (thread %u). Restore scope stack is already empty."), InThreadId);
 		return;
 	}
-	MetadataThread.bIsRestoreSavedStackScope = false;
 
-	// Pop the pushed saved stack.
-	for (uint32 Index = MetadataThread.RestoreSavedStackSize; Index > 0; --Index)
+	// Clears the restored saved stack.
+	InternalClearStack(MetadataThread);
+
+	const uint32 SavedMetadataId = MetadataThread.SavedMetadataStack.Pop();
+
+	if (SavedMetadataId == InvalidMetadataId)
 	{
-		if (MetadataThread.CurrentStack.Num() == 0)
-		{
-			METADATA_PROVIDER_ERROR(RestoreScopeErrors, TEXT("[Meta] Cannot end restore saved stack (thread %u). Invalid current stack."), InThreadId);
-			break;
-		}
+		METADATA_PROVIDER_DEBUG_LOG(InThreadId, TEXT("  --> empty"));
 
-		InternalPopStackEntry(MetadataThread);
+		// Empty stack. Nothing to push.
+		return;
 	}
 
-	MetadataThread.RestoreSavedStackId = 0;
-	MetadataThread.RestoreSavedStackSize = 0;
+	METADATA_PROVIDER_DEBUG_LOG(InThreadId, TEXT("  --> id = %u"), SavedMetadataId);
+
+	if (SavedMetadataId >= MetadataThread.Metadata.Num())
+	{
+		METADATA_PROVIDER_ERROR(ClearScopeErrors, TEXT("[Meta] Cannot end restore saved stack (thread %u). Invalid saved metadata."), InThreadId);
+		return;
+	}
+
+	// Push the saved stack.
+	InternalPushSavedStack(MetadataThread, MetadataThread, SavedMetadataId);
 }
 
 ////////////////////////////////////////////////////////////////////////////////////////////////////
diff --git a/Engine/Source/Developer/TraceServices/Private/Model/MetadataProvider.h b/Engine/Source/Developer/TraceServices/Private/Model/MetadataProvider.h
--- a/Engine/Source/Developer/TraceServices/Private/Model/MetadataProvider.h
+++ b/Engine/Source/Developer/TraceServices/Private/Model/MetadataProvider.h
@@ -79,10 +79,7 @@ private:
 		uint32 ThreadId;
 		TArray<FMetadataStackEntry> CurrentStack;
 		TPagedArray<FMetadataEntry> Metadata; // a metadata id is an index in this array
-		bool bIsClearStackScope;
-		bool bIsRestoreSavedStackScope;
-		uint32 RestoreSavedStackId;
-		uint32 RestoreSavedStackSize;
+		TArray<uint32> SavedMetadataStack; // stack of MetadataId values for each saved Clear or Restore scope
 	};
 
 	struct FMetadataSavedStackInfo
@@ -143,6 +140,7 @@ private:
 	void InternalPushStackEntry(FMetadataThread& InMetadataThread, uint16 InType, const void* InData, uint32 InSize);
 	void InternalPopStackEntry(FMetadataThread& InMetadataThread);
 	void InternalClearStack(FMetadataThread& InMetadataThread);
+	void InternalPushSavedStack(FMetadataThread& InMetadataThread, FMetadataThread& InSavedMetadataThread, uint32 InSavedMetadataId);
 
 private:
 	IAnalysisSession& Session;
@@ -160,6 +158,7 @@ private:
 	uint32 MetaScopeErrors = 0; // debug
 	uint32 ClearScopeErrors = 0; // debug
 	uint32 RestoreScopeErrors = 0; // debug
+	uint32 PushSavedStackErrors = 0; // debug
 
 	uint64 EventCount = 0; // debug
 	uint64 AllocationEventCount = 0; // debug
