[Insights]
  - Memory Insights: Fixed realloc scopes for the video memory allocations. Previously only realloc scopes for system memory were supported.
  - Also split analysis of RouteId_MemScopeTag and RouteId_MemScopePtr in separate case blocks.

diff --git a/Engine/Source/Developer/TraceServices/Private/Analyzers/AllocationsAnalysis.cpp b/Engine/Source/Developer/TraceServices/Private/Analyzers/AllocationsAnalysis.cpp
--- a/Engine/Source/Developer/TraceServices/Private/Analyzers/AllocationsAnalysis.cpp
+++ b/Engine/Source/Developer/TraceServices/Private/Analyzers/AllocationsAnalysis.cpp
@@ -100,7 +100,7 @@ bool FAllocationsAnalyzer::OnEvent(uint16 RouteId, EStyle Style, const FOnEventC
 	LLM_SCOPE_BYNAME(TEXT("Insights/FAllocationsAnalyzer"));
 
 	const auto& EventData = Context.EventData;
-	HeapId RootHeap = 0;
+	HeapId RootHeap = EMemoryTraceRootHeap::SystemMemory;
 
 	switch (RouteId)
 	{
@@ -213,7 +213,7 @@ bool FAllocationsAnalyzer::OnEvent(uint16 RouteId, EStyle Style, const FOnEventC
 			if (RouteId == RouteId_ReallocFree || RouteId == RouteId_ReallocFreeSystem)
 			{
 				const uint8 Tracker = 0; // We only care about the default tracker for now.
-				AllocationsProvider.EditPushTagFromPtr(SystemThreadId, Tracker, Address);
+				AllocationsProvider.EditPushTagFromPtr(SystemThreadId, Tracker, Address, RootHeap);
 			}
 			AllocationsProvider.EditFree(Time, CallstackId, Address, RootHeap);
 			break;
@@ -298,63 +298,78 @@ bool FAllocationsAnalyzer::OnEvent(uint16 RouteId, EStyle Style, const FOnEventC
 			break;
 		}
 
-		case RouteId_MemScopeTag:
-		case RouteId_MemScopePtr:
+		case RouteId_MemScopeTag: // "MemoryScope", see UE_MEMSCOPE
 		{
 			const uint32 ThreadId = Context.ThreadInfo.GetSystemId();
 			const uint8 Tracker = 0; // We only care about the default tracker for now.
 
 			if (Style == EStyle::EnterScope)
 			{
-				if (RouteId == RouteId_MemScopeTag) // "MemoryScope"
+				const TagIdType Tag = Context.EventData.GetValue<TagIdType>("Tag");
 				{
-					const TagIdType Tag = Context.EventData.GetValue<TagIdType>("Tag");
-					{
-						FProviderEditScopeLock _(AllocationsProvider);
-						AllocationsProvider.EditPushTag(ThreadId, Tracker, Tag);
-					}
-#if INSIGHTS_MEM_TRACE_METADATA_TEST
-					{
-						FMetadataProvider::FEditScopeLock _(MetadataProvider);
-						MetadataProvider.PushScopedMetadata(ThreadId, TagIdMetadataType, (void*)&Tag, sizeof(TagIdType));
-					}
-#endif
+					FProviderEditScopeLock _(AllocationsProvider);
+					AllocationsProvider.EditPushTag(ThreadId, Tracker, Tag);
 				}
-				else if (ensure(RouteId == RouteId_MemScopePtr)) // "MemoryScopePtr"
-				{
-					const uint64 Ptr = Context.EventData.GetValue<uint64>("Ptr");
-					{
-						FProviderEditScopeLock _(AllocationsProvider);
-						AllocationsProvider.EditPushTagFromPtr(ThreadId, Tracker, Ptr);
-					}
 #if INSIGHTS_MEM_TRACE_METADATA_TEST
-					{
-						FMetadataProvider::FEditScopeLock _(MetadataProvider);
-						TagIdType Tag = 0; //TODO: AllocationsProvider.GetTagFromPtr(ThreadId, Tracker, Ptr);
-						MetadataProvider.PushScopedMetadata(ThreadId, TagIdMetadataType, (void*)&Tag, sizeof(TagIdType));
-					}
-#endif
+				{
+					FMetadataProvider::FEditScopeLock _(MetadataProvider);
+					MetadataProvider.PushScopedMetadata(ThreadId, TagIdMetadataType, (void*)&Tag, sizeof(TagIdType));
 				}
+#endif
 			}
 			else if (ensure(Style == EStyle::LeaveScope))
 			{
-				if (RouteId == RouteId_MemScopeTag) // "MemoryScope"
+#if INSIGHTS_MEM_TRACE_METADATA_TEST
+				{
+					FMetadataProvider::FEditScopeLock _(MetadataProvider);
+					MetadataProvider.PopScopedMetadata(ThreadId, TagIdMetadataType);
+				}
+#endif
 				{
 					FProviderEditScopeLock _(AllocationsProvider);
 					AllocationsProvider.EditPopTag(ThreadId, Tracker);
 				}
-				else if (ensure(RouteId == RouteId_MemScopePtr)) // "MemoryScopePtr"
+			}
+			break;
+		}
+
+		case RouteId_MemScopePtr: // "MemoryScopePtr", see UE_MEMSCOPE_PTR
+		{
+			const uint32 ThreadId = Context.ThreadInfo.GetSystemId();
+			const uint8 Tracker = 0; // We only care about the default tracker for now.
+
+			if (Style == EStyle::EnterScope)
+			{
+				const uint64 Ptr = Context.EventData.GetValue<uint64>("Ptr");
+
+				// RootHeap defaults to EMemoryTraceRootHeap::SystemMemory
+				//RootHeap = EventData.GetValue<uint8>("RootHeap", static_cast<uint8>(RootHeap)); // TODO: UE_MEMSCOPE_PTR(InPtr, InRootHeap)
+
 				{
 					FProviderEditScopeLock _(AllocationsProvider);
-					//check(AllocationsProvider.HasTagFromPtrScope(ThreadId, Tracker));
-					AllocationsProvider.EditPopTagFromPtr(ThreadId, Tracker);
+					AllocationsProvider.EditPushTagFromPtr(ThreadId, Tracker, Ptr, RootHeap);
+				}
+#if INSIGHTS_MEM_TRACE_METADATA_TEST
+				{
+					FMetadataProvider::FEditScopeLock _(MetadataProvider);
+					TagIdType Tag = 0; //TODO: AllocationsProvider.GetTagFromPtr(ThreadId, Tracker, Ptr, RootHeapId);
+					MetadataProvider.PushScopedMetadata(ThreadId, TagIdMetadataType, (void*)&Tag, sizeof(TagIdType));
 				}
+#endif
+			}
+			else if (ensure(Style == EStyle::LeaveScope))
+			{
 #if INSIGHTS_MEM_TRACE_METADATA_TEST
 				{
 					FMetadataProvider::FEditScopeLock _(MetadataProvider);
 					MetadataProvider.PopScopedMetadata(ThreadId, TagIdMetadataType);
 				}
 #endif
+				{
+					FProviderEditScopeLock _(AllocationsProvider);
+					//check(AllocationsProvider.HasTagFromPtrScope(ThreadId, Tracker));
+					AllocationsProvider.EditPopTagFromPtr(ThreadId, Tracker);
+				}
 			}
 			break;
 		}
diff --git a/Engine/Source/Developer/TraceServices/Private/Model/AllocationsProvider.cpp b/Engine/Source/Developer/TraceServices/Private/Model/AllocationsProvider.cpp
--- a/Engine/Source/Developer/TraceServices/Private/Model/AllocationsProvider.cpp
+++ b/Engine/Source/Developer/TraceServices/Private/Model/AllocationsProvider.cpp
@@ -1852,12 +1852,12 @@ void FAllocationsProvider::EditUnmarkAllocationAsHeap(double Time, uint32 Callst
 		// We cannot just unmark the allocation as heap, there is no timestamp support, instead fake a "free"
 		// event and an "alloc" event. Make sure the new allocation retains the tag from the original.
 		CurrentTracker = 1;
-		EditPushTagFromPtr(CurrentSystemThreadId, CurrentTracker, Address);
+		TagTracker.PushTagFromPtr(CurrentSystemThreadId, CurrentTracker, Alloc->Tag);
 		INSIGHTS_WATCH_INDIRECT_API_LOGF(TEXT("Free"), Address, Time);
 		EditFree(Time, CallstackId, Address, RootHeap);
 		INSIGHTS_WATCH_INDIRECT_API_LOGF(TEXT("Alloc"), Address, Time);
 		EditAlloc(Time, AllocCallstackId, Address, Size, Alignment, RootHeap);
-		EditPopTagFromPtr(CurrentSystemThreadId, CurrentTracker);
+		TagTracker.PopTagFromPtr(CurrentSystemThreadId, CurrentTracker);
 		CurrentTracker = 0;
 	}
 	else
@@ -1988,12 +1988,11 @@ void FAllocationsProvider::EditPopTag(uint32 ThreadId, uint8 Tracker)
 
 ////////////////////////////////////////////////////////////////////////////////////////////////////
 
-void FAllocationsProvider::EditPushTagFromPtr(uint32 ThreadId, uint8 Tracker, uint64 Ptr)
+void FAllocationsProvider::EditPushTagFromPtr(uint32 ThreadId, uint8 Tracker, uint64 Ptr, HeapId RootHeapId)
 {
 	EditAccessCheck();
 
-	// Currently only system root heap is affected by reallocs, so limit search.
-	FLiveAllocCollection* Allocs = LiveAllocs[EMemoryTraceRootHeap::SystemMemory];
+	FLiveAllocCollection* Allocs = LiveAllocs[RootHeapId];
 	FAllocationItem* Alloc = Allocs ? Allocs->FindRef(Ptr) : nullptr;
 	const TagIdType Tag = Alloc ? Alloc->Tag : 0; // If ptr is not found use "Untagged"
 	TagTracker.PushTagFromPtr(ThreadId, Tracker, Tag);
@@ -2003,7 +2002,7 @@ void FAllocationsProvider::EditPushTagFromPtr(uint32 ThreadId, uint8 Tracker, ui
 		++MiscErrors;
 		if (MiscErrors <= MaxLogMessagesPerErrorType)
 		{
-			UE_LOG(LogTraceServices, Error, TEXT("[MemAlloc] Invalid Ptr for MemoryScopePtr event (Ptr=0x%llX)!"), Ptr);
+			UE_LOG(LogTraceServices, Error, TEXT("[MemAlloc] Invalid Ptr (0x%llX) for MemoryScopePtr or ReallocFree event!"), Ptr);
 		}
 	}
 }
diff --git a/Engine/Source/Developer/TraceServices/Private/Model/AllocationsProvider.h b/Engine/Source/Developer/TraceServices/Private/Model/AllocationsProvider.h
--- a/Engine/Source/Developer/TraceServices/Private/Model/AllocationsProvider.h
+++ b/Engine/Source/Developer/TraceServices/Private/Model/AllocationsProvider.h
@@ -339,7 +339,7 @@ public:
 	void EditAddTagSpec(TagIdType Tag, TagIdType ParentTag, const TCHAR* Display) { EditAccessCheck(); TagTracker.AddTagSpec(Tag, ParentTag, Display); }
 	void EditPushTag(uint32 ThreadId, uint8 Tracker, TagIdType Tag);
 	void EditPopTag(uint32 ThreadId, uint8 Tracker);
-	void EditPushTagFromPtr(uint32 ThreadId, uint8 Tracker, uint64 Ptr);
+	void EditPushTagFromPtr(uint32 ThreadId, uint8 Tracker, uint64 Ptr, HeapId RootHeapId);
 	void EditPopTagFromPtr(uint32 ThreadId, uint8 Tracker);
 
 	void EditOnAnalysisCompleted(double Time);
