TasksInsights refactoring:
* fixed UE-166680
* refactored how task traces are used
* refactored task traces analysis
* added task size parameter
* added an ability to select which tasks to show in the tasks table: "alive", "waiting for prerequisites", "queued", "executing", "waiting for nested" and "completed"

diff --git a/Engine/Source/Developer/TraceInsights/Private/Insights/TaskGraphProfiler/ViewModels/TaskEntry.cpp b/Engine/Source/Developer/TraceInsights/Private/Insights/TaskGraphProfiler/ViewModels/TaskEntry.cpp
--- a/Engine/Source/Developer/TraceInsights/Private/Insights/TaskGraphProfiler/ViewModels/TaskEntry.cpp
+++ b/Engine/Source/Developer/TraceInsights/Private/Insights/TaskGraphProfiler/ViewModels/TaskEntry.cpp
@@ -25,11 +25,12 @@ FTaskEntry::FTaskEntry(const TraceServices::FTaskInfo& TaskInfo)
 	, CompletedThreadId(TaskInfo.CompletedThreadId)
 	, DestroyedTimestamp(TaskInfo.DestroyedTimestamp)
 	, DestroyedThreadId(TaskInfo.DestroyedThreadId)
+	, TaskSize(TaskInfo.TaskSize)
+	, NumSubsequents(TaskInfo.Subsequents.Num())
+	, NumPrerequisites(TaskInfo.Prerequisites.Num())
+	, NumParents(TaskInfo.ParentTasks.Num())
+	, NumNested(TaskInfo.NestedTasks.Num())
 {
-	NumSubsequents = TaskInfo.Subsequents.Num();
-	NumPrerequisites = TaskInfo.Prerequisites.Num();
-	NumParents = TaskInfo.ParentTasks.Num();
-	NumNested = TaskInfo.NestedTasks.Num();
 }
 
 ////////////////////////////////////////////////////////////////////////////////////////////////////
diff --git a/Engine/Source/Developer/TraceInsights/Private/Insights/TaskGraphProfiler/ViewModels/TaskEntry.h b/Engine/Source/Developer/TraceInsights/Private/Insights/TaskGraphProfiler/ViewModels/TaskEntry.h
--- a/Engine/Source/Developer/TraceInsights/Private/Insights/TaskGraphProfiler/ViewModels/TaskEntry.h
+++ b/Engine/Source/Developer/TraceInsights/Private/Insights/TaskGraphProfiler/ViewModels/TaskEntry.h
@@ -42,6 +42,8 @@ public:
 	double GetDestroyedTimestamp() const { return DestroyedTimestamp; }
 	uint32 GetDestroyedThreadId() const { return DestroyedThreadId; }
 
+	uint64 GetTaskSize() const { return TaskSize; }
+
 	uint32 GetNumParents() const { return NumParents; }
 	uint32 GetNumNested() const { return NumNested; }
 	uint32 GetNumPrerequisites() const { return NumPrerequisites; }
@@ -76,6 +78,8 @@ private:
 	double DestroyedTimestamp;
 	uint32 DestroyedThreadId;
 
+	uint64 TaskSize;
+
 	uint32 NumSubsequents;
 	uint32 NumPrerequisites;
 	uint32 NumParents;
diff --git a/Engine/Source/Developer/TraceInsights/Private/Insights/TaskGraphProfiler/ViewModels/TaskTable.cpp b/Engine/Source/Developer/TraceInsights/Private/Insights/TaskGraphProfiler/ViewModels/TaskTable.cpp
--- a/Engine/Source/Developer/TraceInsights/Private/Insights/TaskGraphProfiler/ViewModels/TaskTable.cpp
+++ b/Engine/Source/Developer/TraceInsights/Private/Insights/TaskGraphProfiler/ViewModels/TaskTable.cpp
@@ -32,6 +32,7 @@ const FName FTaskTableColumns::CompletedTimestampColumnId(TEXT("CompletedTimesta
 const FName FTaskTableColumns::CompletedThreadIdColumnId(TEXT("CompletedThreadId"));
 const FName FTaskTableColumns::DestroyedTimestampColumnId(TEXT("DestroyedTimestamp"));
 const FName FTaskTableColumns::DestroyedThreadIdColumnId(TEXT("DestroyedThreadId"));
+const FName FTaskTableColumns::TaskSizeColumnId(TEXT("TaskSize"));
 const FName FTaskTableColumns::NumParentColumnId(TEXT("NumParents"));
 const FName FTaskTableColumns::NumNestedColumnId(TEXT("NumNested"));
 const FName FTaskTableColumns::NumSubsequentsColumnId(TEXT("NumSubsequents"));
@@ -94,6 +95,8 @@ struct DefaultTaskFieldGetterFuncts
 	static FTableCellValue GetDestroyedTimestamp(const FTableColumn& Column, const FTaskEntry& Task) { return FTableCellValue(Task.GetDestroyedTimestamp()); }
 	static FTableCellValue GetDestroyedThreadId(const FTableColumn& Column, const FTaskEntry& Task) { return FTableCellValue((int64)Task.GetDestroyedThreadId()); }
 
+	static FTableCellValue GetTaskSize(const FTableColumn& Column, const FTaskEntry& Task) { return FTableCellValue((int64)Task.GetTaskSize()); }
+
 	static FTableCellValue GetNumParents(const FTableColumn& Column, const FTaskEntry& Task) { return FTableCellValue((int64)Task.GetNumParents()); }
 	static FTableCellValue GetNumNested(const FTableColumn& Column, const FTaskEntry& Task) { return FTableCellValue((int64)Task.GetNumNested());	}
 	static FTableCellValue GetNumSubsequents(const FTableColumn& Column, const FTaskEntry& Task) { return FTableCellValue((int64)Task.GetNumSubsequents()); }
@@ -631,6 +634,36 @@ void FTaskTable::AddDefaultColumns()
 		AddColumn(ColumnRef);
 	}
 	//////////////////////////////////////////////////
+	// Task Size Column
+	{
+		TSharedRef<FTableColumn> ColumnRef = MakeShared<FTableColumn>(FTaskTableColumns::TaskSizeColumnId);
+		FTableColumn& Column = *ColumnRef;
+
+		Column.SetIndex(ColumnIndex++);
+
+		Column.SetShortName(LOCTEXT("TaskSizeColumnName", "Task Size"));
+		Column.SetTitleName(LOCTEXT("TaskSizeColumnTitle", "Task Size"));
+		Column.SetDescription(LOCTEXT("TaskSizeColumnDesc", "The size of the task including user-provided task body."));
+
+		Column.SetFlags(ETableColumnFlags::ShouldBeVisible | ETableColumnFlags::CanBeHidden | ETableColumnFlags::CanBeFiltered);
+
+		Column.SetHorizontalAlignment(HAlign_Left);
+		Column.SetInitialWidth(60.0f);
+
+		Column.SetDataType(ETableCellDataType::Int64);
+
+		TSharedRef<ITableCellValueGetter> Getter = MakeShared<FTaskColumnValueGetter<DefaultTaskFieldGetterFuncts::GetTaskSize>>();
+		Column.SetValueGetter(Getter);
+
+		TSharedRef<ITableCellValueFormatter> Formatter = MakeShared<FInt64ValueFormatterAsNumber>();
+		Column.SetValueFormatter(Formatter);
+
+		TSharedRef<ITableCellValueSorter> Sorter = MakeShared<FSorterByInt64Value>(ColumnRef);
+		Column.SetValueSorter(Sorter);
+
+		AddColumn(ColumnRef);
+	}
+	//////////////////////////////////////////////////
 	// Num Parent Tasks Column
 	{
 		TSharedRef<FTableColumn> ColumnRef = MakeShared<FTableColumn>(FTaskTableColumns::NumParentColumnId);
diff --git a/Engine/Source/Developer/TraceInsights/Private/Insights/TaskGraphProfiler/ViewModels/TaskTable.h b/Engine/Source/Developer/TraceInsights/Private/Insights/TaskGraphProfiler/ViewModels/TaskTable.h
--- a/Engine/Source/Developer/TraceInsights/Private/Insights/TaskGraphProfiler/ViewModels/TaskTable.h
+++ b/Engine/Source/Developer/TraceInsights/Private/Insights/TaskGraphProfiler/ViewModels/TaskTable.h
@@ -32,6 +32,7 @@ struct FTaskTableColumns
 	static const FName CompletedThreadIdColumnId;
 	static const FName DestroyedTimestampColumnId;
 	static const FName DestroyedThreadIdColumnId;
+	static const FName TaskSizeColumnId;
 	static const FName NumParentColumnId;
 	static const FName NumNestedColumnId;
 	static const FName NumSubsequentsColumnId;
diff --git a/Engine/Source/Developer/TraceInsights/Private/Insights/TaskGraphProfiler/Widgets/STaskTableTreeView.cpp b/Engine/Source/Developer/TraceInsights/Private/Insights/TaskGraphProfiler/Widgets/STaskTableTreeView.cpp
--- a/Engine/Source/Developer/TraceInsights/Private/Insights/TaskGraphProfiler/Widgets/STaskTableTreeView.cpp
+++ b/Engine/Source/Developer/TraceInsights/Private/Insights/TaskGraphProfiler/Widgets/STaskTableTreeView.cpp
@@ -165,7 +165,7 @@ void STaskTableTreeView::RebuildTree(bool bResync)
 		return;
 	}
 
-	if (NewQueryStartTime != QueryStartTime || NewQueryEndTime != QueryEndTime)
+	if (bResync || NewQueryStartTime != QueryStartTime || NewQueryEndTime != QueryEndTime)
 	{
 		StopAllTableDataTasks();
 
@@ -188,15 +188,15 @@ void STaskTableTreeView::RebuildTree(bool bResync)
 
 				TArray<FTableTreeNodePtr>* Nodes = &TableTreeNodes;
 
-				TasksProvider->EnumerateTasks(QueryStartTime, QueryEndTime, [&Tasks, &TaskTable, &BaseNodeName, Nodes](const TraceServices::FTaskInfo& TaskInfo)
-					{
-						Tasks.Emplace(TaskInfo);
-						uint32 Index = Tasks.Num() - 1;
-						FName NodeName(BaseNodeName, static_cast<int32>(TaskInfo.Id + 1));
-						FTaskNodePtr NodePtr = MakeShared<FTaskNode>(NodeName, TaskTable, Index);
-						Nodes->Add(NodePtr);
-						return TraceServices::ETaskEnumerationResult::Continue;
-					});
+				TasksProvider->EnumerateTasks(QueryStartTime, QueryEndTime, SelectedTasksSelectionOption, [&Tasks, &TaskTable, &BaseNodeName, Nodes](const TraceServices::FTaskInfo& TaskInfo)
+				{
+					Tasks.Emplace(TaskInfo);
+					uint32 Index = Tasks.Num() - 1;
+					FName NodeName(BaseNodeName, static_cast<int32>(TaskInfo.Id + 1));
+					FTaskNodePtr NodePtr = MakeShared<FTaskNode>(NodeName, TaskTable, Index);
+					Nodes->Add(NodePtr);
+					return TraceServices::ETaskEnumerationResult::Continue;
+				});
 			}
 		}
 
@@ -232,6 +232,34 @@ BEGIN_SLATE_FUNCTION_BUILD_OPTIMIZATION
 TSharedPtr<SWidget> STaskTableTreeView::ConstructToolbar()
 {
 	return SNew(SHorizontalBox)
+		+ SHorizontalBox::Slot()
+		.AutoWidth()
+		.Padding(0.0f, 0.0f, 4.0f, 0.0f)
+		.VAlign(VAlign_Center)
+		[
+			SNew(STextBlock)
+			.Text(LOCTEXT("Tasks", "Tasks"))
+		]
+		
+		+ SHorizontalBox::Slot()
+		.AutoWidth()
+		.Padding(4.0f, 0.0f, 10.0f, 0.0f)
+		[
+			SNew(SBox)
+			.MinDesiredWidth(160.0f)
+			[
+				SNew(SComboBox<TSharedPtr<ETaskEnumerationOption>>)
+				.OptionsSource(GetAvailableTasksSelectionOptions())
+				.OnSelectionChanged(this, &STaskTableTreeView::TasksSelectionOptions_OnSelectionChanged)
+				.OnGenerateWidget(this, &STaskTableTreeView::TasksSelectionOptions_OnGenerateWidget)
+				.IsEnabled(this, &STaskTableTreeView::TasksSelectionOptions_IsEnabled)
+				[
+					SNew(STextBlock)
+					.Text(this, &STaskTableTreeView::TasksSelectionOptions_GetSelectionText)
+				]
+			]
+		]
+		
 		+ SHorizontalBox::Slot()
 		.AutoWidth()
 		.Padding(0.0f, 0.0f, 4.0f, 0.0f)
@@ -305,6 +333,60 @@ TSharedRef<SWidget> STaskTableTreeView::TimestampOptions_OnGenerateWidget(TShare
 
 	return Widget;
 }
+
+////////////////////////////////////////////////////////////////////////////////////////////////////
+
+TSharedRef<SWidget> STaskTableTreeView::TasksSelectionOptions_OnGenerateWidget(TSharedPtr<ETaskEnumerationOption> InOption)
+{
+	auto GetTooltipText = [](ETaskEnumerationOption InOption)
+	{
+		switch (InOption)
+		{
+		case ETaskEnumerationOption::Alive:
+		{
+			return LOCTEXT("AliveTooltip", "Tasks that were created before the selection and destroyed after.");
+		}
+		case ETaskEnumerationOption::Active:
+		{
+			return LOCTEXT("ActiveTooltip", "Tasks that were active (being executed) at any moment of the selection.");
+		}
+		case ETaskEnumerationOption::WaitingForPrerequisites:
+		{
+			return LOCTEXT("WaitingForPrerequisitesTooltip", "Tasks that are blocked by prerequisites for the entire duration of the selection.");
+		}
+		case ETaskEnumerationOption::Queued:
+		{
+			return LOCTEXT("QueuedTooltip", "Tasks that for the entire duration of the selection are waiting in the scheduler queue.");
+		}
+		case ETaskEnumerationOption::Executing:
+		{
+			return LOCTEXT("ExecutingTooltip", "Tasks that for the entire duration of the selection have been executed.");
+		}
+		case ETaskEnumerationOption::WaitingForNested:
+		{
+			return LOCTEXT("WaitingForNestedTooltip", "Tasks that for the entire duration of the selection have been waiting for nested tasks.");
+		}
+		case ETaskEnumerationOption::Completed:
+		{
+			return LOCTEXT("CompletedTooltip", "Tasks that for the entire duration of the selection have been already completed but not destroyed.");
+		}
+		}
+
+		return FText();
+	};
+
+	TSharedRef<SHorizontalBox> Widget = SNew(SHorizontalBox);
+	Widget->AddSlot()
+		.AutoWidth()
+		[
+			SNew(STextBlock)
+			.Text(TasksSelectionOptions_GetText(*InOption))
+			.ToolTipText(GetTooltipText(*InOption))
+			.Margin(2.0f)
+		];
+
+	return Widget;
+}
 END_SLATE_FUNCTION_BUILD_OPTIMIZATION
 
 ////////////////////////////////////////////////////////////////////////////////////////////////////
@@ -330,6 +412,24 @@ const TArray<TSharedPtr<STaskTableTreeView::ETimestampOptions>>* STaskTableTreeV
 
 ////////////////////////////////////////////////////////////////////////////////////////////////////
 
+const TArray<TSharedPtr<ETaskEnumerationOption>>* STaskTableTreeView::GetAvailableTasksSelectionOptions()
+{
+	if (AvailableTasksSelectionOptions.Num() == 0)
+	{
+		AvailableTasksSelectionOptions.Add(MakeShared<ETaskEnumerationOption>(ETaskEnumerationOption::Alive));
+		AvailableTasksSelectionOptions.Add(MakeShared<ETaskEnumerationOption>(ETaskEnumerationOption::Active));
+		AvailableTasksSelectionOptions.Add(MakeShared<ETaskEnumerationOption>(ETaskEnumerationOption::WaitingForPrerequisites));
+		AvailableTasksSelectionOptions.Add(MakeShared<ETaskEnumerationOption>(ETaskEnumerationOption::Queued));
+		AvailableTasksSelectionOptions.Add(MakeShared<ETaskEnumerationOption>(ETaskEnumerationOption::Executing));
+		AvailableTasksSelectionOptions.Add(MakeShared<ETaskEnumerationOption>(ETaskEnumerationOption::WaitingForNested));
+		AvailableTasksSelectionOptions.Add(MakeShared<ETaskEnumerationOption>(ETaskEnumerationOption::Completed));
+	}
+
+	return &AvailableTasksSelectionOptions;
+}
+
+////////////////////////////////////////////////////////////////////////////////////////////////////
+
 void STaskTableTreeView::TimestampOptions_OnSelectionChanged(TSharedPtr<ETimestampOptions> InOption, ESelectInfo::Type SelectInfo)
 {
 	TimestampOptions_OnSelectionChanged(*InOption);
@@ -365,6 +465,21 @@ void STaskTableTreeView::TimestampOptions_OnSelectionChanged(ETimestampOptions I
 
 ////////////////////////////////////////////////////////////////////////////////////////////////////
 
+void STaskTableTreeView::TasksSelectionOptions_OnSelectionChanged(TSharedPtr<ETaskEnumerationOption> InOption, ESelectInfo::Type SelectInfo)
+{
+	TasksSelectionOptions_OnSelectionChanged(*InOption);
+}
+
+////////////////////////////////////////////////////////////////////////////////////////////////////
+
+void STaskTableTreeView::TasksSelectionOptions_OnSelectionChanged(ETaskEnumerationOption InOption)
+{
+	SelectedTasksSelectionOption = InOption;
+	RebuildTree(true);
+}
+
+////////////////////////////////////////////////////////////////////////////////////////////////////
+
 FText STaskTableTreeView::TimestampOptions_GetSelectionText() const
 {
 	return TimestampOptions_GetText(SelectedTimestampOption);
@@ -395,6 +510,52 @@ FText STaskTableTreeView::TimestampOptions_GetText(ETimestampOptions InOption) c
 
 ////////////////////////////////////////////////////////////////////////////////////////////////////
 
+FText STaskTableTreeView::TasksSelectionOptions_GetSelectionText() const
+{
+	return TasksSelectionOptions_GetText(SelectedTasksSelectionOption);
+}
+
+////////////////////////////////////////////////////////////////////////////////////////////////////
+
+FText STaskTableTreeView::TasksSelectionOptions_GetText(ETaskEnumerationOption InOption) const
+{
+	switch (InOption)
+	{
+	case ETaskEnumerationOption::Alive:
+	{
+		return LOCTEXT("Alive", "Alive");
+	}
+	case ETaskEnumerationOption::Active:
+	{
+		return LOCTEXT("Active", "Active");
+	}
+	case ETaskEnumerationOption::WaitingForPrerequisites:
+	{
+		return LOCTEXT("WaitingForPrerequisites", "Waiting for prerequisites");
+	}
+	case ETaskEnumerationOption::Queued:
+	{
+		return LOCTEXT("Queued", "Queued");
+	}
+	case ETaskEnumerationOption::Executing:
+	{
+		return LOCTEXT("Executing", "Executing");
+	}
+	case ETaskEnumerationOption::WaitingForNested:
+	{
+		return LOCTEXT("WaitingForNested", "Waiting for nested");
+	}
+	case ETaskEnumerationOption::Completed:
+	{
+		return LOCTEXT("Completed", "Completed");
+	}
+	}
+
+	return FText();
+}
+
+////////////////////////////////////////////////////////////////////////////////////////////////////
+
 bool STaskTableTreeView::TimestampOptions_IsEnabled() const
 {
 	return !bIsUpdateRunning;
@@ -402,6 +563,13 @@ bool STaskTableTreeView::TimestampOptions_IsEnabled() const
 
 ////////////////////////////////////////////////////////////////////////////////////////////////////
 
+bool STaskTableTreeView::TasksSelectionOptions_IsEnabled() const
+{
+	return !bIsUpdateRunning;
+}
+
+////////////////////////////////////////////////////////////////////////////////////////////////////
+
 bool STaskTableTreeView::ContextMenu_GoToTask_CanExecute() const
 {
 	TArray<FTableTreeNodePtr> SelectedItems;
diff --git a/Engine/Source/Developer/TraceInsights/Private/Insights/TaskGraphProfiler/Widgets/STaskTableTreeView.h b/Engine/Source/Developer/TraceInsights/Private/Insights/TaskGraphProfiler/Widgets/STaskTableTreeView.h
--- a/Engine/Source/Developer/TraceInsights/Private/Insights/TaskGraphProfiler/Widgets/STaskTableTreeView.h
+++ b/Engine/Source/Developer/TraceInsights/Private/Insights/TaskGraphProfiler/Widgets/STaskTableTreeView.h
@@ -105,6 +105,14 @@ private:
 	FText TimestampOptions_GetText(ETimestampOptions InOption) const;
 	bool TimestampOptions_IsEnabled() const;
 
+	const TArray<TSharedPtr<TraceServices::ETaskEnumerationOption>>* GetAvailableTasksSelectionOptions();
+	TSharedRef<SWidget> TasksSelectionOptions_OnGenerateWidget(TSharedPtr<TraceServices::ETaskEnumerationOption> InOption);
+	void TasksSelectionOptions_OnSelectionChanged(TSharedPtr<TraceServices::ETaskEnumerationOption> InOption, ESelectInfo::Type SelectInfo);
+	void TasksSelectionOptions_OnSelectionChanged(TraceServices::ETaskEnumerationOption InOption);
+	FText TasksSelectionOptions_GetSelectionText() const;
+	FText TasksSelectionOptions_GetText(TraceServices::ETaskEnumerationOption InOption) const;
+	bool TasksSelectionOptions_IsEnabled() const;
+
 	bool ContextMenu_GoToTask_CanExecute() const;
 	void ContextMenu_GoToTask_Execute();
 
@@ -118,6 +126,9 @@ private:
 	ETimestampOptions SelectedTimestampOption = ETimestampOptions::RelativeToPrevious;
 	TArray<TSharedPtr<ETimestampOptions>> AvailableTimestampOptions;
 
+	TraceServices::ETaskEnumerationOption SelectedTasksSelectionOption = TraceServices::ETaskEnumerationOption::Alive;
+	TArray<TSharedPtr<TraceServices::ETaskEnumerationOption>> AvailableTasksSelectionOptions;
+
 	TaskTrace::FId TaskIdToSelect;
 };
 
diff --git a/Engine/Source/Developer/TraceServices/Private/Analyzers/TasksAnalysis.cpp b/Engine/Source/Developer/TraceServices/Private/Analyzers/TasksAnalysis.cpp
--- a/Engine/Source/Developer/TraceServices/Private/Analyzers/TasksAnalysis.cpp
+++ b/Engine/Source/Developer/TraceServices/Private/Analyzers/TasksAnalysis.cpp
@@ -58,7 +58,8 @@ bool FTasksAnalyzer::OnEvent(uint16 RouteId, EStyle Style, const FOnEventContext
 		{
 			double Timestamp = Context.EventTime.AsSeconds(EventData.GetValue<uint64>("Timestamp"));
 			TaskTrace::FId TaskId = EventData.GetValue<TaskTrace::FId>("TaskId");
-			TasksProvider.TaskCreated(TaskId, Timestamp, ThreadId);
+			uint64 TaskSize = EventData.GetValue<uint64>("TaskSize");
+			TasksProvider.TaskCreated(TaskId, Timestamp, ThreadId, TaskSize);
 			break;
 		}
 		case RouteId_Launched:
@@ -75,8 +76,9 @@ bool FTasksAnalyzer::OnEvent(uint16 RouteId, EStyle Style, const FOnEventContext
 
 			bool bTracked = EventData.GetValue<bool>("Tracked");
 			int32 ThreadToExecuteOn = EventData.GetValue<int32>("ThreadToExecuteOn");
+			uint64 TaskSize = EventData.GetValue<uint64>("TaskSize");
 
-			TasksProvider.TaskLaunched(TaskId, DebugName, bTracked, ThreadToExecuteOn, Timestamp, ThreadId);
+			TasksProvider.TaskLaunched(TaskId, DebugName, bTracked, ThreadToExecuteOn, Timestamp, ThreadId, TaskSize);
 
 			break;
 		}
diff --git a/Engine/Source/Developer/TraceServices/Private/Model/TasksProfiler.cpp b/Engine/Source/Developer/TraceServices/Private/Model/TasksProfiler.cpp
--- a/Engine/Source/Developer/TraceServices/Private/Model/TasksProfiler.cpp
+++ b/Engine/Source/Developer/TraceServices/Private/Model/TasksProfiler.cpp
@@ -9,6 +9,31 @@
 #include "Common/Utils.h"
 #include "Async/TaskGraphInterfaces.h"
 
+//////////////////////////////////////////
+// Handling of task traces
+// To reduce the amount of traces, or when a particular trace event is not possible for a particular task type, they are often omitted. 
+// In this case omitted traces are assumed and restored here.
+// Different task types trace differently. Notable examples are:
+// * Tasks System regular tasks: (almost) no surprises here, 
+// ["created"] -> "launched" -> ["scheduled"] -> "started" -> "finished" -> "completed" -> "destroyed",
+// where [trace] can be ommitted and assumed to happen at the same time and in the same thread as the subsequent trace.
+// The only surprise here is that "started" can race "scheduled" traces during task retraction.
+// * Tasks System task events (`FTaskEvent`): are never scheduled or executed, "launched" -> "completed" -> "destroyed". 
+// "created" == "launched", "scheduled" == "started" == "finished" == "completed". 
+// FTaskEvent is detected by "invalid" "started" timestamp when "completed" timestamp arrives.
+// * TaskGraph regular tasks: behave mostly as Tasks System regular tasks except for the case when in the moment of their completion 
+// they have incomplete "don't complete until" tasks (aka "nested" tasks). In this case they are traced as two different tasks: 
+// the original one and the new one created for flagging the graph event completed when all nested tasks are completed. 
+// the original task will have only "completed" trace at the moment when the second task is created, the second task will behave as a regular task.
+// * TaskGraph standalone FGraphEvent (created by `FGraphEvent::CreateGraphEvent()`): similar to regular TaskGraph tasks, except w/o nested
+// tasks - "completed" -> "destroyed".
+// * ParallelFor tasks: usually - "launched" -> "started" -> "completed" -> "destroyed". "created" == "launched", "scheduled" == "launched", 
+// "finished" == "completed". but if it was cancelled: "launched" -> "destroyed"
+
+#if !defined(UE_TASK_TRACES_ENABLE_OLD_TASKS_SUPPORT)
+#define UE_TASK_TRACES_ENABLE_OLD_TASKS_SUPPORT 1
+#endif
+
 namespace TraceServices
 {
 	FTasksProvider::FTasksProvider(IAnalysisSession& InSession)
@@ -91,6 +116,16 @@ namespace TraceServices
 		ExecutionTimeCounter->SetDescription(TEXT("Tasks: execution time"));
 		ExecutionTimeCounter->SetIsFloatingPoint(true);
 
+		AliveTasksCounter = EditableCounterProvider.CreateEditableCounter();
+		AliveTasksCounter->SetName(TEXT("Tasks::AliveTasks"));
+		AliveTasksCounter->SetDescription(TEXT("Tasks: the numbers of tasks alive (created but not destroyed)"));
+		AliveTasksCounter->SetIsFloatingPoint(false);
+
+		AliveTasksSizeCounter = EditableCounterProvider.CreateEditableCounter();
+		AliveTasksSizeCounter->SetName(TEXT("Tasks::AliveTasksSize"));
+		AliveTasksSizeCounter->SetDescription(TEXT("Tasks: the total size of alive tasks (created but not destroyed)"));
+		AliveTasksSizeCounter->SetIsFloatingPoint(false);
+
 		bCountersCreated = true;
 	}
 
@@ -104,7 +139,7 @@ namespace TraceServices
 		}
 	}
 
-	void FTasksProvider::TaskCreated(TaskTrace::FId TaskId, double Timestamp, uint32 ThreadId)
+	void FTasksProvider::TaskCreated(TaskTrace::FId TaskId, double Timestamp, uint32 ThreadId, uint64 TaskSize)
 	{
 		UE_LOG(LogTraceServices, Verbose, TEXT("TaskCreated(TaskId: %d, Timestamp %.6f)"), TaskId, Timestamp);
 
@@ -119,9 +154,18 @@ namespace TraceServices
 
 		Task->CreatedTimestamp = Timestamp;
 		Task->CreatedThreadId = ThreadId;
+		Task->TaskSize = TaskSize;
+
+		if (!bCountersCreated)
+		{
+			CreateCounters();
+		}
+
+		AliveTasksCounter->SetValue(Timestamp, ++AliveTasksNum);
+		AliveTasksSizeCounter->SetValue(Timestamp, AliveTasksSize += TaskSize);
 	}
 
-	void FTasksProvider::TaskLaunched(TaskTrace::FId TaskId, const TCHAR* DebugName, bool bTracked, int32 ThreadToExecuteOn, double Timestamp, uint32 ThreadId)
+	void FTasksProvider::TaskLaunched(TaskTrace::FId TaskId, const TCHAR* DebugName, bool bTracked, int32 ThreadToExecuteOn, double Timestamp, uint32 ThreadId, uint64 TaskSize)
 	{
 		UE_LOG(LogTraceServices, Verbose, TEXT("TaskLaunched(TaskId: %d, DebugName: %s, bTracked: %d, Timestamp %.6f)"), TaskId, DebugName, bTracked, Timestamp);
 
@@ -134,12 +178,18 @@ namespace TraceServices
 			return;
 		}
 
-		checkf(Task->LaunchedTimestamp == FTaskInfo::InvalidTimestamp, TEXT("%d"), TaskId);
+#if UE_TASK_TRACES_ENABLE_OLD_TASKS_SUPPORT // a task with incomplete nested tasks is "relaunched" when it's execution is finished
+		if (Task->LaunchedTimestamp != FTaskInfo::InvalidTimestamp)
+		{
+			return;
+		}
+#else
+		checkf(Task->LaunchedTimestamp == FTaskInfo::InvalidTimestamp, TEXT("TaskLaunched(TaskId %d, Timestamp %.6f, ThreadId %d) repeated, prev Timestamp %.6f, ThreadId %d"), TaskId, Timestamp, ThreadId, Task->LaunchedTimestamp, Task->LaunchedThreadId);
+#endif
 			
 		if (Task->CreatedTimestamp == FTaskInfo::InvalidTimestamp) // created and launched in one go
 		{
-			Task->CreatedTimestamp = Timestamp;
-			Task->CreatedThreadId = ThreadId;
+			TaskCreated(TaskId, Timestamp, ThreadId, TaskSize);
 		}
 
 		Task->DebugName = DebugName;
@@ -167,15 +217,17 @@ namespace TraceServices
 			return;
 		}
 
-		TaskScheduled(*Task, Timestamp, ThreadId);
-	}
-
-	void FTasksProvider::TaskScheduled(FTaskInfo& Task, double Timestamp, uint32 ThreadId)
-	{
-		if (!TryRegisterEvent(TEXT("TaskScheduled"), Task, &FTaskInfo::ScheduledTimestamp, Timestamp, &FTaskInfo::ScheduledThreadId, ThreadId))
+#if UE_TASK_TRACES_ENABLE_OLD_TASKS_SUPPORT // a task with incomplete nested tasks is "relaunched" when it's execution is finished
+		if (Task->ScheduledTimestamp != FTaskInfo::InvalidTimestamp)
 		{
 			return;
 		}
+#else
+		checkf(Task->ScheduledTimestamp == FTaskInfo::InvalidTimestamp, TEXT("TaskScheduled(TaskId %d, Timestamp %.6f, ThreadId %d) repeated, prev Timestamp %.6f, ThreadId %d"), TaskId, Timestamp, ThreadId, Task->ScheduledTimestamp, Task->ScheduledThreadId);
+#endif
+
+		Task->ScheduledTimestamp = Timestamp;
+		Task->ScheduledThreadId = ThreadId;
 
 		if (!bCountersCreated)
 		{
@@ -183,7 +235,7 @@ namespace TraceServices
 		}
 
 		WaitingForPrerequisitesTasksCounter->SetValue(Timestamp, --WaitingForPrerequisitesTasksNum);
-		if (IsNamedThread(Task.ThreadToExecuteOn))
+		if (IsNamedThread(Task->ThreadToExecuteOn))
 		{
 			NamedThreadsScheduledTasksCounter->SetValue(Timestamp, ++ScheduledTasksNum);
 		}
@@ -240,16 +292,17 @@ namespace TraceServices
 			return;
 		}
 
-		if (Task->ScheduledTimestamp == FTaskInfo::InvalidTimestamp)
-		{
-			// optimisation for tasks that can't have dependencies, and so are scheduled immediately after launching. in this case "scheduled" trace is not sent
-			TaskScheduled(*Task, Task->LaunchedTimestamp, Task->LaunchedThreadId);
-		}
-
-		if (!TryRegisterEvent(TEXT("TaskStarted"), *Task, &FTaskInfo::StartedTimestamp, Timestamp, &FTaskInfo::StartedThreadId, ThreadId))
+#if UE_TASK_TRACES_ENABLE_OLD_TASKS_SUPPORT // a task with incomplete nested tasks is "relaunched" when it's execution is finished
+		if (Task->StartedTimestamp != FTaskInfo::InvalidTimestamp)
 		{
 			return;
 		}
+#else
+		checkf(Task->StartedTimestamp == FTaskInfo::InvalidTimestamp, TEXT("TaskStarted(TaskId %d, Timestamp %.6f, ThreadId %d) repeated, prev Timestamp %.6f, ThreadId %d"), TaskId, Timestamp, ThreadId, Task->StartedTimestamp, Task->StartedThreadId);
+#endif
+
+		Task->StartedTimestamp = Timestamp;
+		Task->StartedThreadId = ThreadId;
 
 		ExecutionThreads.FindOrAdd(ThreadId).Add(TaskId);
 
@@ -282,15 +335,16 @@ namespace TraceServices
 			return;
 		}
 
-		TaskFinished(*Task, Timestamp);
-	}
-
-	void FTasksProvider::TaskFinished(FTaskInfo& Task, double Timestamp)
-	{
-		if (!TryRegisterEvent(TEXT("TaskFinished"), Task, &FTaskInfo::FinishedTimestamp, Timestamp))
+#if UE_TASK_TRACES_ENABLE_OLD_TASKS_SUPPORT // a task with incomplete nested tasks is "relaunched" when it's execution is finished
+		if (Task->FinishedTimestamp != FTaskInfo::InvalidTimestamp)
 		{
 			return;
 		}
+#else
+		checkf(Task->FinishedTimestamp == FTaskInfo::InvalidTimestamp, TEXT("TaskFinished(TaskId %d, Timestamp %.6f) repeated, prev Timestamp %.6f"), TaskId, Timestamp, Task->StartedTimestamp);
+#endif
+
+		Task->FinishedTimestamp = Timestamp;
 
 		if (!bCountersCreated)
 		{
@@ -298,7 +352,7 @@ namespace TraceServices
 		}
 
 		RunningTasksCounter->SetValue(Timestamp, --RunningTasksNum);
-		ExecutionTimeCounter->SetValue(Timestamp, Task.FinishedTimestamp - Task.StartedTimestamp);
+		ExecutionTimeCounter->SetValue(Timestamp, Task->FinishedTimestamp - Task->StartedTimestamp);
 	}
 
 	void FTasksProvider::TaskCompleted(TaskTrace::FId TaskId, double Timestamp, uint32 ThreadId)
@@ -314,33 +368,81 @@ namespace TraceServices
 			return;
 		}
 
-		Task->Id = TaskId;
-
-		if (Task->ScheduledTimestamp == FTaskInfo::InvalidTimestamp) // task events are never scheduled or executed
-		{
-			// pretend scheduling and execution happened at the moment of completion
-			Task->LaunchedTimestamp = Timestamp;
-			Task->LaunchedThreadId = ThreadId;
-			Task->ScheduledTimestamp = Timestamp;
-			Task->ScheduledThreadId = ThreadId;
-			Task->StartedTimestamp = Timestamp;
-			Task->StartedThreadId = ThreadId;
-			Task->FinishedTimestamp = Timestamp;
+#if UE_TASK_TRACES_ENABLE_OLD_TASKS_SUPPORT
+		if (Task->LaunchedTimestamp == FTaskInfo::InvalidTimestamp)
+		{	// a standalone FGraphEvent w/o an accompanying task or subsequents, the first trace is "completed" event on `DispatchSubsequents`
+			TaskLaunched(TaskId, /*DebugName = */ TEXT("GraphEvent"), /*bTracked = */ true, /*ThreadToExecuteOn = */ 0, Timestamp, ThreadId, /*TaskSize = */ sizeof(FGraphEvent));
+			//TaskStarted(TaskId, Timestamp, ThreadId);
+			//TaskFinished(TaskId, Timestamp);
 		}
-		else if (Task->FinishedTimestamp == FTaskInfo::InvalidTimestamp)
+#endif
+
+		checkf(Task->LaunchedTimestamp != FTaskInfo::InvalidTimestamp, TEXT("TaskCompleted(TaskId %d, Timestamp %.6f): no \"launched\" timestamp"), TaskId, Timestamp);
+		checkf(Task->CreatedTimestamp != FTaskInfo::InvalidTimestamp, TEXT("TaskCompleted(TaskId %d, Timestamp %.6f): no \"created\" timestamp"), TaskId, Timestamp);
+
+		if (Task->ScheduledTimestamp == FTaskInfo::InvalidTimestamp)
 		{
-			// optimisation for ParalellFor tasks that are finished and completed in one go, so no need to send separate trace messages
-			TaskFinished(*Task, Timestamp);
+			if (Task->FinishedTimestamp != FTaskInfo::InvalidTimestamp) // a regular task just missing omitted "scheduled" trace
+			{
+				TaskScheduled(TaskId, Task->StartedTimestamp, Task->StartedThreadId);
+			}
+			else if (Task->StartedTimestamp == FTaskInfo::InvalidTimestamp) // FTaskEvent
+			{
+				TaskScheduled(TaskId, Timestamp, ThreadId);
+				TaskStarted(TaskId, Timestamp, ThreadId);
+				TaskFinished(TaskId, Timestamp);
+			}
+			else // ParallelFor task
+			{
+				TaskScheduled(TaskId, Task->LaunchedTimestamp, Task->LaunchedThreadId);
+				TaskFinished(TaskId, Timestamp);
+			}
 		}
 
-		TryRegisterEvent(TEXT("TaskCompleted"), *Task, &FTaskInfo::CompletedTimestamp, Timestamp, &FTaskInfo::CompletedThreadId, ThreadId);
+		checkf(Task->CompletedTimestamp == FTaskInfo::InvalidTimestamp, TEXT("TaskCompleted(TaskId %d, Timestamp %.6f, ThreadId %d): repeated, prev Timestamp %.6f, ThreadId %d"), TaskId, Timestamp, ThreadId, Task->CompletedTimestamp, Task->CompletedThreadId);
+
+		Task->CompletedTimestamp = Timestamp;
+		Task->CompletedThreadId = ThreadId;
 	}
 
 	void FTasksProvider::TaskDestroyed(TaskTrace::FId TaskId, double Timestamp, uint32 ThreadId)
 	{
 		InitTaskIdToIndexConversion(TaskId);
 
-		TryRegisterEvent(TEXT("TaskDestroyed"), TaskId, &FTaskInfo::DestroyedTimestamp, Timestamp, &FTaskInfo::DestroyedThreadId, ThreadId);
+		FTaskInfo* Task = TryGetOrCreateTask(TaskId);
+		if (Task == nullptr)
+		{
+			UE_LOG(LogTraceServices, Log, TEXT("TaskStarted(TaskId %d, Timestamp %.6f) skipped"), TaskId, Timestamp);
+			return;
+		}
+
+		if (Task->CompletedTimestamp == FTaskInfo::InvalidTimestamp) // cancelled ParallelFor task
+		{
+			check(Task->ScheduledTimestamp == FTaskInfo::InvalidTimestamp &&
+				Task->StartedTimestamp == FTaskInfo::InvalidTimestamp &&
+				Task->FinishedTimestamp == FTaskInfo::InvalidTimestamp);
+			
+			TaskScheduled(TaskId, Timestamp, ThreadId);
+			TaskStarted(TaskId, Timestamp, ThreadId);
+			TaskFinished(TaskId, Timestamp);
+			TaskCompleted(TaskId, Timestamp, ThreadId);
+		}
+
+		// all initialised
+		check(Task->CreatedTimestamp != FTaskInfo::InvalidTimestamp && 
+			Task->LaunchedTimestamp != FTaskInfo::InvalidTimestamp &&
+			Task->ScheduledTimestamp != FTaskInfo::InvalidTimestamp &&
+			Task->StartedTimestamp != FTaskInfo::InvalidTimestamp &&
+			Task->FinishedTimestamp != FTaskInfo::InvalidTimestamp &&
+			Task->CompletedTimestamp != FTaskInfo::InvalidTimestamp);
+
+		checkf(Task->DestroyedTimestamp == FTaskInfo::InvalidTimestamp, TEXT("TaskDestroyed(TaskId %d, Timestamp %.6f, ThreadId %d): repeated, prev Timestamp %.6f, ThreadId %d"), TaskId, Timestamp, ThreadId, Task->DestroyedTimestamp, Task->DestroyedThreadId);
+
+		Task->DestroyedTimestamp = Timestamp;
+		Task->DestroyedThreadId = ThreadId;
+
+		AliveTasksCounter->SetValue(Timestamp, --AliveTasksNum);
+		AliveTasksSizeCounter->SetValue(Timestamp, AliveTasksSize -= Task->TaskSize);
 	}
 
 	void FTasksProvider::WaitingStarted(TArray<TaskTrace::FId> InTasks, double Timestamp, uint32 ThreadId)
@@ -403,8 +505,12 @@ namespace TraceServices
 
 		if (TaskIndex >= Tasks.Num())
 		{
+			int32 PrevNum = Tasks.Num();
 			Tasks.AddDefaulted(TaskIndex - Tasks.Num() + 1);
-			Tasks[TaskIndex].Id = TaskId;
+			for (int32 i = PrevNum; i != Tasks.Num(); ++i)
+			{
+				Tasks[i].Id = i + FirstTaskId;
+			}
 		}
 
 		return &Tasks[TaskIndex];
@@ -415,32 +521,6 @@ namespace TraceServices
 		return ENamedThreads::GetThreadIndex((ENamedThreads::Type)Thread) != ENamedThreads::AnyThread;
 	}
 
-	bool FTasksProvider::TryRegisterEvent(const TCHAR* EventName, TaskTrace::FId TaskId, double FTaskInfo::* TimestampPtr, double TimestampValue, uint32 FTaskInfo::* ThreadIdPtr/* = nullptr*/, uint32 ThreadIdValue/* = 0*/)
-	{
-		FTaskInfo* Task = TryGetTask(TaskId);
-		if (Task == nullptr)
-		{
-			UE_LOG(LogTraceServices, Log, TEXT("%s(TaskId %d, Timestamp %.6f) skipped"), EventName, TaskId, TimestampValue);
-			return false;
-		}
-
-		return TryRegisterEvent(EventName, *Task, TimestampPtr, TimestampValue, ThreadIdPtr, ThreadIdValue);
-	}
-
-	bool FTasksProvider::TryRegisterEvent(const TCHAR* EventName, FTaskInfo& Task, double FTaskInfo::* TimestampPtr, double TimestampValue, uint32 FTaskInfo::* ThreadIdPtr/* = nullptr*/, uint32 ThreadIdValue/* = 0*/)
-	{
-		UE_LOG(LogTraceServices, Verbose, TEXT("%s(TaskId: %d, Timestamp %.6f)"), EventName, Task.Id, TimestampValue);
-
-		checkf(Task.*TimestampPtr == FTaskInfo::InvalidTimestamp, TEXT("%s: TaskId %d, old TS %.6f, new TS %.6f"), EventName, Task.Id, Task.*TimestampPtr, TimestampValue);
-		Task.*TimestampPtr = TimestampValue;
-		if (ThreadIdPtr != nullptr)
-		{
-			Task.*ThreadIdPtr = ThreadIdValue;
-		}
-
-		return true;
-	}
-
 	void FTasksProvider::AddRelative(const TCHAR* RelationType, TaskTrace::FId TaskId, TArray<FTaskInfo::FRelationInfo> FTaskInfo::* RelationsPtr, TaskTrace::FId RelativeId, double Timestamp, uint32 ThreadId)
 	{
 		UE_LOG(LogTraceServices, Verbose, TEXT("%s (%d) added to TaskId: %d, Timestamp %.6f)"), RelationType, RelativeId, TaskId, Timestamp);
@@ -527,7 +607,7 @@ namespace TraceServices
 			return Result;
 		}
 
-		EnumerateTasks(StartTime, EndTime, 
+		EnumerateTasks(StartTime, EndTime, ETaskEnumerationOption::Alive,
 			[ThreadId, StartTime, EndTime, &Result] (const FTaskInfo& Task)
 			{
 				if (Task.LaunchedThreadId == ThreadId && Task.LaunchedTimestamp > StartTime && Task.FinishedTimestamp < EndTime)
@@ -546,43 +626,75 @@ namespace TraceServices
 		return Tasks.Num();
 	}
 
-	void FTasksProvider::EnumerateTasks(double StartTime, double EndTime, TaskCallback Callback) const
+	void FTasksProvider::EnumerateTasks(double StartTime, double EndTime, ETaskEnumerationOption EnumerationOption, TaskCallback Callback) const
 	{
-		for (const TPair<uint32, TArray64<TaskTrace::FId>>& KeyValue : ExecutionThreads)
-		{
-			const TArray64<TaskTrace::FId>& Thread = KeyValue.Value;
-			
-			// find the first task with `StartedTimestamp <= StartTime`
-			int64 TaskIndex = Algo::LowerBound(Thread, StartTime, 
-				[this](TaskTrace::FId TaskId, double Timestamp)
-				{
-					const FTaskInfo* Task = TryGetTask(TaskId);
-					return Task != nullptr && Task->StartedTimestamp <= Timestamp;
-				}
-			);
+		using FSelector = bool(*)(double, double, const FTaskInfo&);
+		FSelector Selectors[(uint32)ETaskEnumerationOption::Count] = 
+		{
+			// Alive
+			[](double StartTime, double EndTime, const FTaskInfo& Task)
+			{	// created before `StartTime` and destroyed after `EndTime`
+				return (Task.CreatedTimestamp <= StartTime || Task.CreatedTimestamp == FTaskInfo::InvalidTimestamp) &&
+					(Task.DestroyedTimestamp >= EndTime || Task.DestroyedTimestamp == FTaskInfo::InvalidTimestamp);
+			},
+			// Active
+			[](double StartTime, double EndTime, const FTaskInfo& Task)
+			{	// completed before `StartTime` and destroyed after `EndTime`
+				return (Task.StartedTimestamp < EndTime && Task.StartedTimestamp != FTaskInfo::InvalidTimestamp) &&
+					(Task.FinishedTimestamp > StartTime || Task.FinishedTimestamp == FTaskInfo::InvalidTimestamp);
+			},
+			// WaitingForPrerequisites
+			[](double StartTime, double EndTime, const FTaskInfo& Task)
+			{	// launched before `StartTime` and scheduled after `EndTime`
+				return (Task.LaunchedTimestamp <= StartTime && Task.LaunchedTimestamp != FTaskInfo::InvalidTimestamp) && 
+					(Task.ScheduledTimestamp >= EndTime || Task.ScheduledTimestamp == FTaskInfo::InvalidTimestamp);
+			},
+			// Queued
+			[](double StartTime, double EndTime, const FTaskInfo& Task)
+			{	// scheduled before `StartTime` and started execution after `EndTime`
+				return (Task.ScheduledTimestamp <= StartTime && Task.ScheduledTimestamp != FTaskInfo::InvalidTimestamp) &&
+					(Task.StartedTimestamp >= EndTime || Task.StartedTimestamp == FTaskInfo::InvalidTimestamp);
+			},
+			// Executing
+			[](double StartTime, double EndTime, const FTaskInfo& Task)
+			{	// started execution before `StartTime` and finished after `EndTime`
+				return (Task.StartedTimestamp <= StartTime && Task.StartedTimestamp != FTaskInfo::InvalidTimestamp) &&
+					(Task.FinishedTimestamp >= EndTime || Task.FinishedTimestamp == FTaskInfo::InvalidTimestamp);
+			},
+			// WaitingForNested
+			[](double StartTime, double EndTime, const FTaskInfo& Task)
+			{	// finished execution before `StartTime` and completed after `EndTime`
+				return (Task.FinishedTimestamp <= StartTime && Task.FinishedTimestamp != FTaskInfo::InvalidTimestamp) &&
+					(Task.CompletedTimestamp >= EndTime || Task.CompletedTimestamp == FTaskInfo::InvalidTimestamp);
+			},
+			// Completed
+			[](double StartTime, double EndTime, const FTaskInfo& Task)
+			{	// completed before `StartTime` and destroyed after `EndTime`
+				return (Task.CompletedTimestamp <= StartTime && Task.CompletedTimestamp != FTaskInfo::InvalidTimestamp) &&
+					(Task.DestroyedTimestamp >= EndTime || Task.DestroyedTimestamp == FTaskInfo::InvalidTimestamp);
+			},
+		};
 
-			// check if there's a previous task whose execution overlaps `StartTime`
-			if (TaskIndex != 0)
+		auto Enumerate = [StartTime, EndTime, Callback = MoveTemp(Callback), this](auto Selector)
+		{
+			for (const TPair<uint32, TArray64<TaskTrace::FId>>& KeyValue : ExecutionThreads)
 			{
-				const FTaskInfo* Task = TryGetTask(Thread[TaskIndex - 1]);
-				if (Task != nullptr && Task->FinishedTimestamp > StartTime)
+				const TArray64<TaskTrace::FId>& TaskIds = KeyValue.Value;
+				for (int64 i = 0; i != TaskIds.Num(); ++i)
 				{
-					--TaskIndex;
+					const FTaskInfo* Task = TryGetTask(TaskIds[i]);
+					check(Task);
+					if (Task != nullptr && Selector(StartTime, EndTime, *Task))
+					{
+						if (Callback(*Task) == ETaskEnumerationResult::Stop)
+						{
+							break;
+						}
+					}
 				}
 			}
+		};
 
-			if (TaskIndex == Thread.Num())
-			{
-				continue; // all tasks on this thread are before StartTime so nothing to do here, go to the next thread
-			}
-
-			// report all tasks whose execution overlaps [StartTime, EndTime]
-			const FTaskInfo* Task = TryGetTask(Thread[TaskIndex]);
-			while (Task != nullptr && Task->StartedTimestamp <= EndTime && Callback(*Task) != ETaskEnumerationResult::Stop && TaskIndex < Thread.Num() - 1)
-			{
-				++TaskIndex;
-				Task = TryGetTask(Thread[TaskIndex]);
-			}
-		}
+		Enumerate(Selectors[(uint32)EnumerationOption]);
 	}
 }
diff --git a/Engine/Source/Developer/TraceServices/Private/Model/TasksProfilerPrivate.h b/Engine/Source/Developer/TraceServices/Private/Model/TasksProfilerPrivate.h
--- a/Engine/Source/Developer/TraceServices/Private/Model/TasksProfilerPrivate.h
+++ b/Engine/Source/Developer/TraceServices/Private/Model/TasksProfilerPrivate.h
@@ -22,14 +22,12 @@ namespace TraceServices
 		void Init(uint32 Version);
 
 		// traces consumption
-		void TaskCreated(TaskTrace::FId TaskId, double Timestamp, uint32 ThreadId);
-		void TaskLaunched(TaskTrace::FId TaskId, const TCHAR* DebugName, bool bTracked, int32 ThreadToExecuteOn, double Timestamp, uint32 ThreadId);
+		void TaskCreated(TaskTrace::FId TaskId, double Timestamp, uint32 ThreadId, uint64 TaskSize);
+		void TaskLaunched(TaskTrace::FId TaskId, const TCHAR* DebugName, bool bTracked, int32 ThreadToExecuteOn, double Timestamp, uint32 ThreadId, uint64 TaskSize);
 		void TaskScheduled(TaskTrace::FId TaskId, double Timestamp, uint32 ThreadId);
-		void TaskScheduled(FTaskInfo& Task, double Timestamp, uint32 ThreadId);
 		void SubsequentAdded(TaskTrace::FId TaskId, TaskTrace::FId SubsequentId, double Timestamp, uint32 ThreadId);
 		void TaskStarted(TaskTrace::FId TaskId, double Timestamp, uint32 ThreadId);
 		void TaskFinished(TaskTrace::FId TaskId, double Timestamp);
-		void TaskFinished(FTaskInfo& Task, double Timestamp);
 		void TaskCompleted(TaskTrace::FId TaskId, double Timestamp, uint32 ThreadId);
 		void TaskDestroyed(TaskTrace::FId TaskId, double Timestamp, uint32 ThreadId);
 
@@ -42,7 +40,7 @@ namespace TraceServices
 		virtual const FWaitingForTasks* TryGetWaiting(const TCHAR* TimerName, uint32 ThreadId, double Timestamp) const override;
 		virtual TArray<TaskTrace::FId> TryGetParallelForTasks(const TCHAR* TimerName, uint32 ThreadId, double StartTime, double EndTime) const override;
 		virtual int64 GetNumTasks() const override;
-		virtual void EnumerateTasks(double StartTime, double EndTime, TaskCallback Callback) const override;
+		virtual void EnumerateTasks(double StartTime, double EndTime, ETaskEnumerationOption EnumerationOption, TaskCallback Callback) const override;
 
 	private:
 		// initializes conversion of task id to its index in the internal container, if it's not initialized already
@@ -59,11 +57,6 @@ namespace TraceServices
 
 		void CreateCounters();
 
-		// registers an event with a name, task id, when and where it happened
-		bool TryRegisterEvent(const TCHAR* EventName, TaskTrace::FId TaskId, double FTaskInfo::* TimestampPtr, double TimestampValue, uint32 FTaskInfo::* ThreadIdPtr = nullptr, uint32 ThreadIdValue = 0);
-		// registers an event for the given task
-		bool TryRegisterEvent(const TCHAR* EventName, FTaskInfo& Task, double FTaskInfo::* TimestampPtr, double TimestampValue, uint32 FTaskInfo::* ThreadIdPtr = nullptr, uint32 ThreadIdValue = 0);
-		// registers a relation with another task, when and where it was established
 		void AddRelative(const TCHAR* RelativeType, TaskTrace::FId TaskId, TArray<FTaskInfo::FRelationInfo> FTaskInfo::* RelationsPtr, TaskTrace::FId RelativeId, double Timestamp, uint32 ThreadId);
 
 	private:
@@ -95,6 +88,9 @@ namespace TraceServices
 		int64 RunningTasksNum = 0;
 		// a bool that is set when the counters are created
 		bool bCountersCreated = false;
+		// number of alive tasks
+		int64 AliveTasksNum = 0;
+		int64 AliveTasksSize = 0;
 
 		IEditableCounterProvider& EditableCounterProvider;
 		IEditableCounter* TaskLatencyCounter;
@@ -103,5 +99,7 @@ namespace TraceServices
 		IEditableCounter* NamedThreadsScheduledTasksCounter;
 		IEditableCounter* RunningTasksCounter;
 		IEditableCounter* ExecutionTimeCounter;
+		IEditableCounter* AliveTasksCounter;
+		IEditableCounter* AliveTasksSizeCounter;
 	};
 } // namespace TraceServices
diff --git a/Engine/Source/Developer/TraceServices/Public/TraceServices/Model/TasksProfiler.h b/Engine/Source/Developer/TraceServices/Public/TraceServices/Model/TasksProfiler.h
--- a/Engine/Source/Developer/TraceServices/Public/TraceServices/Model/TasksProfiler.h
+++ b/Engine/Source/Developer/TraceServices/Public/TraceServices/Model/TasksProfiler.h
@@ -47,6 +47,8 @@ namespace TraceServices
 		// the last reference is released and the task is destroyed
 		double DestroyedTimestamp = InvalidTimestamp;
 		uint32 DestroyedThreadId;
+		// task size including user-provided task body
+		uint64 TaskSize = 0;
 
 		//  relation with another task, when and where it was established
 		struct FRelationInfo
@@ -80,6 +82,18 @@ namespace TraceServices
 		double FinishedTimestamp = FTaskInfo::InvalidTimestamp;
 	};
 
+	enum class ETaskEnumerationOption
+	{
+		Alive,
+		Active,
+		WaitingForPrerequisites,
+		Queued,
+		Executing,
+		WaitingForNested,
+		Completed,
+		Count
+	};
+
 	enum class ETaskEnumerationResult
 	{
 		Continue,
@@ -109,7 +123,7 @@ namespace TraceServices
 		virtual int64 GetNumTasks() const = 0;
 
 		// Calls the callback for each task stored in the provider with CreatedTimestamp <= EndTime and FinishedTimestamp >= StartTime
-		virtual void EnumerateTasks(double StartTime, double EndTime, TaskCallback Callback) const = 0;
+		virtual void EnumerateTasks(double StartTime, double EndTime, ETaskEnumerationOption EnumerationOption, TaskCallback Callback) const = 0;
 	};
 
 	TRACESERVICES_API const ITasksProvider* ReadTasksProvider(const IAnalysisSession& Session);
diff --git a/Engine/Source/Runtime/Core/Private/Async/TaskGraph.cpp b/Engine/Source/Runtime/Core/Private/Async/TaskGraph.cpp
--- a/Engine/Source/Runtime/Core/Private/Async/TaskGraph.cpp
+++ b/Engine/Source/Runtime/Core/Private/Async/TaskGraph.cpp
@@ -2401,7 +2401,12 @@ bool FTaskGraphInterface::IsMultithread()
 
 #if TASKGRAPH_NEW_FRONTEND
 
-FGraphEventImplAllocator GraphEventImplAllocator;
+FGraphEventImplAllocator& GetGraphEventImplAllocator()
+{
+	static FGraphEventImplAllocator Singleton;
+	return Singleton;
+}
+
 FGraphTaskAllocator SmallTaskAllocator;
 
 #else
@@ -2477,10 +2482,10 @@ void FGraphEvent::DispatchSubsequents(TArray<FBaseGraphTask*>& NewTasks, ENamedT
 				}
 			}
 
-#if UE_TASK_TRACE_ENABLED
-			// regenerate TraceId as we're going to use the same event with another task
-			TraceId = TaskTrace::GenerateTaskId();
-#endif
+//#if UE_TASK_TRACE_ENABLED
+//			// regenerate TraceId as we're going to use the same event with another task
+//			TraceId = TaskTrace::GenerateTaskId();
+//#endif
 			TGraphTask<FNullGraphTask>::CreateTask(FGraphEventRef(this), &TempEventsToWaitFor, CurrentThreadIfKnown).ConstructAndDispatchWhenReady(GET_STATID(STAT_FNullGraphTask_DontCompleteUntil), LocalThreadToDoGatherOn);
 			return;
 		}
diff --git a/Engine/Source/Runtime/Core/Private/Async/TaskTrace.cpp b/Engine/Source/Runtime/Core/Private/Async/TaskTrace.cpp
--- a/Engine/Source/Runtime/Core/Private/Async/TaskTrace.cpp
+++ b/Engine/Source/Runtime/Core/Private/Async/TaskTrace.cpp
@@ -20,6 +20,7 @@ namespace TaskTrace
 	UE_TRACE_EVENT_BEGIN(TaskTrace, Created)
 		UE_TRACE_EVENT_FIELD(uint64, Timestamp)
 		UE_TRACE_EVENT_FIELD(uint64, TaskId)
+		UE_TRACE_EVENT_FIELD(uint64, TaskSize)
 	UE_TRACE_EVENT_END()
 
 	UE_TRACE_EVENT_BEGIN(TaskTrace, Launched)
@@ -28,6 +29,7 @@ namespace TaskTrace
 		UE_TRACE_EVENT_FIELD(UE::Trace::WideString, DebugName)
 		UE_TRACE_EVENT_FIELD(bool, Tracked)
 		UE_TRACE_EVENT_FIELD(int32, ThreadToExecuteOn)
+		UE_TRACE_EVENT_FIELD(uint64, TaskSize)
 	UE_TRACE_EVENT_END()
 
 	UE_TRACE_EVENT_BEGIN(TaskTrace, Scheduled)
@@ -92,7 +94,7 @@ namespace TaskTrace
 #endif
 	}
 
-	void Created(FId TaskId)
+	void Created(FId TaskId, uint64 TaskSize)
 	{
 		if (!bGTaskTraceInitialized)
 		{
@@ -103,10 +105,11 @@ namespace TaskTrace
 
 		UE_TRACE_LOG(TaskTrace, Created, TaskChannel)
 			<< Created.Timestamp(FPlatformTime::Cycles64())
-			<< Created.TaskId(TaskId);
+			<< Created.TaskId(TaskId)
+			<< Created.TaskSize(TaskSize);
 	}
 
-	void Launched(FId TaskId, const TCHAR* DebugName, bool bTracked, ENamedThreads::Type ThreadToExecuteOn)
+	void Launched(FId TaskId, const TCHAR* DebugName, bool bTracked, ENamedThreads::Type ThreadToExecuteOn, uint64 TaskSize)
 	{
 		if (!bGTaskTraceInitialized)
 		{
@@ -120,7 +123,8 @@ namespace TaskTrace
 			<< Launched.TaskId(TaskId)
 			<< Launched.DebugName(DebugName != nullptr ? DebugName : TEXT(""))
 			<< Launched.Tracked(bTracked)
-			<< Launched.ThreadToExecuteOn(ThreadToExecuteOn);
+			<< Launched.ThreadToExecuteOn(ThreadToExecuteOn)
+			<< Launched.TaskSize(TaskSize);
 	}
 
 	void Scheduled(FId TaskId)
diff --git a/Engine/Source/Runtime/Core/Private/Tasks/TaskPrivate.cpp b/Engine/Source/Runtime/Core/Private/Tasks/TaskPrivate.cpp
--- a/Engine/Source/Runtime/Core/Private/Tasks/TaskPrivate.cpp
+++ b/Engine/Source/Runtime/Core/Private/Tasks/TaskPrivate.cpp
@@ -216,7 +216,7 @@ namespace UE::Tasks
 			FWaitingTask WaitingTask{ TEXT("Waiting Task"), MoveTemp(WaitingTaskBody), ETaskPriority::Default /* doesn't matter*/, EExtendedTaskPriority::Inline };
 			WaitingTask.AddPrerequisites(*this);
 
-			if (WaitingTask.TryLaunch())
+			if (WaitingTask.TryLaunch(sizeof(WaitingTask)))
 			{	// was executed inline
 				check(WaitingTask.IsCompleted());
 			}
@@ -259,7 +259,7 @@ namespace UE::Tasks
 			TRefCountPtr<FWaitingTask> WaitingTask{ FWaitingTask::Create(TEXT("Waiting Task"), MoveTemp(WaitingTaskBody), ETaskPriority::Default /* doesn't matter*/, EExtendedTaskPriority::Inline), /*bAddRef=*/ false };
 			WaitingTask->AddPrerequisites(*this);
 
-			if (WaitingTask->TryLaunch())
+			if (WaitingTask->TryLaunch(sizeof(WaitingTask)))
 			{	// was executed inline
 				check(WaitingTask->IsCompleted());
 				return true;
@@ -319,7 +319,7 @@ namespace UE::Tasks
 				using FReturnFromNamedThreadTask = TExecutableTask<decltype(TaskBody)>;
 				FReturnFromNamedThreadTask ReturnTask{ TEXT("ReturnFromNamedThreadTask"), MoveTemp(TaskBody), ETaskPriority::High, ExtendedPriority };
 				ReturnTask.AddPrerequisites(Task);
-				ReturnTask.TryLaunch(); // the result doesn't matter
+				ReturnTask.TryLaunch(sizeof(ReturnTask)); // the result doesn't matter
 
 				TaskGraph.ProcessThreadUntilRequestReturn(CurrentThread);
 				return true;
diff --git a/Engine/Source/Runtime/Core/Public/Async/ParallelFor.h b/Engine/Source/Runtime/Core/Public/Async/ParallelFor.h
--- a/Engine/Source/Runtime/Core/Public/Async/ParallelFor.h
+++ b/Engine/Source/Runtime/Core/Public/Async/ParallelFor.h
@@ -379,7 +379,7 @@ namespace ParallelForImpl
 					TaskTrace::Destroyed(TracedTask.TraceId);
 				}
 				TracedTask.TraceId = TaskTrace::GenerateTaskId();
-				TaskTrace::Launched(TracedTask.TraceId, DebugName, false, ENamedThreads::AnyThread);
+				TaskTrace::Launched(TracedTask.TraceId, DebugName, false, ENamedThreads::AnyThread, 0);
 
 				TracedTask.Task.Init(DebugName, InPriority, FParallelExecutor(MoveTemp(InData), InWorkerIndex, InPriority));
 				verify(LowLevelTasks::TryLaunch(TracedTask.Task, LowLevelTasks::EQueuePreference::GlobalQueuePreference));
diff --git a/Engine/Source/Runtime/Core/Public/Async/TaskGraphInterfaces.h b/Engine/Source/Runtime/Core/Public/Async/TaskGraphInterfaces.h
--- a/Engine/Source/Runtime/Core/Public/Async/TaskGraphInterfaces.h
+++ b/Engine/Source/Runtime/Core/Public/Async/TaskGraphInterfaces.h
@@ -483,7 +483,7 @@ public:
 
 	void Unlock(ENamedThreads::Type CurrentThreadIfKnown = ENamedThreads::AnyThread)
 	{
-		TryLaunch();
+		TryLaunch(0);
 	}
 
 	void Execute(TArray<FBaseGraphTask*>& NewTasks, ENamedThreads::Type CurrentThread, bool bDeleteOnCompletion)
@@ -529,7 +529,7 @@ public:
 	{
 		//check(CurrentThreadIfKnown == ENamedThreads::AnyThread); // the feature is not used
 		AddRef(); // scheduler's reference
-		TryLaunch();
+		TryLaunch(0);
 	}
 
 	void DispatchSubsequents(TArray<FBaseGraphTask*>& NewTasks, ENamedThreads::Type CurrentThreadIfKnown = ENamedThreads::AnyThread)
@@ -646,7 +646,7 @@ public:
 		FORCEINLINE_DEBUGGABLE FGraphEventRef ConstructAndDispatchWhenReady(T&&... Args)
 		{
 			FGraphEventRef Ref{ ConstructAndHoldImpl(Forward<T>(Args)...) };
-			Ref->TryLaunch();
+			Ref->TryLaunch(sizeof(TGraphTask));
 			return Ref;
 		}
 
@@ -655,7 +655,7 @@ public:
 		FORCEINLINE_DEBUGGABLE TGraphTask* ConstructAndHold(T&&... Args)
 		{
 			TGraphTask* Task = ConstructAndHoldImpl(Forward<T>(Args)...);
-			TaskTrace::Created(Task->GetTraceId());
+			TaskTrace::Created(Task->GetTraceId(), sizeof(*Task));
 			return Task;
 		}
 
@@ -752,7 +752,8 @@ public:
 	FGraphEventImpl()
 		: FBaseGraphTask(nullptr)
 	{
-		Init(TEXT("GraphEvent"), UE::Tasks::ETaskPriority::Normal, UE::Tasks::EExtendedTaskPriority::TaskEvent);
+		TaskTrace::Created(GetTraceId(), sizeof(*this));
+		Init((TEXT("GraphEvent"), UE::Tasks::ETaskPriority::Normal, UE::Tasks::EExtendedTaskPriority::TaskEvent);
 	}
 
 	static void* operator new(size_t Size);
@@ -767,16 +768,16 @@ private:
 };
 
 using FGraphEventImplAllocator = TLockFreeFixedSizeAllocator_TLSCache<sizeof(FGraphEventImpl), PLATFORM_CACHE_LINE_SIZE>;
-CORE_API extern FGraphEventImplAllocator GraphEventImplAllocator;
+CORE_API FGraphEventImplAllocator& GetGraphEventImplAllocator();
 
 inline void* FGraphEventImpl::operator new(size_t Size)
 {
-	return GraphEventImplAllocator.Allocate();
+	return GetGraphEventImplAllocator().Allocate();
 }
 
 inline void FGraphEventImpl::operator delete(void* Ptr)
 {
-	GraphEventImplAllocator.Free(Ptr);
+	GetGraphEventImplAllocator().Free(Ptr);
 }
 
 inline FGraphEventRef FBaseGraphTask::CreateGraphEvent()
@@ -1280,7 +1281,7 @@ public:
 
 	void Unlock(ENamedThreads::Type CurrentThreadIfKnown = ENamedThreads::AnyThread)
 	{
-		TaskTrace::Launched(GetTraceId(), nullptr, Subsequents.IsValid(), ((TTask*)&TaskStorage)->GetDesiredThread());
+		TaskTrace::Launched(GetTraceId(), nullptr, Subsequents.IsValid(), ((TTask*)&TaskStorage)->GetDesiredThread(), sizeof(*this));
 
 		bool bWakeUpWorker = true;
 		ConditionalQueueTask(CurrentThreadIfKnown, bWakeUpWorker);
@@ -1415,7 +1416,7 @@ private:
 	 **/
 	FGraphEventRef Setup(const FGraphEventArray* Prerequisites = NULL, ENamedThreads::Type CurrentThreadIfKnown = ENamedThreads::AnyThread)
 	{
-		TaskTrace::Launched(GetTraceId(), nullptr, Subsequents.IsValid(), ((TTask*)&TaskStorage)->GetDesiredThread());
+		TaskTrace::Launched(GetTraceId(), nullptr, Subsequents.IsValid(), ((TTask*)&TaskStorage)->GetDesiredThread(), sizeof(*this));
 
 		FGraphEventRef ReturnedEventRef = Subsequents; // very important so that this doesn't get destroyed before we return
 		SetupPrereqs(Prerequisites, CurrentThreadIfKnown, true);
@@ -1435,7 +1436,7 @@ private:
 	 **/
 	TGraphTask* Hold(const FGraphEventArray* Prerequisites = NULL, ENamedThreads::Type CurrentThreadIfKnown = ENamedThreads::AnyThread)
 	{
-		TaskTrace::Created(GetTraceId());
+		TaskTrace::Created(GetTraceId(), sizeof(*this));
 
 		SetupPrereqs(Prerequisites, CurrentThreadIfKnown, false);
 		return this;
diff --git a/Engine/Source/Runtime/Core/Public/Async/TaskTrace.h b/Engine/Source/Runtime/Core/Public/Async/TaskTrace.h
--- a/Engine/Source/Runtime/Core/Public/Async/TaskTrace.h
+++ b/Engine/Source/Runtime/Core/Public/Async/TaskTrace.h
@@ -37,8 +37,8 @@ namespace TaskTrace
 	FId CORE_API GenerateTaskId();
 
 	void CORE_API Init();
-	void CORE_API Created(FId TaskId); // optional, used only if a task was created but not launched immediately
-	void CORE_API Launched(FId TaskId, const TCHAR* DebugName, bool bTracked, ENamedThreads::Type ThreadToExecuteOn);
+	void CORE_API Created(FId TaskId, uint64 TaskSize); // optional, used only if a task was created but not launched immediately
+	void CORE_API Launched(FId TaskId, const TCHAR* DebugName, bool bTracked, ENamedThreads::Type ThreadToExecuteOn, uint64 TaskSize);
 	void CORE_API Scheduled(FId TaskId);
 	void CORE_API SubsequentAdded(FId TaskId, FId SubsequentId);
 	void CORE_API Started(FId TaskId);
@@ -67,8 +67,8 @@ namespace TaskTrace
 	// NOOP implementation
 	inline FId GenerateTaskId() { return InvalidId; }
 	inline void Init() {}
-	inline void Created(FId TaskId) {}
-	inline void Launched(FId TaskId, const TCHAR* DebugName, bool bTracked, ENamedThreads::Type ThreadToExecuteOn) {}
+	inline void Created(FId TaskId, uint64 TaskSize) {}
+	inline void Launched(FId TaskId, const TCHAR* DebugName, bool bTracked, ENamedThreads::Type ThreadToExecuteOn, uint64 TaskSize) {}
 	inline void Scheduled(FId TaskId) {}
 	inline void SubsequentAdded(FId TaskId, FId SubsequentId) {}
 	inline void Started(FId TaskId) {}
diff --git a/Engine/Source/Runtime/Core/Public/Tasks/Pipe.h b/Engine/Source/Runtime/Core/Public/Tasks/Pipe.h
--- a/Engine/Source/Runtime/Core/Public/Tasks/Pipe.h
+++ b/Engine/Source/Runtime/Core/Public/Tasks/Pipe.h
@@ -72,7 +72,7 @@ namespace UE::Tasks
 
 			FExecutableTask* Task = FExecutableTask::Create(InDebugName, Forward<TaskBodyType>(TaskBody), Priority, ExtendedPriority);
 			Task->SetPipe(*this);
-			Task->TryLaunch();
+			Task->TryLaunch(sizeof(*Task));
 			return TTask<FResult>{ Task };
 		}
 
@@ -98,7 +98,7 @@ namespace UE::Tasks
 			FExecutableTask* Task = FExecutableTask::Create(InDebugName, Forward<TaskBodyType>(TaskBody), Priority, ExtendedPriority);
 			Task->SetPipe(*this);
 			Task->AddPrerequisites(Forward<PrerequisitesCollectionType>(Prerequisites));
-			Task->TryLaunch();
+			Task->TryLaunch(sizeof(*Task));
 			return TTask<FResult>{ Task };
 		}
 
diff --git a/Engine/Source/Runtime/Core/Public/Tasks/Task.h b/Engine/Source/Runtime/Core/Public/Tasks/Task.h
--- a/Engine/Source/Runtime/Core/Public/Tasks/Task.h
+++ b/Engine/Source/Runtime/Core/Public/Tasks/Task.h
@@ -146,7 +146,7 @@ namespace UE::Tasks
 				FExecutableTask* Task = FExecutableTask::Create(DebugName, Forward<TaskBodyType>(TaskBody), Priority, ExtendedPriority);
 				// this must happen before launching, to support an ability to access the task itself from inside it
 				*Pimpl.GetInitReference() = Task;
-				Task->TryLaunch();
+				Task->TryLaunch(sizeof(*Task));
 			}
 
 			// launches a task for asynchronous execution, with prerequisites that must be completed before the task is scheduled
@@ -172,7 +172,7 @@ namespace UE::Tasks
 				Task->AddPrerequisites(Forward<PrerequisitesCollectionType>(Prerequisites));
 				// this must happen before launching, to support an ability to access the task itself from inside it
 				*Pimpl.GetInitReference() = Task;
-				Task->TryLaunch();
+				Task->TryLaunch(sizeof(*Task));
 			}
 
 			bool IsAwaitable() const
@@ -255,7 +255,7 @@ namespace UE::Tasks
 				// An event is not "in the system" until it's triggered, and should be kept alive only by external references. Once it's triggered it's in the system 
 				// and can outlive external references, so we need to keep it alive by holding an internal reference. It will be released when the event is signalled
 				Pimpl->AddRef();
-				Pimpl->TryLaunch();
+				Pimpl->TryLaunch(sizeof(*Pimpl));
 			}
 		}
 	};
@@ -355,7 +355,7 @@ namespace UE::Tasks
 		FWaitingTask WaitingTask{ TEXT("Waiting Task"), MoveTemp(WaitingTaskBody), ETaskPriority::Default /* doesn't matter */,  EExtendedTaskPriority::Inline };
 		WaitingTask.AddPrerequisites(Tasks);
 
-		if (WaitingTask.TryLaunch())
+		if (WaitingTask.TryLaunch(sizeof(WaitingTask)))
 		{	// was executed inline
 			check(WaitingTask.IsCompleted());
 		}
@@ -395,7 +395,7 @@ namespace UE::Tasks
 		TRefCountPtr<FWaitingTask> WaitingTask{ FWaitingTask::Create(TEXT("Waiting Task"), MoveTemp(WaitingTaskBody), ETaskPriority::Default /* doesn't matter */, EExtendedTaskPriority::Inline), /*bAddRef=*/ false};
 		WaitingTask->AddPrerequisites(Tasks);
 
-		if (WaitingTask->TryLaunch())
+		if (WaitingTask->TryLaunch(sizeof(WaitingTask)))
 		{	// was executed inline
 			check(WaitingTask->IsCompleted());
 			return true;
diff --git a/Engine/Source/Runtime/Core/Public/Tasks/TaskPrivate.h b/Engine/Source/Runtime/Core/Public/Tasks/TaskPrivate.h
--- a/Engine/Source/Runtime/Core/Public/Tasks/TaskPrivate.h
+++ b/Engine/Source/Runtime/Core/Public/Tasks/TaskPrivate.h
@@ -301,9 +301,9 @@ namespace UE::Tasks
 
 			// Tries to schedule task execution. Returns false if the task has incomplete dependencies (prerequisites or is blocked by a pipe). 
 			// In this case the task will be automatically scheduled when all dependencies are completed.
-			bool TryLaunch()
+			bool TryLaunch(uint64 TaskSize)
 			{
-				TaskTrace::Launched(GetTraceId(), LowLevelTask.GetDebugName(), true, (ENamedThreads::Type)0xff);
+				TaskTrace::Launched(GetTraceId(), LowLevelTask.GetDebugName(), true, (ENamedThreads::Type)0xff, TaskSize);
 				return TryUnlock();
 			}
 
@@ -792,6 +792,7 @@ namespace UE::Tasks
 			FTaskEventBase(const TCHAR* InDebugName)
 				: FTaskBase(/*InitRefCount=*/ 1) // for the initial reference (we don't increment it on passing to `TRefCountPtr`)
 			{
+				TaskTrace::Created(GetTraceId(), sizeof(*this));
 				Init(InDebugName, ETaskPriority::Normal, EExtendedTaskPriority::TaskEvent);
 			}
 
