[Insights]
  - TraceServices: Simplified the FProviderLock API (deprecated the old API).
  - TraceServices: Updated FCookProfilerProvider to use the simplified FProviderLock API.
  - TraceServices: Replaced FMetadataProviderLock and FAllocationsProviderLock with the generic FProviderLock.

diff --git a/Engine/Source/Developer/TraceServices/Private/Common/ProviderLock.cpp b/Engine/Source/Developer/TraceServices/Private/Common/ProviderLock.cpp
--- a/Engine/Source/Developer/TraceServices/Private/Common/ProviderLock.cpp
+++ b/Engine/Source/Developer/TraceServices/Private/Common/ProviderLock.cpp
@@ -5,56 +5,115 @@
 namespace TraceServices
 {
 
-void FProviderLock::ReadAccessCheck(const FProviderLock* CurrentProviderLock, const int32& CurrentReadProviderLockCount, const int32& CurrentWriteProviderLockCount) const
+void FProviderLock::ReadAccessCheck(FProviderLock::FThreadLocalState& State) const
 {
-	checkf(CurrentProviderLock == this && (CurrentReadProviderLockCount > 0 || CurrentWriteProviderLockCount > 0),
+	checkf(State.Lock == this && (State.ReadLockCount > 0 || State.WriteLockCount > 0),
 		TEXT("Trying to READ from provider outside of a READ scope"));
 }
 
-void FProviderLock::WriteAccessCheck(const int32& CurrentWriteAllocationsProviderLockCount) const
+void FProviderLock::WriteAccessCheck(FProviderLock::FThreadLocalState& State) const
 {
-	checkf(CurrentWriteAllocationsProviderLockCount > 0,
+	checkf(State.Lock == this && (State.WriteLockCount > 0),
 		TEXT("Trying to WRITE to provider outside of an EDIT/WRITE scope"));
 }
 
-void FProviderLock::BeginRead(FProviderLock*& CurrentProviderLock, int32& CurrentReadProviderLockCount, const int32& WriteProviderLockCount)
+void FProviderLock::BeginRead(FProviderLock::FThreadLocalState& State)
 {
+	check(!State.Lock || State.Lock == this);
+	checkf(State.WriteLockCount == 0, TEXT("Trying to lock provider for READ while holding EDIT/WRITE access"));
+	if (State.ReadLockCount++ == 0)
+	{
+		State.Lock = this;
+		RWLock.ReadLock();
+	}
+}
+
+void FProviderLock::EndRead(FProviderLock::FThreadLocalState& State)
+{
+	check(State.Lock == this);
+	check(State.ReadLockCount > 0);
+	if (--State.ReadLockCount == 0)
+	{
+		RWLock.ReadUnlock();
+		State.Lock = nullptr;
+	}
+}
+
+void FProviderLock::BeginWrite(FProviderLock::FThreadLocalState& State)
+{
+	check(!State.Lock || State.Lock == this);
+	checkf(State.ReadLockCount == 0, TEXT("Trying to lock provider for EDIT/WRITE while holding READ access"));
+	if (State.WriteLockCount++ == 0)
+	{
+		State.Lock = this;
+		RWLock.WriteLock();
+	}
+}
+
+void FProviderLock::EndWrite(FProviderLock::FThreadLocalState& State)
+{
+	check(State.Lock == this);
+	check(State.WriteLockCount > 0);
+	if (--State.WriteLockCount == 0)
+	{
+		RWLock.WriteUnlock();
+		State.Lock = nullptr;
+	}
+}
+
+void FProviderLock::ReadAccessCheck(const FProviderLock* ProviderLock, const int32& ReadProviderLockCount, const int32& WriteProviderLockCount) const
+{
+	checkf(ProviderLock == this && (ReadProviderLockCount > 0 || WriteProviderLockCount > 0),
+		TEXT("Trying to READ from provider outside of a READ scope"));
+}
+
+void FProviderLock::WriteAccessCheck(const int32& WriteProviderLockCount) const
+{
+	checkf(WriteProviderLockCount > 0,
+		TEXT("Trying to WRITE to provider outside of an EDIT/WRITE scope"));
+}
+
+void FProviderLock::BeginRead(FProviderLock*& ProviderLock, int32& ReadProviderLockCount, const int32& WriteProviderLockCount)
+{
+	check(!ProviderLock || ProviderLock == this);
 	checkf(WriteProviderLockCount == 0, TEXT("Trying to lock provider for READ while holding EDIT/WRITE access"));
-	if (CurrentReadProviderLockCount++ == 0)
+	if (ReadProviderLockCount++ == 0)
 	{
-		CurrentProviderLock = this;
+		ProviderLock = this;
 		RWLock.ReadLock();
 	}
 }
 
-void FProviderLock::EndRead(FProviderLock*& CurrentProviderLock, int32& CurrentReadProviderLockCount)
+void FProviderLock::EndRead(FProviderLock*& ProviderLock, int32& ReadProviderLockCount)
 {
-	check(CurrentReadProviderLockCount > 0);
-	if (--CurrentReadProviderLockCount == 0)
+	check(ProviderLock == this);
+	check(ReadProviderLockCount > 0);
+	if (--ReadProviderLockCount == 0)
 	{
 		RWLock.ReadUnlock();
-		CurrentProviderLock = nullptr;
+		ProviderLock = nullptr;
 	}
 }
 
-void FProviderLock::BeginWrite(FProviderLock*& CurrentProviderLock, const int32& CurrentReadProviderLockCount, int32& WriteProviderLockCount)
+void FProviderLock::BeginWrite(FProviderLock*& ProviderLock, const int32& ReadProviderLockCount, int32& WriteProviderLockCount)
 {
-	check(!CurrentProviderLock || CurrentProviderLock == this);
-	checkf(CurrentReadProviderLockCount == 0, TEXT("Trying to lock provider for EDIT/WRITE while holding READ access"));
+	check(!ProviderLock || ProviderLock == this);
+	checkf(ReadProviderLockCount == 0, TEXT("Trying to lock provider for EDIT/WRITE while holding READ access"));
 	if (WriteProviderLockCount++ == 0)
 	{
-		CurrentProviderLock = this;
+		ProviderLock = this;
 		RWLock.WriteLock();
 	}
 }
 
-void FProviderLock::EndWrite(FProviderLock*& CurrentProviderLock, int32& WriteProviderLockCount)
+void FProviderLock::EndWrite(FProviderLock*& ProviderLock, int32& WriteProviderLockCount)
 {
+	check(ProviderLock == this);
 	check(WriteProviderLockCount > 0);
 	if (--WriteProviderLockCount == 0)
 	{
 		RWLock.WriteUnlock();
-		CurrentProviderLock = nullptr;
+		ProviderLock = nullptr;
 	}
 }
 
diff --git a/Engine/Source/Developer/TraceServices/Private/Model/AllocationsProvider.cpp b/Engine/Source/Developer/TraceServices/Private/Model/AllocationsProvider.cpp
--- a/Engine/Source/Developer/TraceServices/Private/Model/AllocationsProvider.cpp
+++ b/Engine/Source/Developer/TraceServices/Private/Model/AllocationsProvider.cpp
@@ -126,67 +126,7 @@ static bool IsAddressWatched(uint64 InAddress)
 }
 #endif // INSIGHTS_DEBUG_WATCH
 
-////////////////////////////////////////////////////////////////////////////////////////////////////
-// FAllocationsProviderLock
-////////////////////////////////////////////////////////////////////////////////////////////////////
-
-thread_local FAllocationsProviderLock* GThreadCurrentAllocationsProviderLock;
-thread_local int32 GThreadCurrentReadAllocationsProviderLockCount;
-thread_local int32 GThreadCurrentWriteAllocationsProviderLockCount;
-
-void FAllocationsProviderLock::ReadAccessCheck() const
-{
-	checkf(GThreadCurrentAllocationsProviderLock == this && (GThreadCurrentReadAllocationsProviderLockCount > 0 || GThreadCurrentWriteAllocationsProviderLockCount > 0),
-		TEXT("Trying to READ from allocations provider outside of a READ scope"));
-}
-
-void FAllocationsProviderLock::WriteAccessCheck() const
-{
-	checkf(GThreadCurrentAllocationsProviderLock == this && GThreadCurrentWriteAllocationsProviderLockCount > 0,
-		TEXT("Trying to WRITE to allocations provider outside of an EDIT/WRITE scope"));
-}
-
-void FAllocationsProviderLock::BeginRead()
-{
-	check(!GThreadCurrentAllocationsProviderLock || GThreadCurrentAllocationsProviderLock == this);
-	checkf(GThreadCurrentWriteAllocationsProviderLockCount == 0, TEXT("Trying to lock allocations provider for READ while holding EDIT/WRITE access"));
-	if (GThreadCurrentReadAllocationsProviderLockCount++ == 0)
-	{
-		GThreadCurrentAllocationsProviderLock = this;
-		RWLock.ReadLock();
-	}
-}
-
-void FAllocationsProviderLock::EndRead()
-{
-	check(GThreadCurrentReadAllocationsProviderLockCount > 0);
-	if (--GThreadCurrentReadAllocationsProviderLockCount == 0)
-	{
-		RWLock.ReadUnlock();
-		GThreadCurrentAllocationsProviderLock = nullptr;
-	}
-}
-
-void FAllocationsProviderLock::BeginWrite()
-{
-	check(!GThreadCurrentAllocationsProviderLock || GThreadCurrentAllocationsProviderLock == this);
-	checkf(GThreadCurrentReadAllocationsProviderLockCount == 0, TEXT("Trying to lock allocations provider for EDIT/WRITE while holding READ access"));
-	if (GThreadCurrentWriteAllocationsProviderLockCount++ == 0)
-	{
-		GThreadCurrentAllocationsProviderLock = this;
-		RWLock.WriteLock();
-	}
-}
-
-void FAllocationsProviderLock::EndWrite()
-{
-	check(GThreadCurrentWriteAllocationsProviderLockCount > 0);
-	if (--GThreadCurrentWriteAllocationsProviderLockCount == 0)
-	{
-		RWLock.WriteUnlock();
-		GThreadCurrentAllocationsProviderLock = nullptr;
-	}
-}
+thread_local FProviderLock::FThreadLocalState GAllocationsProviderLockState;
 
 ////////////////////////////////////////////////////////////////////////////////////////////////////
 // FTagTracker
@@ -1297,7 +1237,7 @@ FAllocationsProvider::~FAllocationsProvider()
 
 void FAllocationsProvider::EditInit(double InTime, uint8 InMinAlignment)
 {
-	Lock.WriteAccessCheck();
+	EditAccessCheck();
 
 	if (bInitialized)
 	{
@@ -1322,7 +1262,7 @@ void FAllocationsProvider::EditInit(double InTime, uint8 InMinAlignment)
 
 void FAllocationsProvider::EditAlloc(double Time, uint32 CallstackId, uint64 Address, uint64 InSize, uint32 InAlignment, HeapId RootHeap)
 {
-	Lock.WriteAccessCheck();
+	EditAccessCheck();
 
 	if (!bInitialized)
 	{
@@ -1464,7 +1404,7 @@ void FAllocationsProvider::EditAlloc(double Time, uint32 CallstackId, uint64 Add
 
 void FAllocationsProvider::EditFree(double Time, uint32 CallstackId, uint64 Address, HeapId RootHeap)
 {
-	Lock.WriteAccessCheck();
+	EditAccessCheck();
 
 	if (!bInitialized)
 	{
@@ -1605,7 +1545,7 @@ void FAllocationsProvider::EditFree(double Time, uint32 CallstackId, uint64 Addr
 
 void FAllocationsProvider::EditHeapSpec(HeapId Id, HeapId ParentId, const FStringView& Name, EMemoryTraceHeapFlags Flags)
 {
-	Lock.WriteAccessCheck();
+	EditAccessCheck();
 
 	INSIGHTS_API_LOGF(TEXT("HeapSpec : Id=%u ParentId=%u Name=\"%*s\" Flags=0x%X"), -1.0, Id, ParentId, Name.Len(), Name.GetData(), uint32(Flags));
 
@@ -1663,7 +1603,7 @@ void FAllocationsProvider::AddHeapSpec(HeapId Id, HeapId ParentId, const FString
 
 void FAllocationsProvider::EditMarkAllocationAsHeap(double Time, uint32 CallstackId, uint64 Address, HeapId Heap, EMemoryTraceHeapAllocationFlags Flags)
 {
-	Lock.WriteAccessCheck();
+	EditAccessCheck();
 
 	INSIGHTS_WATCH_API_LOGF(Address, TEXT("MarkAllocAsHeap 0x%llX : Heap=%u Flags=0x%X CallstackId=%u"), Time, Address, Heap, uint32(Flags), CallstackId);
 
@@ -1736,7 +1676,7 @@ void FAllocationsProvider::EditMarkAllocationAsHeap(double Time, uint32 Callstac
 
 void FAllocationsProvider::EditUnmarkAllocationAsHeap(double Time, uint32 CallstackId, uint64 Address, HeapId Heap)
 {
-	Lock.WriteAccessCheck();
+	EditAccessCheck();
 
 	INSIGHTS_WATCH_API_LOGF(Address, TEXT("UnmarkAllocAsHeap 0x%llX : Heap=%u CallstackId=%u"), Time, Address, Heap, CallstackId);
 
@@ -2020,7 +1960,7 @@ void FAllocationsProvider::EditPopTagFromPtr(uint32 ThreadId, uint8 Tracker)
 
 void FAllocationsProvider::EditOnAnalysisCompleted(double Time)
 {
-	Lock.WriteAccessCheck();
+	EditAccessCheck();
 
 	if (!bInitialized)
 	{
diff --git a/Engine/Source/Developer/TraceServices/Private/Model/AllocationsProvider.h b/Engine/Source/Developer/TraceServices/Private/Model/AllocationsProvider.h
--- a/Engine/Source/Developer/TraceServices/Private/Model/AllocationsProvider.h
+++ b/Engine/Source/Developer/TraceServices/Private/Model/AllocationsProvider.h
@@ -3,6 +3,7 @@
 #pragma once
 
 #include "Common/PagedArray.h"
+#include "Common/ProviderLock.h"
 #include "Containers/Map.h"
 #include "HAL/CriticalSection.h"
 #include "ProfilingDebugging/MemoryTrace.h" // for EMemoryTraceHeapFlags and EMemoryTraceHeapAllocationFlags
@@ -19,23 +20,7 @@ class ILinearAllocator;
 class FMetadataProvider;
 class FSbTree;
 
-////////////////////////////////////////////////////////////////////////////////////////////////////
-
-class FAllocationsProviderLock
-{
-public:
-	void ReadAccessCheck() const;
-	void WriteAccessCheck() const;
-
-	void BeginRead();
-	void EndRead();
-
-	void BeginWrite();
-	void EndWrite();
-
-private:
-	FRWLock RWLock;
-};
+extern thread_local FProviderLock::FThreadLocalState GAllocationsProviderLockState;
 
 ////////////////////////////////////////////////////////////////////////////////////////////////////
 
@@ -285,13 +270,12 @@ public:
 	explicit FAllocationsProvider(IAnalysisSession& InSession, FMetadataProvider& InMetadataProvider);
 	virtual ~FAllocationsProvider();
 
-	virtual void BeginEdit() const override { Lock.BeginWrite(); }
-	virtual void EndEdit() const override { Lock.EndWrite(); }
-	void EditAccessCheck() const { return Lock.WriteAccessCheck(); }
-
-	virtual void BeginRead() const override { Lock.BeginRead(); }
-	virtual void EndRead() const override { Lock.EndRead(); }
-	void ReadAccessCheck() const { return Lock.ReadAccessCheck(); }
+	virtual void BeginEdit() const override       { Lock.BeginWrite(GAllocationsProviderLockState); }
+	virtual void EndEdit() const override         { Lock.EndWrite(GAllocationsProviderLockState); }
+	virtual void EditAccessCheck() const override { Lock.WriteAccessCheck(GAllocationsProviderLockState); }
+	virtual void BeginRead() const override       { Lock.BeginRead(GAllocationsProviderLockState); }
+	virtual void EndRead() const override         { Lock.EndRead(GAllocationsProviderLockState); }
+	virtual void ReadAccessCheck() const override { Lock.ReadAccessCheck(GAllocationsProviderLockState); }
 
 	//////////////////////////////////////////////////
 	// Read operations
@@ -360,11 +344,11 @@ private:
 	HeapId FindRootHeap(HeapId Heap) const;
 
 private:
+	mutable FProviderLock Lock;
+
 	IAnalysisSession& Session;
 	FMetadataProvider& MetadataProvider;
 
-	mutable FAllocationsProviderLock Lock;
-
 	// Number of supported root heaps
 	constexpr static uint8 MaxRootHeaps = 16;
 
diff --git a/Engine/Source/Developer/TraceServices/Private/Model/CookProfilerProvider.cpp b/Engine/Source/Developer/TraceServices/Private/Model/CookProfilerProvider.cpp
--- a/Engine/Source/Developer/TraceServices/Private/Model/CookProfilerProvider.cpp
+++ b/Engine/Source/Developer/TraceServices/Private/Model/CookProfilerProvider.cpp
@@ -6,9 +6,8 @@
 
 namespace TraceServices
 {
-thread_local FProviderLock* GThreadCurrentCookProviderLock;
-thread_local int32 GThreadCurrentReadCookProviderLockCount;
-thread_local int32 GThreadCurrentWriteCookProviderLockCount;
+
+thread_local FProviderLock::FThreadLocalState GCookProviderLockState;
 
 const TCHAR* GUnknownPackage = TEXT("Unknown Package");
 const TCHAR* GUnknownClass = TEXT("Unknown Class");
@@ -17,46 +16,18 @@ FPackageData::FPackageData(uint64 InId)
 	: Id(InId)
 	, Name(GUnknownPackage)
 	, AssetClass(GUnknownClass)
-{}
-
-void FCookProfilerProvider::BeginEdit() const
-{
-	Lock.BeginWrite(GThreadCurrentCookProviderLock, GThreadCurrentReadCookProviderLockCount, GThreadCurrentWriteCookProviderLockCount);
-}
-
-void FCookProfilerProvider::EndEdit() const
-{
-	Lock.EndWrite(GThreadCurrentCookProviderLock, GThreadCurrentWriteCookProviderLockCount);
-}
-
-void FCookProfilerProvider::EditAccessCheck() const
-{
-	Lock.WriteAccessCheck(GThreadCurrentReadCookProviderLockCount);
-}
-
-void FCookProfilerProvider::BeginRead() const
-{
-	Lock.BeginRead(GThreadCurrentCookProviderLock, GThreadCurrentReadCookProviderLockCount, GThreadCurrentWriteCookProviderLockCount);
-}
-
-void FCookProfilerProvider::EndRead() const
-{
-	Lock.EndRead(GThreadCurrentCookProviderLock, GThreadCurrentReadCookProviderLockCount);
-}
-
-void FCookProfilerProvider::ReadAccessCheck() const
 {
-	Lock.ReadAccessCheck(GThreadCurrentCookProviderLock, GThreadCurrentReadCookProviderLockCount, GThreadCurrentWriteCookProviderLockCount);
 }
 
 FCookProfilerProvider::FCookProfilerProvider(IAnalysisSession& InSession)
 	: Session(InSession)
 {
-	
 }
 
 void FCookProfilerProvider::EnumeratePackages(double StartTime, double EndTime, EnumeratePackagesCallback Callback) const
 {
+	ReadAccessCheck();
+
 	for (const FPackageData& Package : Packages)
 	{
 		if (Callback(Package) == false)
@@ -68,11 +39,15 @@ void FCookProfilerProvider::EnumeratePackages(double StartTime, double EndTime,
 
 uint32 FCookProfilerProvider::GetNumPackages() const
 {
+	ReadAccessCheck();
+
 	return (uint32) Packages.Num();
 }
 
 FPackageData* FCookProfilerProvider::EditPackage(uint64 Id)
 {
+	EditAccessCheck();
+
 	uint32 Index = FindOrAddPackage(Id);
 
 	FPackageData& Package = Packages[Index];
diff --git a/Engine/Source/Developer/TraceServices/Private/Model/CookProfilerProviderPrivate.h b/Engine/Source/Developer/TraceServices/Private/Model/CookProfilerProviderPrivate.h
--- a/Engine/Source/Developer/TraceServices/Private/Model/CookProfilerProviderPrivate.h
+++ b/Engine/Source/Developer/TraceServices/Private/Model/CookProfilerProviderPrivate.h
@@ -2,27 +2,29 @@
 
 #pragma once
 
+#include "Common/ProviderLock.h"
 #include "Containers/Map.h"
 #include "TraceServices/Model/CookProfilerProvider.h"
 
 namespace TraceServices
 {
 
+extern thread_local FProviderLock::FThreadLocalState GCookProviderLockState;
+
 class FCookProfilerProvider
 	: public ICookProfilerProvider
 {
 public:
-	virtual void BeginEdit() const override;
-	virtual void EndEdit() const override;
-	virtual void EditAccessCheck() const override;
-
-	virtual void BeginRead() const override;
-	virtual void EndRead() const override;
-	virtual void ReadAccessCheck() const override;
-
 	explicit FCookProfilerProvider(IAnalysisSession& Session);
 	virtual ~FCookProfilerProvider() {}
 
+	virtual void BeginEdit() const override       { Lock.BeginWrite(GCookProviderLockState); }
+	virtual void EndEdit() const override         { Lock.EndWrite(GCookProviderLockState); }
+	virtual void EditAccessCheck() const override { Lock.WriteAccessCheck(GCookProviderLockState); }
+	virtual void BeginRead() const override       { Lock.BeginRead(GCookProviderLockState); }
+	virtual void EndRead() const override         { Lock.EndRead(GCookProviderLockState); }
+	virtual void ReadAccessCheck() const override { Lock.ReadAccessCheck(GCookProviderLockState); }
+
 	FPackageData* EditPackage(uint64 Id);
 
 	// Functions to query the data.
@@ -33,11 +35,12 @@ private:
 	uint32 FindOrAddPackage(uint64 Id);
 
 private:
+	mutable FProviderLock Lock;
+
 	IAnalysisSession& Session;
+
 	TMap<uint64, uint32> PackageIdToIndexMap;
 	TArray64<FPackageData> Packages;
-
-	mutable FProviderLock Lock;
 };
 
 } // namespace TraceServices
diff --git a/Engine/Source/Developer/TraceServices/Private/Model/MetadataProvider.cpp b/Engine/Source/Developer/TraceServices/Private/Model/MetadataProvider.cpp
--- a/Engine/Source/Developer/TraceServices/Private/Model/MetadataProvider.cpp
+++ b/Engine/Source/Developer/TraceServices/Private/Model/MetadataProvider.cpp
@@ -12,67 +12,7 @@
 namespace TraceServices
 {
 
-////////////////////////////////////////////////////////////////////////////////////////////////////
-// FMetadataProviderLock
-////////////////////////////////////////////////////////////////////////////////////////////////////
-
-thread_local FMetadataProviderLock* GThreadCurrentMetadataProviderLock;
-thread_local int32 GThreadCurrentReadMetadataProviderLockCount;
-thread_local int32 GThreadCurrentWriteMetadataProviderLockCount;
-
-void FMetadataProviderLock::ReadAccessCheck() const
-{
-	checkf(GThreadCurrentMetadataProviderLock == this && (GThreadCurrentReadMetadataProviderLockCount > 0 || GThreadCurrentWriteMetadataProviderLockCount > 0),
-		TEXT("Trying to READ from metadata provider outside of a READ scope"));
-}
-
-void FMetadataProviderLock::WriteAccessCheck() const
-{
-	checkf(GThreadCurrentMetadataProviderLock == this && GThreadCurrentWriteMetadataProviderLockCount > 0,
-		TEXT("Trying to WRITE to metadata provider outside of an EDIT/WRITE scope"));
-}
-
-void FMetadataProviderLock::BeginRead()
-{
-	check(!GThreadCurrentMetadataProviderLock || GThreadCurrentMetadataProviderLock == this);
-	checkf(GThreadCurrentWriteMetadataProviderLockCount == 0, TEXT("Trying to lock metadata provider for READ while holding EDIT/WRITE access"));
-	if (GThreadCurrentReadMetadataProviderLockCount++ == 0)
-	{
-		GThreadCurrentMetadataProviderLock = this;
-		RWLock.ReadLock();
-	}
-}
-
-void FMetadataProviderLock::EndRead()
-{
-	check(GThreadCurrentReadMetadataProviderLockCount > 0);
-	if (--GThreadCurrentReadMetadataProviderLockCount == 0)
-	{
-		RWLock.ReadUnlock();
-		GThreadCurrentMetadataProviderLock = nullptr;
-	}
-}
-
-void FMetadataProviderLock::BeginWrite()
-{
-	check(!GThreadCurrentMetadataProviderLock || GThreadCurrentMetadataProviderLock == this);
-	checkf(GThreadCurrentReadMetadataProviderLockCount == 0, TEXT("Trying to lock metadata provider for EDIT/WRITE while holding READ access"));
-	if (GThreadCurrentWriteMetadataProviderLockCount++ == 0)
-	{
-		GThreadCurrentMetadataProviderLock = this;
-		RWLock.WriteLock();
-	}
-}
-
-void FMetadataProviderLock::EndWrite()
-{
-	check(GThreadCurrentWriteMetadataProviderLockCount > 0);
-	if (--GThreadCurrentWriteMetadataProviderLockCount == 0)
-	{
-		RWLock.WriteUnlock();
-		GThreadCurrentMetadataProviderLock = nullptr;
-	}
-}
+thread_local FProviderLock::FThreadLocalState GMetadataProviderLockState;
 
 ////////////////////////////////////////////////////////////////////////////////////////////////////
 // FMetadataProvider
@@ -117,7 +57,7 @@ FMetadataProvider::~FMetadataProvider()
 
 uint16 FMetadataProvider::RegisterMetadataType(const TCHAR* InName, const FMetadataSchema& InSchema)
 {
-	Lock.WriteAccessCheck();
+	EditAccessCheck();
 
 	check(RegisteredTypes.Num() <= MaxMetadataTypeId);
 	const uint16 Type = static_cast<uint16>(RegisteredTypes.Num());
@@ -131,7 +71,7 @@ uint16 FMetadataProvider::RegisterMetadataType(const TCHAR* InName, const FMetad
 
 uint16 FMetadataProvider::GetRegisteredMetadataType(FName InName) const
 {
-	Lock.ReadAccessCheck();
+	ReadAccessCheck();
 
 	const uint16* TypePtr = RegisteredTypesMap.Find(InName);
 	return TypePtr ? *TypePtr : InvalidMetadataType;
@@ -141,7 +81,7 @@ uint16 FMetadataProvider::GetRegisteredMetadataType(FName InName) const
 
 FName FMetadataProvider::GetRegisteredMetadataName(uint16 InType) const
 {
-	Lock.ReadAccessCheck();
+	ReadAccessCheck();
 
 	const FName* Name = RegisteredTypesMap.FindKey(InType);
 	return Name ? *Name : FName();
@@ -151,7 +91,7 @@ FName FMetadataProvider::GetRegisteredMetadataName(uint16 InType) const
 
 const FMetadataSchema* FMetadataProvider::GetRegisteredMetadataSchema(uint16 InType) const
 {
-	Lock.ReadAccessCheck();
+	ReadAccessCheck();
 
 	const int32 Index = (int32)InType;
 	return Index < RegisteredTypes.Num() ? &RegisteredTypes[Index] : nullptr;
@@ -303,7 +243,7 @@ void FMetadataProvider::InternalPushSavedStack(FMetadataThread& InMetadataThread
 
 void FMetadataProvider::PushScopedMetadata(uint32 InThreadId, uint16 InType, const void* InData, uint32 InSize)
 {
-	Lock.WriteAccessCheck();
+	EditAccessCheck();
 	++EventCount;
 
 	METADATA_PROVIDER_DEBUG_LOG(InThreadId, TEXT("PushScopedMetadata(type=%u, size=%u)"), uint32(InType), InSize);
@@ -332,7 +272,7 @@ void FMetadataProvider::PushScopedMetadata(uint32 InThreadId, uint16 InType, con
 
 void FMetadataProvider::PopScopedMetadata(uint32 InThreadId, uint16 InType)
 {
-	Lock.WriteAccessCheck();
+	EditAccessCheck();
 	++EventCount;
 
 	METADATA_PROVIDER_DEBUG_LOG(InThreadId, TEXT("PopScopedMetadata(type=%u)"), uint32(InType));
@@ -360,7 +300,7 @@ void FMetadataProvider::PopScopedMetadata(uint32 InThreadId, uint16 InType)
 
 void FMetadataProvider::BeginClearStackScope(uint32 InThreadId)
 {
-	Lock.WriteAccessCheck();
+	EditAccessCheck();
 	++EventCount;
 
 	METADATA_PROVIDER_DEBUG_LOG(InThreadId, TEXT("BeginClearStackScope()"));
@@ -384,7 +324,7 @@ void FMetadataProvider::BeginClearStackScope(uint32 InThreadId)
 
 void FMetadataProvider::EndClearStackScope(uint32 InThreadId)
 {
-	Lock.WriteAccessCheck();
+	EditAccessCheck();
 	++EventCount;
 
 	METADATA_PROVIDER_DEBUG_LOG(InThreadId, TEXT("EndClearStackScope()"));
@@ -427,7 +367,7 @@ void FMetadataProvider::EndClearStackScope(uint32 InThreadId)
 
 void FMetadataProvider::SaveStack(uint32 InThreadId, uint32 InSavedStackId)
 {
-	Lock.WriteAccessCheck();
+	EditAccessCheck();
 	++EventCount;
 
 	const uint32 MetadataId = PinAndGetId(InThreadId);
@@ -438,7 +378,7 @@ void FMetadataProvider::SaveStack(uint32 InThreadId, uint32 InSavedStackId)
 
 void FMetadataProvider::BeginRestoreSavedStackScope(uint32 InThreadId, uint32 InSavedStackId)
 {
-	Lock.WriteAccessCheck();
+	EditAccessCheck();
 	++EventCount;
 
 	METADATA_PROVIDER_DEBUG_LOG(InThreadId, TEXT("BeginRestoreSavedStackScope()"));
@@ -486,7 +426,7 @@ void FMetadataProvider::BeginRestoreSavedStackScope(uint32 InThreadId, uint32 In
 
 void FMetadataProvider::EndRestoreSavedStackScope(uint32 InThreadId)
 {
-	Lock.WriteAccessCheck();
+	EditAccessCheck();
 	++EventCount;
 
 	METADATA_PROVIDER_DEBUG_LOG(InThreadId, TEXT("EndRestoreSavedStackScope()"));
@@ -527,7 +467,7 @@ void FMetadataProvider::EndRestoreSavedStackScope(uint32 InThreadId)
 
 uint32 FMetadataProvider::PinAndGetId(uint32 InThreadId)
 {
-	Lock.WriteAccessCheck();
+	EditAccessCheck();
 
 	FMetadataThread* MetadataThread = GetThread(InThreadId);
 
@@ -603,7 +543,7 @@ uint32 FMetadataProvider::PinAndGetId(uint32 InThreadId)
 
 void FMetadataProvider::OnAnalysisCompleted()
 {
-	Lock.WriteAccessCheck();
+	EditAccessCheck();
 
 	if (MetaScopeErrors > 0)
 	{
@@ -626,7 +566,7 @@ void FMetadataProvider::OnAnalysisCompleted()
 
 uint32 FMetadataProvider::GetMetadataStackSize(uint32 InThreadId, uint32 InMetadataId) const
 {
-	Lock.ReadAccessCheck();
+	ReadAccessCheck();
 
 	const FMetadataThread* MetadataThread = GetThread(InThreadId);
 
@@ -645,7 +585,7 @@ uint32 FMetadataProvider::GetMetadataStackSize(uint32 InThreadId, uint32 InMetad
 
 bool FMetadataProvider::GetMetadata(uint32 InThreadId, uint32 InMetadataId, uint32 InStackDepth, uint16& OutType, const void*& OutData, uint32& OutSize) const
 {
-	Lock.ReadAccessCheck();
+	ReadAccessCheck();
 
 	const FMetadataThread* MetadataThread = GetThread(InThreadId);
 
@@ -686,7 +626,7 @@ bool FMetadataProvider::GetMetadata(uint32 InThreadId, uint32 InMetadataId, uint
 
 void FMetadataProvider::EnumerateMetadata(uint32 InThreadId, uint32 InMetadataId, TFunctionRef<bool(uint32 StackDepth, uint16 Type, const void* Data, uint32 Size)> Callback) const
 {
-	Lock.ReadAccessCheck();
+	ReadAccessCheck();
 
 	const FMetadataThread* MetadataThread = GetThread(InThreadId);
 
diff --git a/Engine/Source/Developer/TraceServices/Private/Model/MetadataProvider.h b/Engine/Source/Developer/TraceServices/Private/Model/MetadataProvider.h
--- a/Engine/Source/Developer/TraceServices/Private/Model/MetadataProvider.h
+++ b/Engine/Source/Developer/TraceServices/Private/Model/MetadataProvider.h
@@ -3,6 +3,7 @@
 #pragma once
 
 #include "Common/PagedArray.h"
+#include "Common/ProviderLock.h"
 #include "TraceServices/Model/MetadataProvider.h"
 
 namespace TraceServices
@@ -11,6 +12,8 @@ namespace TraceServices
 class IAnalysisSession;
 class ILinearAllocator;
 
+extern thread_local FProviderLock::FThreadLocalState GMetadataProviderLockState;
+
 ////////////////////////////////////////////////////////////////////////////////////////////////////
 
 class FMetadataProviderLock
@@ -92,13 +95,12 @@ public:
 	explicit FMetadataProvider(IAnalysisSession& InSession);
 	virtual ~FMetadataProvider();
 
-	virtual void BeginEdit() const override { Lock.BeginWrite(); }
-	virtual void EndEdit() const override { Lock.EndWrite(); }
-	void EditAccessCheck() const { return Lock.WriteAccessCheck(); }
-
-	virtual void BeginRead() const override { Lock.BeginRead(); }
-	virtual void EndRead() const override { Lock.EndRead(); }
-	void ReadAccessCheck() const { return Lock.ReadAccessCheck(); }
+	virtual void BeginEdit() const override       { Lock.BeginWrite(GMetadataProviderLockState); }
+	virtual void EndEdit() const override         { Lock.EndWrite(GMetadataProviderLockState); }
+	virtual void EditAccessCheck() const override { Lock.WriteAccessCheck(GMetadataProviderLockState); }
+	virtual void BeginRead() const override       { Lock.BeginRead(GMetadataProviderLockState); }
+	virtual void EndRead() const override         { Lock.EndRead(GMetadataProviderLockState); }
+	virtual void ReadAccessCheck() const override { Lock.ReadAccessCheck(GMetadataProviderLockState); }
 
 	//////////////////////////////////////////////////
 	// Edit operations
@@ -143,9 +145,10 @@ private:
 	void InternalPushSavedStack(FMetadataThread& InMetadataThread, FMetadataThread& InSavedMetadataThread, uint32 InSavedMetadataId);
 
 private:
+	mutable FProviderLock Lock;
+
 	IAnalysisSession& Session;
 
-	mutable FMetadataProviderLock Lock;
 	TPagedArray<FMetadataSchema> RegisteredTypes;
 	TMap<FName, uint16> RegisteredTypesMap;
 
diff --git a/Engine/Source/Developer/TraceServices/Public/Common/ProviderLock.h b/Engine/Source/Developer/TraceServices/Public/Common/ProviderLock.h
--- a/Engine/Source/Developer/TraceServices/Public/Common/ProviderLock.h
+++ b/Engine/Source/Developer/TraceServices/Public/Common/ProviderLock.h
@@ -41,16 +41,46 @@ struct TRACESERVICES_API FProviderReadScopeLock
 	const IProvider& Provider;
 };
 
+/**
+ * Utility class to implement the read/write lock for a provider.
+ * Example usage:
+       extern thread_local FProviderLock::FThreadLocalState MyProviderLockState;
+ *     virtual void EditAccessCheck() const override.{ Lock.WriteAccessCheck(MyProviderLockState); }
+ *     FProviderLock Lock;
+ */
 class FProviderLock
 {
 public:
+	struct FThreadLocalState
+	{
+		FProviderLock* Lock;
+		int32 ReadLockCount;
+		int32 WriteLockCount;
+	};
+
+public:
+	void ReadAccessCheck(FThreadLocalState& State) const;
+	void WriteAccessCheck(FThreadLocalState& State) const;
+
+	void BeginRead(FThreadLocalState& State);
+	void EndRead(FThreadLocalState& State);
+
+	void BeginWrite(FThreadLocalState& State);
+	void EndWrite(FThreadLocalState& State);
+
+	UE_DEPRECATED(5.3, "Please use the FThreadLocalState overload")
 	void ReadAccessCheck(const FProviderLock* CurrentProviderLock, const int32& CurrentReadProviderLockCount, const int32& CurrentWriteProviderLockCount) const;
+	UE_DEPRECATED(5.3, "Please use the FThreadLocalState overload")
 	void WriteAccessCheck(const int32& CurrentWriteProviderLockCount) const;
 
+	UE_DEPRECATED(5.3, "Please use the FThreadLocalState overload")
 	void BeginRead(FProviderLock*& CurrentProviderLock, int32& CurrentReadProviderLockCount, const int32& WriteProviderLockCount);
+	UE_DEPRECATED(5.3, "Please use the FThreadLocalState overload")
 	void EndRead(FProviderLock*& CurrentProviderLock, int32& CurrentReadProviderLockCount);
 
+	UE_DEPRECATED(5.3, "Please use the FThreadLocalState overload")
 	void BeginWrite(FProviderLock*& CurrentProviderLock, const int32& CurrentReadProviderLockCount, int32& WriteProviderLockCount);
+	UE_DEPRECATED(5.3, "Please use the FThreadLocalState overload")
 	void EndWrite(FProviderLock*& CurrentProviderLock, int32& WriteProviderLockCount);
 
 private:
