[Insights] CookProfiler: Minor code improvements (store strings outside the provider lock + fewer allocations during initialisation of Packages table tree view).

diff --git a/Engine/Source/Developer/TraceInsights/Private/Insights/CookProfiler/Widgets/SPackageTableTreeView.cpp b/Engine/Source/Developer/TraceInsights/Private/Insights/CookProfiler/Widgets/SPackageTableTreeView.cpp
--- a/Engine/Source/Developer/TraceInsights/Private/Insights/CookProfiler/Widgets/SPackageTableTreeView.cpp
+++ b/Engine/Source/Developer/TraceInsights/Private/Insights/CookProfiler/Widgets/SPackageTableTreeView.cpp
@@ -133,8 +133,12 @@ void SPackageTableTreeView::RebuildTree(bool bResync)
 	if (CookProvider)
 	{
 		TraceServices::FProviderReadScopeLock ProviderReadScope(*CookProvider);
-		TArray<FTableTreeNodePtr>* Nodes = &TableTreeNodes;
 
+		const uint32 NumPackages = CookProvider->GetNumPackages();
+		Packages.Reserve(NumPackages);
+		TableTreeNodes.Reserve(NumPackages);
+
+		TArray<FTableTreeNodePtr>* Nodes = &TableTreeNodes;
 		CookProvider->EnumeratePackages(0, 0, [&Packages, &PackageTable, Nodes](const TraceServices::FPackageData& Package)
 			{
 				Packages.Emplace(Package);
diff --git a/Engine/Source/Developer/TraceServices/Private/Analyzers/CookAnalysis.cpp b/Engine/Source/Developer/TraceServices/Private/Analyzers/CookAnalysis.cpp
--- a/Engine/Source/Developer/TraceServices/Private/Analyzers/CookAnalysis.cpp
+++ b/Engine/Source/Developer/TraceServices/Private/Analyzers/CookAnalysis.cpp
@@ -41,10 +41,11 @@ bool FCookAnalyzer::OnEvent(uint16 RouteId, EStyle Style, const FOnEventContext&
 		uint64 Id = EventData.GetValue<uint64>("Id");
 		FString Name;
 		EventData.GetString("Name", Name);
+		const TCHAR* PersistentName = Session.StoreString(Name);
 
-		TraceServices::FProviderEditScopeLock ProviderEditScope(CookProfilerProvider);
+		FProviderEditScopeLock ProviderEditScope(CookProfilerProvider);
 		FPackageData* Package = CookProfilerProvider.EditPackage(Id);
-		Package->Name = Session.StoreString(Name);
+		Package->Name = PersistentName;
 		break;
 	}
 	case RouteId_PackageStat:
@@ -53,12 +54,11 @@ bool FCookAnalyzer::OnEvent(uint16 RouteId, EStyle Style, const FOnEventContext&
 		double Value = Context.EventTime.AsSecondsAbsolute(EventData.GetValue<uint64>("Duration"));
 		EPackageEventStatType StatType = (EPackageEventStatType) EventData.GetValue<uint8>("StatType");
 
-		TraceServices::FProviderEditScopeLock ProviderEditScope(CookProfilerProvider);
+		FProviderEditScopeLock ProviderEditScope(CookProfilerProvider);
 		FPackageData* Package = CookProfilerProvider.EditPackage(Id);
-		if (Package)
+		check(Package);
+		switch (StatType)
 		{
-			switch (StatType)
-			{
 			case EPackageEventStatType::LoadPackage:
 			{
 				// We measure the Loadtime in multiple scopes so we receive many LoadPackage events.
@@ -86,7 +86,6 @@ bool FCookAnalyzer::OnEvent(uint16 RouteId, EStyle Style, const FOnEventContext&
 			{
 				break;
 			}
-			}
 		}
 		break;
 	}
@@ -95,10 +94,12 @@ bool FCookAnalyzer::OnEvent(uint16 RouteId, EStyle Style, const FOnEventContext&
 		uint64 Id = EventData.GetValue<uint64>("Id");
 		FString ClassName;
 		EventData.GetString("ClassName", ClassName);
+		const TCHAR* PersistentClassName = Session.StoreString(ClassName);
 
-		TraceServices::FProviderEditScopeLock ProviderEditScope(CookProfilerProvider);
+		FProviderEditScopeLock ProviderEditScope(CookProfilerProvider);
 		FPackageData* Package = CookProfilerProvider.EditPackage(Id);
-		Package->AssetClass = Session.StoreString(ClassName);
+		check(Package);
+		Package->AssetClass = PersistentClassName;
 		break;
 	}
 	}
