TasksInsights: refactored analysis a bit to systematically handle tasks that miss some events due to dynamically enabled task channel

diff --git a/Engine/Source/Developer/TraceServices/Private/Model/TasksProfiler.cpp b/Engine/Source/Developer/TraceServices/Private/Model/TasksProfiler.cpp
--- a/Engine/Source/Developer/TraceServices/Private/Model/TasksProfiler.cpp
+++ b/Engine/Source/Developer/TraceServices/Private/Model/TasksProfiler.cpp
@@ -152,17 +152,26 @@ namespace TraceServices
 			return;
 		}
 
+		// some task events can be duplicates, as the result of generating them here if they are missing (some task events are not sent to reduce trace 
+		// file size), but there's no way to distinguish missing and late events (task events can race and come out of order). A generated event comes
+		// always before a real one, so the real one overrides the generated one silently. but as counters were already updated for a generated event,
+		// we don't update them for a duplicate.
+		bool bDuplicateEvent = Task->CreatedTimestamp != FTaskInfo::InvalidTimestamp;
+
 		Task->CreatedTimestamp = Timestamp;
 		Task->CreatedThreadId = ThreadId;
 		Task->TaskSize = TaskSize;
 
-		if (!bCountersCreated)
+		if (!bDuplicateEvent)
 		{
-			CreateCounters();
-		}
+			if (!bCountersCreated)
+			{
+				CreateCounters();
+			}
 
-		AliveTasksCounter->SetValue(Timestamp, ++AliveTasksNum);
-		AliveTasksSizeCounter->SetValue(Timestamp, AliveTasksSize += TaskSize);
+			AliveTasksCounter->SetValue(Timestamp, ++AliveTasksNum);
+			AliveTasksSizeCounter->SetValue(Timestamp, AliveTasksSize += TaskSize);
+		}
 	}
 
 	void FTasksProvider::TaskLaunched(TaskTrace::FId TaskId, const TCHAR* DebugName, bool bTracked, int32 ThreadToExecuteOn, double Timestamp, uint32 ThreadId, uint64 TaskSize)
@@ -178,32 +187,28 @@ namespace TraceServices
 			return;
 		}
 
-#if UE_TASK_TRACES_ENABLE_OLD_TASKS_SUPPORT // a task with incomplete nested tasks is "relaunched" when it's execution is finished
-		if (Task->LaunchedTimestamp != FTaskInfo::InvalidTimestamp)
-		{
-			return;
-		}
-#else
-		checkf(Task->LaunchedTimestamp == FTaskInfo::InvalidTimestamp, TEXT("TaskLaunched(TaskId %d, Timestamp %.6f, ThreadId %d) repeated, prev Timestamp %.6f, ThreadId %d"), TaskId, Timestamp, ThreadId, Task->LaunchedTimestamp, Task->LaunchedThreadId);
-#endif
-			
 		if (Task->CreatedTimestamp == FTaskInfo::InvalidTimestamp) // created and launched in one go
 		{
 			TaskCreated(TaskId, Timestamp, ThreadId, TaskSize);
 		}
 
+		bool bDuplicateEvent = Task->LaunchedTimestamp != FTaskInfo::InvalidTimestamp;
+
 		Task->DebugName = DebugName;
 		Task->bTracked = bTracked;
 		Task->ThreadToExecuteOn = ThreadToExecuteOn;
 		Task->LaunchedTimestamp = Timestamp;
 		Task->LaunchedThreadId = ThreadId;
 
-		if (!bCountersCreated)
+		if (!bDuplicateEvent)
 		{
-			CreateCounters();
-		}
+			if (!bCountersCreated)
+			{
+				CreateCounters();
+			}
 
-		WaitingForPrerequisitesTasksCounter->SetValue(Timestamp, ++WaitingForPrerequisitesTasksNum);
+			WaitingForPrerequisitesTasksCounter->SetValue(Timestamp, ++WaitingForPrerequisitesTasksNum);
+		}
 	}
 
 	void FTasksProvider::TaskScheduled(TaskTrace::FId TaskId, double Timestamp, uint32 ThreadId)
@@ -217,31 +222,27 @@ namespace TraceServices
 			return;
 		}
 
-#if UE_TASK_TRACES_ENABLE_OLD_TASKS_SUPPORT // a task with incomplete nested tasks is "relaunched" when it's execution is finished
-		if (Task->ScheduledTimestamp != FTaskInfo::InvalidTimestamp)
-		{
-			return;
-		}
-#else
-		checkf(Task->ScheduledTimestamp == FTaskInfo::InvalidTimestamp, TEXT("TaskScheduled(TaskId %d, Timestamp %.6f, ThreadId %d) repeated, prev Timestamp %.6f, ThreadId %d"), TaskId, Timestamp, ThreadId, Task->ScheduledTimestamp, Task->ScheduledThreadId);
-#endif
+		bool bDuplicateEvent = Task->ScheduledTimestamp != FTaskInfo::InvalidTimestamp;
 
 		Task->ScheduledTimestamp = Timestamp;
 		Task->ScheduledThreadId = ThreadId;
 
-		if (!bCountersCreated)
+		if (!bDuplicateEvent)
 		{
-			CreateCounters();
-		}
+			if (!bCountersCreated)
+			{
+				CreateCounters();
+			}
 
-		WaitingForPrerequisitesTasksCounter->SetValue(Timestamp, --WaitingForPrerequisitesTasksNum);
-		if (IsNamedThread(Task->ThreadToExecuteOn))
-		{
-			NamedThreadsScheduledTasksCounter->SetValue(Timestamp, ++ScheduledTasksNum);
-		}
-		else
-		{
-			ScheduledTasksCounter->SetValue(Timestamp, ++ScheduledTasksNum);
+			WaitingForPrerequisitesTasksCounter->SetValue(Timestamp, --WaitingForPrerequisitesTasksNum);
+			if (IsNamedThread(Task->ThreadToExecuteOn))
+			{
+				NamedThreadsScheduledTasksCounter->SetValue(Timestamp, ++ScheduledTasksNum);
+			}
+			else
+			{
+				ScheduledTasksCounter->SetValue(Timestamp, ++ScheduledTasksNum);
+			}
 		}
 	}
 
@@ -292,36 +293,38 @@ namespace TraceServices
 			return;
 		}
 
-#if UE_TASK_TRACES_ENABLE_OLD_TASKS_SUPPORT // a task with incomplete nested tasks is "relaunched" when it's execution is finished
-		if (Task->StartedTimestamp != FTaskInfo::InvalidTimestamp)
+		if (Task->ScheduledTimestamp == FTaskInfo::InvalidTimestamp && Task->LaunchedTimestamp != FTaskInfo::InvalidTimestamp)
 		{
-			return;
+			// omitted "scheduled" event, generate it
+			TaskScheduled(TaskId, Timestamp, Task->LaunchedThreadId); // `LaunchedThreadId` is just a value that potentially is closest to the truth
 		}
-#else
-		checkf(Task->StartedTimestamp == FTaskInfo::InvalidTimestamp, TEXT("TaskStarted(TaskId %d, Timestamp %.6f, ThreadId %d) repeated, prev Timestamp %.6f, ThreadId %d"), TaskId, Timestamp, ThreadId, Task->StartedTimestamp, Task->StartedThreadId);
-#endif
+
+		bool bDuplicateEvent = Task->StartedTimestamp != FTaskInfo::InvalidTimestamp;
 
 		Task->StartedTimestamp = Timestamp;
 		Task->StartedThreadId = ThreadId;
 
 		ExecutionThreads.FindOrAdd(ThreadId).Add(TaskId);
 
-		if (!bCountersCreated)
+		if (!bDuplicateEvent)
 		{
-			CreateCounters();
-		}
+			if (!bCountersCreated)
+			{
+				CreateCounters();
+			}
 
-		if (IsNamedThread(Task->ThreadToExecuteOn))
-		{
-			NamedThreadsScheduledTasksCounter->SetValue(Timestamp, --ScheduledTasksNum);
-		}
-		else
-		{
-			ScheduledTasksCounter->SetValue(Timestamp, --ScheduledTasksNum);
-		}
-		RunningTasksCounter->SetValue(Timestamp, ++RunningTasksNum);
+			if (IsNamedThread(Task->ThreadToExecuteOn))
+			{
+				NamedThreadsScheduledTasksCounter->SetValue(Timestamp, --ScheduledTasksNum);
+			}
+			else
+			{
+				ScheduledTasksCounter->SetValue(Timestamp, --ScheduledTasksNum);
+			}
+			RunningTasksCounter->SetValue(Timestamp, ++RunningTasksNum);
 
-		TaskLatencyCounter->SetValue(Timestamp, Task->StartedTimestamp - Task->ScheduledTimestamp);
+			TaskLatencyCounter->SetValue(Timestamp, Task->StartedTimestamp - Task->ScheduledTimestamp);
+		}
 	}
 
 	void FTasksProvider::TaskFinished(TaskTrace::FId TaskId, double Timestamp)
@@ -335,24 +338,26 @@ namespace TraceServices
 			return;
 		}
 
-#if UE_TASK_TRACES_ENABLE_OLD_TASKS_SUPPORT // a task with incomplete nested tasks is "relaunched" when it's execution is finished
-		if (Task->FinishedTimestamp != FTaskInfo::InvalidTimestamp)
+		if (Task->StartedTimestamp == FTaskInfo::InvalidTimestamp && Task->LaunchedTimestamp != FTaskInfo::InvalidTimestamp)
 		{
-			return;
+			// missing "started" event for a task covered by this trace
+			TaskStarted(TaskId, Timestamp, Task->LaunchedThreadId);
 		}
-#else
-		checkf(Task->FinishedTimestamp == FTaskInfo::InvalidTimestamp, TEXT("TaskFinished(TaskId %d, Timestamp %.6f) repeated, prev Timestamp %.6f"), TaskId, Timestamp, Task->StartedTimestamp);
-#endif
+
+		bool bDuplicateEvent = Task->FinishedTimestamp != FTaskInfo::InvalidTimestamp;
 
 		Task->FinishedTimestamp = Timestamp;
 
-		if (!bCountersCreated)
+		if (!bDuplicateEvent)
 		{
-			CreateCounters();
-		}
+			if (!bCountersCreated)
+			{
+				CreateCounters();
+			}
 
-		RunningTasksCounter->SetValue(Timestamp, --RunningTasksNum);
-		ExecutionTimeCounter->SetValue(Timestamp, Task->FinishedTimestamp - Task->StartedTimestamp);
+			RunningTasksCounter->SetValue(Timestamp, --RunningTasksNum);
+			ExecutionTimeCounter->SetValue(Timestamp, Task->FinishedTimestamp - Task->StartedTimestamp);
+		}
 	}
 
 	void FTasksProvider::TaskCompleted(TaskTrace::FId TaskId, double Timestamp, uint32 ThreadId)
@@ -368,39 +373,12 @@ namespace TraceServices
 			return;
 		}
 
-#if UE_TASK_TRACES_ENABLE_OLD_TASKS_SUPPORT
-		if (Task->LaunchedTimestamp == FTaskInfo::InvalidTimestamp)
-		{	// a standalone FGraphEvent w/o an accompanying task or subsequents, the first trace is "completed" event on `DispatchSubsequents`
-			TaskLaunched(TaskId, /*DebugName = */ TEXT("GraphEvent"), /*bTracked = */ true, /*ThreadToExecuteOn = */ 0, Timestamp, ThreadId, /*TaskSize = */ sizeof(FGraphEvent));
-			//TaskStarted(TaskId, Timestamp, ThreadId);
-			//TaskFinished(TaskId, Timestamp);
-		}
-#endif
-
-		checkf(Task->LaunchedTimestamp != FTaskInfo::InvalidTimestamp, TEXT("TaskCompleted(TaskId %d, Timestamp %.6f): no \"launched\" timestamp"), TaskId, Timestamp);
-		checkf(Task->CreatedTimestamp != FTaskInfo::InvalidTimestamp, TEXT("TaskCompleted(TaskId %d, Timestamp %.6f): no \"created\" timestamp"), TaskId, Timestamp);
-
-		if (Task->ScheduledTimestamp == FTaskInfo::InvalidTimestamp)
+		if (Task->FinishedTimestamp == FTaskInfo::InvalidTimestamp && Task->LaunchedTimestamp != FTaskInfo::InvalidTimestamp)
 		{
-			if (Task->FinishedTimestamp != FTaskInfo::InvalidTimestamp) // a regular task just missing omitted "scheduled" trace
-			{
-				TaskScheduled(TaskId, Task->StartedTimestamp, Task->StartedThreadId);
-			}
-			else if (Task->StartedTimestamp == FTaskInfo::InvalidTimestamp) // FTaskEvent
-			{
-				TaskScheduled(TaskId, Timestamp, ThreadId);
-				TaskStarted(TaskId, Timestamp, ThreadId);
-				TaskFinished(TaskId, Timestamp);
-			}
-			else // ParallelFor task
-			{
-				TaskScheduled(TaskId, Task->LaunchedTimestamp, Task->LaunchedThreadId);
-				TaskFinished(TaskId, Timestamp);
-			}
+			// missing "finished" event for a task that is covered by this trace, generate it
+			TaskFinished(TaskId, Timestamp);
 		}
 
-		checkf(Task->CompletedTimestamp == FTaskInfo::InvalidTimestamp, TEXT("TaskCompleted(TaskId %d, Timestamp %.6f, ThreadId %d): repeated, prev Timestamp %.6f, ThreadId %d"), TaskId, Timestamp, ThreadId, Task->CompletedTimestamp, Task->CompletedThreadId);
-
 		Task->CompletedTimestamp = Timestamp;
 		Task->CompletedThreadId = ThreadId;
 	}
@@ -412,37 +390,31 @@ namespace TraceServices
 		FTaskInfo* Task = TryGetOrCreateTask(TaskId);
 		if (Task == nullptr)
 		{
-			UE_LOG(LogTraceServices, Log, TEXT("TaskStarted(TaskId %d, Timestamp %.6f) skipped"), TaskId, Timestamp);
+			UE_LOG(LogTraceServices, Log, TEXT("TaskDestroyed(TaskId %d, Timestamp %.6f) skipped"), TaskId, Timestamp);
 			return;
 		}
 
-		if (Task->CompletedTimestamp == FTaskInfo::InvalidTimestamp) // cancelled ParallelFor task
+		if (Task->CompletedTimestamp == FTaskInfo::InvalidTimestamp && Task->LaunchedTimestamp != FTaskInfo::InvalidTimestamp)
 		{
-			check(Task->ScheduledTimestamp == FTaskInfo::InvalidTimestamp &&
-				Task->StartedTimestamp == FTaskInfo::InvalidTimestamp &&
-				Task->FinishedTimestamp == FTaskInfo::InvalidTimestamp);
-			
-			TaskScheduled(TaskId, Timestamp, ThreadId);
-			TaskStarted(TaskId, Timestamp, ThreadId);
-			TaskFinished(TaskId, Timestamp);
+			// missing "completed" event for a task that is covered by this trace
 			TaskCompleted(TaskId, Timestamp, ThreadId);
 		}
 
-		// all initialised
-		check(Task->CreatedTimestamp != FTaskInfo::InvalidTimestamp && 
-			Task->LaunchedTimestamp != FTaskInfo::InvalidTimestamp &&
-			Task->ScheduledTimestamp != FTaskInfo::InvalidTimestamp &&
-			Task->StartedTimestamp != FTaskInfo::InvalidTimestamp &&
-			Task->FinishedTimestamp != FTaskInfo::InvalidTimestamp &&
-			Task->CompletedTimestamp != FTaskInfo::InvalidTimestamp);
-
-		checkf(Task->DestroyedTimestamp == FTaskInfo::InvalidTimestamp, TEXT("TaskDestroyed(TaskId %d, Timestamp %.6f, ThreadId %d): repeated, prev Timestamp %.6f, ThreadId %d"), TaskId, Timestamp, ThreadId, Task->DestroyedTimestamp, Task->DestroyedThreadId);
+		bool bUpdateCounters = Task->LaunchedTimestamp != FTaskInfo::InvalidTimestamp; // we accounted this task in counters
 
 		Task->DestroyedTimestamp = Timestamp;
 		Task->DestroyedThreadId = ThreadId;
 
-		AliveTasksCounter->SetValue(Timestamp, --AliveTasksNum);
-		AliveTasksSizeCounter->SetValue(Timestamp, AliveTasksSize -= Task->TaskSize);
+		if (bUpdateCounters)
+		{
+			if (!bCountersCreated)
+			{
+				CreateCounters();
+			}
+
+			AliveTasksCounter->SetValue(Timestamp, --AliveTasksNum);
+			AliveTasksSizeCounter->SetValue(Timestamp, AliveTasksSize -= Task->TaskSize);
+		}
 	}
 
 	void FTasksProvider::WaitingStarted(TArray<TaskTrace::FId> InTasks, double Timestamp, uint32 ThreadId)
