[Insights]
  - Memory Insights: Added support for unlimited number of heap specs. Previously, it was limited to 256 heap specs. Also, encapsulated the SbTree, the live allocations collection and the event index into a new FRootHeap class.
  - Fixed usage of not null terminated heap name in debug log message.
  - Disabled warning for 0x0 address for root heaps other than SystemMemory.

diff --git a/Engine/Source/Developer/TraceServices/Private/Model/AllocationsProvider.cpp b/Engine/Source/Developer/TraceServices/Private/Model/AllocationsProvider.cpp
--- a/Engine/Source/Developer/TraceServices/Private/Model/AllocationsProvider.cpp
+++ b/Engine/Source/Developer/TraceServices/Private/Model/AllocationsProvider.cpp
@@ -128,6 +128,8 @@ static bool IsAddressWatched(uint64 InAddress)
 
 thread_local FProviderLock::FThreadLocalState GAllocationsProviderLockState;
 
+static const TCHAR* GDefaultHeapName = TEXT("Unknown");
+
 ////////////////////////////////////////////////////////////////////////////////////////////////////
 // FTagTracker
 ////////////////////////////////////////////////////////////////////////////////////////////////////
@@ -165,10 +167,10 @@ void FTagTracker::AddTagSpec(TagIdType InTag, TagIdType InParentTag, const TCHAR
 			if (InParentTag == InvalidTagId)
 			{
 				const FStringView Parent = FPathViews::GetPathLeaf(FPathViews::GetPath(Display));
-				for (auto EntryPair : TagMap)
+				for (const auto& EntryPair : TagMap)
 				{
-					const auto Entry = EntryPair.Get<1>();
 					const uint32 Id = EntryPair.Get<0>();
+					const FTagEntry Entry = EntryPair.Get<1>();
 					if (Parent.Equals(Entry.Display))
 					{
 						InParentTag = Id;
@@ -195,10 +197,10 @@ void FTagTracker::AddTagSpec(TagIdType InTag, TagIdType InParentTag, const TCHAR
 		});
 
 		// Check if this new tag has been referenced before by a child tag
-		for (auto Pending : PendingTags)
+		for (const TTuple<TagIdType,FString>& Pending : PendingTags)
 		{
-			const FString& Name = Pending.Get<1>();
 			const TagIdType ReferencingId = Pending.Get<0>();
+			const FString& Name = Pending.Get<1>();
 			if (Name.Equals(DisplayName))
 			{
 				TagMap[ReferencingId].ParentTag = InTag;
@@ -1211,26 +1213,47 @@ FAllocationsProvider::FAllocationsProvider(IAnalysisSession& InSession, FMetadat
 	, AllocEventsTimeline(Session.GetLinearAllocator(), 1024)
 	, FreeEventsTimeline(Session.GetLinearAllocator(), 1024)
 {
-	HeapSpecs.AddZeroed(256);
+	// Initial number of heap spec locations. Actual number of heap specs is unlimited (uint32).
+	HeapSpecs.AddDefaulted(FMath::Max(1024u, MaxRootHeaps));
 }
 
 ////////////////////////////////////////////////////////////////////////////////////////////////////
 
 FAllocationsProvider::~FAllocationsProvider()
 {
-	for (uint8 RootHeapIdx = 0; RootHeapIdx < MaxRootHeaps; ++RootHeapIdx)
+	// Delete root heaps.
+	for (uint32 RootHeapIdx = 0; RootHeapIdx < MaxRootHeaps; ++RootHeapIdx)
 	{
-		if (SbTree[RootHeapIdx])
+		FRootHeap* RootHeap = RootHeaps[RootHeapIdx];
+		if (RootHeap)
 		{
-			delete SbTree[RootHeapIdx];
-			SbTree[RootHeapIdx] = nullptr;
+			RootHeap->HeapSpec = nullptr;
+
+			if (RootHeap->SbTree)
+			{
+				delete RootHeap->SbTree;
+				RootHeap->SbTree = nullptr;
+			}
+			if (RootHeap->LiveAllocs)
+			{
+				delete RootHeap->LiveAllocs;
+				RootHeap->LiveAllocs = nullptr;
+			}
+
+			delete RootHeap;
+			RootHeaps[RootHeapIdx] = nullptr;
 		}
-		if (LiveAllocs[RootHeapIdx])
+	}
+
+	// Delete all heap specs.
+	for (FHeapSpec* HeapSpec : HeapSpecs)
+	{
+		if (HeapSpec)
 		{
-			delete LiveAllocs[RootHeapIdx];
-			LiveAllocs[RootHeapIdx] = nullptr;
+			delete HeapSpec;
 		}
 	}
+	HeapSpecs.Reset();
 }
 
 ////////////////////////////////////////////////////////////////////////////////////////////////////
@@ -1260,7 +1283,7 @@ void FAllocationsProvider::EditInit(double InTime, uint8 InMinAlignment)
 
 ////////////////////////////////////////////////////////////////////////////////////////////////////
 
-void FAllocationsProvider::EditAlloc(double Time, uint32 CallstackId, uint64 Address, uint64 InSize, uint32 InAlignment, HeapId RootHeap)
+void FAllocationsProvider::EditAlloc(double Time, uint32 CallstackId, uint64 Address, uint64 Size, uint32 Alignment, HeapId RootHeapId)
 {
 	EditAccessCheck();
 
@@ -1269,30 +1292,31 @@ void FAllocationsProvider::EditAlloc(double Time, uint32 CallstackId, uint64 Add
 		return;
 	}
 
-	INSIGHTS_WATCH_API_LOGF(Address, TEXT("Alloc 0x%llX : Size=%llu Alignment=%u RootHeap=%u CallstackId=%u"), Time, Address, InSize, InAlignment, RootHeap, CallstackId);
+	INSIGHTS_WATCH_API_LOGF(Address, TEXT("Alloc 0x%llX : Size=%llu Alignment=%u RootHeap=%u CallstackId=%u"), Time, Address, Size, Alignment, RootHeapId, CallstackId);
 
-	if (Address == 0)
+	if (Address == 0 && RootHeapId == EMemoryTraceRootHeap::SystemMemory)
 	{
 		++AllocWarnings;
 		if (AllocWarnings <= MaxLogMessagesPerWarningType)
 		{
-			UE_LOG(LogTraceServices, Warning, TEXT("[MemAlloc] Alloc at address 0 : Size=%llu, RootHeap=%u, Time=%f, CallstackId=%u"), InSize, RootHeap, Time, CallstackId);
+			UE_LOG(LogTraceServices, Warning, TEXT("[MemAlloc] Alloc at address 0 : Size=%llu, RootHeap=%u, Time=%f, CallstackId=%u"), Size, RootHeapId, Time, CallstackId);
 		}
 		return;
 	}
 
-	check(RootHeap < MaxRootHeaps);
-	if (SbTree[RootHeap] == nullptr)
+	if (!IsValidRootHeap(RootHeapId))
 	{
 		++AllocWarnings;
 		if (AllocWarnings <= MaxLogMessagesPerWarningType)
 		{
-			UE_LOG(LogTraceServices, Warning, TEXT("[MemAlloc] Alloc with invalid root heap id (%u). This heap has not yet been registered."), RootHeap);
+			UE_LOG(LogTraceServices, Warning, TEXT("[MemAlloc] Alloc with invalid root heap id (%u)."), RootHeapId);
 		}
 		return;
 	}
 
-	SbTree[RootHeap]->SetTimeForEvent(EventIndex[RootHeap], Time);
+	FRootHeap& RootHeap = *RootHeaps[RootHeapId];
+
+	RootHeap.SbTree->SetTimeForEvent(RootHeap.EventIndex, Time);
 
 	AdvanceTimelines(Time);
 
@@ -1301,7 +1325,7 @@ void FAllocationsProvider::EditAlloc(double Time, uint32 CallstackId, uint64 Add
 	FAllocationItem* AllocationPtr = nullptr;
 
 #if INSIGHTS_VALIDATE_ALLOC_EVENTS
-	AllocationPtr = LiveAllocs[RootHeap]->FindRef(Address);
+	AllocationPtr = RootHeap.LiveAllocs->FindRef(Address);
 #if INSIGHTS_REMAP_INVALID_ALLOC_EVENTS
 	if (AllocationPtr)
 	{
@@ -1310,7 +1334,7 @@ void FAllocationsProvider::EditAlloc(double Time, uint32 CallstackId, uint64 Add
 		for (uint64 Offset = 1; Offset <= 0x3FFull; ++Offset)
 		{
 			Address = OriginalAddress | (Offset << 54);
-			AllocationPtr = LiveAllocs[RootHeap]->FindRef(Address);
+			AllocationPtr = RootHeap.LiveAllocs->FindRef(Address);
 			if (!AllocationPtr)
 			{
 				break;
@@ -1319,7 +1343,7 @@ void FAllocationsProvider::EditAlloc(double Time, uint32 CallstackId, uint64 Add
 		++AllocWarnings;
 		if (AllocWarnings <= MaxLogMessagesPerWarningType)
 		{
-			UE_LOG(LogTraceServices, Warning, TEXT("[MemAlloc] Invalid ALLOC event remapped (Address=0x%llX --> 0x%llX, Size=%llu, Tag=%u, RootHeap=%u, Time=%f, CallstackId=%u)!"), OriginalAddress, Address, InSize, Tag, RootHeap, Time, CallstackId);
+			UE_LOG(LogTraceServices, Warning, TEXT("[MemAlloc] Invalid ALLOC event remapped (Address=0x%llX --> 0x%llX, Size=%llu, Tag=%u, RootHeap=%u, Time=%f, CallstackId=%u)!"), OriginalAddress, Address, Size, Tag, RootHeapId, Time, CallstackId);
 		}
 	}
 #endif // INSIGHTS_REMAP_INVALID_ALLOC_EVENTS
@@ -1327,7 +1351,7 @@ void FAllocationsProvider::EditAlloc(double Time, uint32 CallstackId, uint64 Add
 
 	if (!AllocationPtr)
 	{
-		AllocationPtr = LiveAllocs[RootHeap]->AddNew(Address);
+		AllocationPtr = RootHeap.LiveAllocs->AddNew(Address);
 		FAllocationItem& Allocation = *AllocationPtr;
 
 		uint32 MetadataId = MetadataProvider.InvalidMetadataId;
@@ -1355,8 +1379,8 @@ void FAllocationsProvider::EditAlloc(double Time, uint32 CallstackId, uint64 Add
 
 		INSIGHTS_SLOW_CHECK(Allocation.Address == Address);
 		INSIGHTS_SLOW_CHECK(InAlignment < 256);
-		Allocation.SizeAndAlignment = FAllocationItem::PackSizeAndAlignment(InSize, static_cast<uint8>(InAlignment));
-		Allocation.StartEventIndex = EventIndex[RootHeap];
+		Allocation.SizeAndAlignment = FAllocationItem::PackSizeAndAlignment(Size, static_cast<uint8>(Alignment));
+		Allocation.StartEventIndex = RootHeap.EventIndex;
 		Allocation.EndEventIndex = (uint32)-1;
 		Allocation.StartTime = Time;
 		Allocation.EndTime = std::numeric_limits<double>::infinity();
@@ -1365,13 +1389,13 @@ void FAllocationsProvider::EditAlloc(double Time, uint32 CallstackId, uint64 Add
 		Allocation.FreeCallstackId = 0; // no callstack yet
 		Allocation.MetadataId = MetadataId;
 		Allocation.Tag = Tag;
-		Allocation.RootHeap = static_cast<uint8>(RootHeap);
+		Allocation.RootHeap = static_cast<uint8>(RootHeapId);
 		Allocation.Flags = EMemoryTraceHeapAllocationFlags::None;
 
-		UpdateHistogramByAllocSize(InSize);
+		UpdateHistogramByAllocSize(Size);
 
 		// Update stats for the current timeline sample.
-		TotalAllocatedMemory += InSize;
+		TotalAllocatedMemory += Size;
 		++TotalLiveAllocations;
 		SampleMaxTotalAllocatedMemory = FMath::Max(SampleMaxTotalAllocatedMemory, TotalAllocatedMemory);
 		SampleMaxLiveAllocations = FMath::Max(SampleMaxLiveAllocations, TotalLiveAllocations);
@@ -1383,15 +1407,15 @@ void FAllocationsProvider::EditAlloc(double Time, uint32 CallstackId, uint64 Add
 		++AllocErrors;
 		if (AllocErrors <= MaxLogMessagesPerErrorType)
 		{
-			UE_LOG(LogTraceServices, Error, TEXT("[MemAlloc] Invalid ALLOC event (Address=0x%llX, Size=%llu, Tag=%u, RootHeap=%u, Time=%f, CallstackId=%u)!"), Address, InSize, Tag, RootHeap, Time, CallstackId);
+			UE_LOG(LogTraceServices, Error, TEXT("[MemAlloc] Invalid ALLOC event (Address=0x%llX, Size=%llu, Tag=%u, RootHeap=%u, Time=%f, CallstackId=%u)!"), Address, Size, Tag, RootHeapId, Time, CallstackId);
 		}
 	}
 #endif
 
 	++AllocCount;
-	if (EventIndex[RootHeap] != ~0u)
+	if (RootHeap.EventIndex != ~0u)
 	{
-		++EventIndex[RootHeap];
+		RootHeap.EventIndex++;
 	}
 	else
 	{
@@ -1402,7 +1426,7 @@ void FAllocationsProvider::EditAlloc(double Time, uint32 CallstackId, uint64 Add
 
 ////////////////////////////////////////////////////////////////////////////////////////////////////
 
-void FAllocationsProvider::EditFree(double Time, uint32 CallstackId, uint64 Address, HeapId RootHeap)
+void FAllocationsProvider::EditFree(double Time, uint32 CallstackId, uint64 Address, HeapId RootHeapId)
 {
 	EditAccessCheck();
 
@@ -1411,46 +1435,47 @@ void FAllocationsProvider::EditFree(double Time, uint32 CallstackId, uint64 Addr
 		return;
 	}
 
-	INSIGHTS_WATCH_API_LOGF(Address, TEXT("Free 0x%llX : RootHeap=%u CallstackId=%u"), Time, Address, RootHeap, CallstackId);
+	INSIGHTS_WATCH_API_LOGF(Address, TEXT("Free 0x%llX : RootHeap=%u CallstackId=%u"), Time, Address, RootHeapId, CallstackId);
 
-	if (Address == 0)
+	if (Address == 0 && RootHeapId == EMemoryTraceRootHeap::SystemMemory)
 	{
 		++FreeWarnings;
 		if (FreeWarnings <= MaxLogMessagesPerWarningType)
 		{
-			UE_LOG(LogTraceServices, Warning, TEXT("[MemAlloc] Free for address 0 : RootHeap=%u, Time=%f, CallstackId=%u"), RootHeap, Time, CallstackId);
+			UE_LOG(LogTraceServices, Warning, TEXT("[MemAlloc] Free for address 0 : RootHeap=%u, Time=%f, CallstackId=%u"), RootHeapId, Time, CallstackId);
 		}
 		return;
 	}
 
-	check(RootHeap < MaxRootHeaps);
-	if (SbTree[RootHeap] == nullptr)
+	if (!IsValidRootHeap(RootHeapId))
 	{
 		++FreeWarnings;
 		if (FreeWarnings <= MaxLogMessagesPerWarningType)
 		{
-			UE_LOG(LogTraceServices, Warning, TEXT("[MemAlloc] Free with invalid root heap id (%u). This heap has not yet been registered."), RootHeap);
+			UE_LOG(LogTraceServices, Warning, TEXT("[MemAlloc] Free with invalid root heap id (%u)."), RootHeapId);
 		}
 		return;
 	}
 
-	SbTree[RootHeap]->SetTimeForEvent(EventIndex[RootHeap], Time);
+	FRootHeap& RootHeap = *RootHeaps[RootHeapId];
+
+	RootHeap.SbTree->SetTimeForEvent(RootHeap.EventIndex, Time);
 
 	AdvanceTimelines(Time);
 
-	FAllocationItem* AllocationPtr = LiveAllocs[RootHeap]->Remove(Address); // we take ownership of AllocationPtr
+	FAllocationItem* AllocationPtr = RootHeap.LiveAllocs->Remove(Address); // we take ownership of AllocationPtr
 	if (!AllocationPtr)
 	{
 		// Try to free a heap allocation.
-		AllocationPtr = LiveAllocs[RootHeap]->FindHeapRef(Address);
+		AllocationPtr = RootHeap.LiveAllocs->FindHeapRef(Address);
 		if (AllocationPtr)
 		{
 			INSIGHTS_SLOW_CHECK(AllocationPtr->IsHeap());
 			HeapId Heap = AllocationPtr->RootHeap;
 			INSIGHTS_WATCH_INDIRECT_API_LOGF(TEXT("UnmarkAllocationAsHeap"), Address, Time);
 			EditUnmarkAllocationAsHeap(Time, CallstackId, Address, Heap);
-			SbTree[RootHeap]->SetTimeForEvent(EventIndex[RootHeap], Time); // for the case where the free event is first event in a new SbTree column after changing the heap alloc
-			AllocationPtr = LiveAllocs[RootHeap]->Remove(Address); // we take ownership of AllocationPtr
+			RootHeap.SbTree->SetTimeForEvent(RootHeap.EventIndex, Time); // for the case where the free event is first event in a new SbTree column after changing the heap alloc
+			AllocationPtr = RootHeap.LiveAllocs->Remove(Address); // we take ownership of AllocationPtr
 		}
 	}
 	else
@@ -1466,7 +1491,7 @@ void FAllocationsProvider::EditFree(double Time, uint32 CallstackId, uint64 Addr
 		for (uint64 Offset = 1; Offset <= 0x3FFull; ++Offset)
 		{
 			Address = OriginalAddress | (Offset << 54);
-			AllocationPtr = LiveAllocs[RootHeap]->FindRef(Address);
+			AllocationPtr = RootHeap.LiveAllocs->FindRef(Address);
 			if (AllocationPtr)
 			{
 				break;
@@ -1475,7 +1500,7 @@ void FAllocationsProvider::EditFree(double Time, uint32 CallstackId, uint64 Addr
 		++AllocWarnings;
 		if (AllocWarnings <= MaxLogMessagesPerWarningType)
 		{
-			UE_LOG(LogTraceServices, Warning, TEXT("[MemAlloc] Invalid FREE event remapped (Address=0x%llX --> 0x%llX, RootHeap=%u, Time=%f, CallstackId=%u)!"), OriginalAddress, Address, RootHeap, Time, CallstackId);
+			UE_LOG(LogTraceServices, Warning, TEXT("[MemAlloc] Invalid FREE event remapped (Address=0x%llX --> 0x%llX, RootHeap=%u, Time=%f, CallstackId=%u)!"), OriginalAddress, Address, RootHeapId, Time, CallstackId);
 		}
 	}
 #endif // INSIGHTS_REMAP_INVALID_FREE_EVENTS
@@ -1485,27 +1510,27 @@ void FAllocationsProvider::EditFree(double Time, uint32 CallstackId, uint64 Addr
 		++FreeErrors;
 		if (FreeErrors <= MaxLogMessagesPerErrorType)
 		{
-			UE_LOG(LogTraceServices, Error, TEXT("[MemAlloc] Invalid FREE event (Address=0x%llX, RootHeap=%u, Time=%f, CallstackId=%u)! A fake alloc will be created with size 0."), Address, RootHeap, Time, CallstackId);
+			UE_LOG(LogTraceServices, Error, TEXT("[MemAlloc] Invalid FREE event (Address=0x%llX, RootHeap=%u, Time=%f, CallstackId=%u)! A fake alloc will be created with size 0."), Address, RootHeapId, Time, CallstackId);
 		}
 		// Fake the missing alloc.
 		constexpr uint64 FakeAllocSize = 0;
 		constexpr uint32 FakeAllocAlignment = 0;
 		INSIGHTS_WATCH_INDIRECT_API_LOGF(TEXT("Alloc"), Address, Time);
-		EditAlloc(Time, CallstackId, Address, FakeAllocSize, FakeAllocAlignment, RootHeap);
-		SbTree[RootHeap]->SetTimeForEvent(EventIndex[RootHeap], Time); // for the case where the free event is first event in a new SbTree column after adding the fake alloc
-		AllocationPtr = LiveAllocs[RootHeap]->Remove(Address); // we take ownership of AllocationPtr
+		EditAlloc(Time, CallstackId, Address, FakeAllocSize, FakeAllocAlignment, RootHeapId);
+		RootHeap.SbTree->SetTimeForEvent(RootHeap.EventIndex, Time); // for the case where the free event is first event in a new SbTree column after adding the fake alloc
+		AllocationPtr = RootHeap.LiveAllocs->Remove(Address); // we take ownership of AllocationPtr
 	}
 	if (AllocationPtr)
 	{
-		check(EventIndex[RootHeap] > AllocationPtr->StartEventIndex);
-		AllocationPtr->EndEventIndex = EventIndex[RootHeap];
+		check(RootHeap.EventIndex > AllocationPtr->StartEventIndex);
+		AllocationPtr->EndEventIndex = RootHeap.EventIndex;
 		AllocationPtr->EndTime = Time;
 
 		AllocationPtr->FreeCallstackId = CallstackId;
 
 		const uint64 OldSize = AllocationPtr->GetSize();
 
-		SbTree[RootHeap]->AddAlloc(AllocationPtr); // SbTree takes ownership of AllocationPtr
+		RootHeap.SbTree->AddAlloc(AllocationPtr); // SbTree takes ownership of AllocationPtr
 
 		uint32 EventDistance = AllocationPtr->EndEventIndex - AllocationPtr->StartEventIndex;
 		UpdateHistogramByEventDistance(EventDistance);
@@ -1525,14 +1550,14 @@ void FAllocationsProvider::EditFree(double Time, uint32 CallstackId, uint64 Addr
 		++FreeErrors;
 		if (FreeErrors <= MaxLogMessagesPerErrorType)
 		{
-			UE_LOG(LogTraceServices, Error, TEXT("[MemAlloc] Invalid FREE event (Address=0x%llX, RootHeap=%u, Time=%f, CallstackId=%u)!"), Address, RootHeap, Time, CallstackId);
+			UE_LOG(LogTraceServices, Error, TEXT("[MemAlloc] Invalid FREE event (Address=0x%llX, RootHeap=%u, Time=%f, CallstackId=%u)!"), Address, RootHeapId, Time, CallstackId);
 		}
 	}
 
 	++FreeCount;
-	if (EventIndex[RootHeap] != ~0u)
+	if (RootHeap.EventIndex != ~0u)
 	{
-		++EventIndex[RootHeap];
+		RootHeap.EventIndex++;
 	}
 	else
 	{
@@ -1547,55 +1572,157 @@ void FAllocationsProvider::EditHeapSpec(HeapId Id, HeapId ParentId, const FStrin
 {
 	EditAccessCheck();
 
-	INSIGHTS_API_LOGF(TEXT("HeapSpec : Id=%u ParentId=%u Name=\"%*s\" Flags=0x%X"), -1.0, Id, ParentId, Name.Len(), Name.GetData(), uint32(Flags));
+	INSIGHTS_API_LOGF(TEXT("HeapSpec : Id=%u ParentId=%u Name=\"%.*s\" Flags=0x%X"), -1.0, Id, ParentId, Name.Len(), Name.GetData(), uint32(Flags));
 
 	AddHeapSpec(Id, ParentId, Name, Flags);
 }
 
 ////////////////////////////////////////////////////////////////////////////////////////////////////
 
+IAllocationsProvider::FHeapSpec& FAllocationsProvider::GetOrCreateHeapSpec(HeapId Id)
+{
+	if (Id < (uint32)HeapSpecs.Num())
+	{
+		if (HeapSpecs[Id] != nullptr)
+		{
+			return *HeapSpecs[Id];
+		}
+	}
+	else
+	{
+		HeapSpecs.AddDefaulted(Id - HeapSpecs.Num() + 1);
+	}
+
+	FHeapSpec* HeapSpec = new FHeapSpec();
+
+	HeapSpec->Id = Id;
+	HeapSpec->Parent = nullptr;
+	HeapSpec->Name = GDefaultHeapName;
+	HeapSpec->Flags = EMemoryTraceHeapFlags::None;
+
+	HeapSpecs[Id] = HeapSpec;
+	return *HeapSpec;
+}
+
+////////////////////////////////////////////////////////////////////////////////////////////////////
+
+FAllocationsProvider::FRootHeap& FAllocationsProvider::FindParentRootHeapUnchecked(HeapId Id) const
+{
+	const FHeapSpec* HeapSpec = HeapSpecs[Id];
+	while (HeapSpec->Parent != nullptr)
+	{
+		HeapSpec = HeapSpec->Parent;
+	}
+	return *RootHeaps[HeapSpec->Id];
+}
+
+////////////////////////////////////////////////////////////////////////////////////////////////////
+
+void FAllocationsProvider::CreateRootHeap(HeapId Id)
+{
+	FRootHeap* RootHeapPtr = new FRootHeap();
+
+	RootHeapPtr->HeapSpec = &GetOrCreateHeapSpec(Id);
+
+	constexpr uint32 ColumnShift = 17; // 1<<17 = 128K
+	RootHeapPtr->SbTree = new FSbTree(Session.GetLinearAllocator(), ColumnShift);
+
+	RootHeapPtr->LiveAllocs = new FLiveAllocCollection();
+
+	RootHeaps[Id] = RootHeapPtr;
+}
+
+////////////////////////////////////////////////////////////////////////////////////////////////////
+
 void FAllocationsProvider::AddHeapSpec(HeapId Id, HeapId ParentId, const FStringView& Name, EMemoryTraceHeapFlags Flags)
 {
+	const TCHAR* HeapName = Session.StoreString(Name);
+
 	if (Id < MaxRootHeaps)
 	{
-		if (Id == 0 && SbTree[Id] != nullptr)
+		FHeapSpec* HeapSpec = HeapSpecs[Id];
+		FRootHeap* RootHeapPtr = RootHeaps[Id];
+
+		if (RootHeapPtr != nullptr)
 		{
+			check(HeapSpec != nullptr)
+			check(HeapSpec->Id == Id);
+
+			check(RootHeapPtr->HeapSpec == HeapSpec);
+			check(RootHeapPtr->SbTree != nullptr);
+			check(RootHeapPtr->LiveAllocs != nullptr);
+
 			// "System" root heap is already created. See class constructor.
-			check(LiveAllocs[Id] != nullptr);
+			if (Id != 0)
+			{
+				++HeapWarnings;
+				if (HeapWarnings <= MaxLogMessagesPerWarningType)
+				{
+					UE_LOG(LogTraceServices, Warning, TEXT("[MemAlloc] Root heap %u has changed from (\"%s\", flags=%u) to (\"%s\", flags=%u)."),
+						Id,
+						HeapSpec->Name,
+						HeapSpec->Flags,
+						HeapName,
+						Flags);
+				}
+			}
 		}
 		else
 		{
-			check(SbTree[Id] == nullptr && LiveAllocs[Id] == nullptr);
+			check(HeapSpec == nullptr);
+
+			CreateRootHeap(Id);
 
-			constexpr uint32 ColumnShift = 17; // 1<<17 = 128K
-			SbTree[Id] = new FSbTree(Session.GetLinearAllocator(), ColumnShift);
-			LiveAllocs[Id] = new FLiveAllocCollection();
+			HeapSpec = HeapSpecs[Id];
+			check(HeapSpec != nullptr);
 		}
 
-		FHeapSpec& RootHeapSpec = HeapSpecs[Id];
-		RootHeapSpec.Name = Session.StoreString(Name);
-		RootHeapSpec.Parent = nullptr;
-		RootHeapSpec.Id = Id;
-		RootHeapSpec.Flags = Flags;
+		HeapSpec->Name = HeapName;
+		HeapSpec->Flags = Flags;
 	}
 	else
 	{
-		FHeapSpec& ParentSpec = HeapSpecs[ParentId];
-		if (ParentSpec.Name == nullptr)
+		FHeapSpec& ParentHeapSpec = GetOrCreateHeapSpec(ParentId);
+		if (ParentHeapSpec.Name == GDefaultHeapName)
 		{
+			if (ParentId < MaxRootHeaps && RootHeaps[ParentId] == nullptr)
+			{
+				CreateRootHeap(ParentId);
+			}
+
 			++HeapWarnings;
 			if (HeapWarnings <= MaxLogMessagesPerWarningType)
 			{
-				UE_LOG(LogTraceServices, Warning, TEXT("[MemAlloc] Parent heap id (%u) used before it was announced for heap %s."), ParentId, *FString(Name));
+				UE_LOG(LogTraceServices, Warning, TEXT("[MemAlloc] Heap %u (\"%s\") is used before its parent heap %u was announced."), Id, HeapName, ParentId);
+			}
+		}
+
+		if (IsValidHeap(Id))
+		{
+			FHeapSpec& PreviousHeap = GetHeapSpecUnchecked(Id);
+			if (PreviousHeap.Name != GDefaultHeapName)
+			{
+				++HeapWarnings;
+				if (HeapWarnings <= MaxLogMessagesPerWarningType)
+				{
+					UE_LOG(LogTraceServices, Warning, TEXT("[MemAlloc] Heap %u has changed from (\"%s\", parent=\"%s\", flags=%u) to (\"%s\", parent=\"%s\", flags=%u)."),
+						Id,
+						PreviousHeap.Name,
+						PreviousHeap.Parent && PreviousHeap.Parent->Name ? PreviousHeap.Parent->Name : TEXT(""),
+						PreviousHeap.Flags,
+						HeapName,
+						ParentHeapSpec.Name ? ParentHeapSpec.Name : TEXT(""),
+						Flags);
+				}
 			}
 		}
-		FHeapSpec& NewSpec = HeapSpecs[Id];
-		NewSpec.Name = Session.StoreString(Name);
-		NewSpec.Parent = &ParentSpec;
-		NewSpec.Id = Id;
-		NewSpec.Flags = Flags;
 
-		ParentSpec.Children.Add(&NewSpec);
+		FHeapSpec& HeapSpec = GetOrCreateHeapSpec(Id);
+		HeapSpec.Parent = &ParentHeapSpec;
+		HeapSpec.Name = HeapName;
+		HeapSpec.Flags = Flags;
+
+		ParentHeapSpec.Children.Add(&HeapSpec);
 	}
 }
 
@@ -1607,7 +1734,7 @@ void FAllocationsProvider::EditMarkAllocationAsHeap(double Time, uint32 Callstac
 
 	INSIGHTS_WATCH_API_LOGF(Address, TEXT("MarkAllocAsHeap 0x%llX : Heap=%u Flags=0x%X CallstackId=%u"), Time, Address, Heap, uint32(Flags), CallstackId);
 
-	if (Heap >= (uint32)HeapSpecs.Num())
+	if (!IsValidHeap(Heap))
 	{
 		++HeapErrors;
 		if (HeapErrors <= MaxLogMessagesPerErrorType)
@@ -1616,12 +1743,12 @@ void FAllocationsProvider::EditMarkAllocationAsHeap(double Time, uint32 Callstac
 		}
 		return;
 	}
-	const FHeapSpec& HeapSpec = HeapSpecs[Heap];
 
-	const HeapId RootHeap = FindRootHeap(Heap);
+	const FHeapSpec& HeapSpec = GetHeapSpecUnchecked(Heap);
+	const FRootHeap& RootHeap = FindParentRootHeapUnchecked(Heap);
+
 #if 0 // TODO
-	ensure(RootHeap != Heap);
-	if (Heap == RootHeap)
+	if (Heap == RootHeap.HeapSpec->Id)
 	{
 		++HeapWarnings;
 		if (HeapWarnings <= MaxLogMessagesPerWarningType)
@@ -1632,7 +1759,7 @@ void FAllocationsProvider::EditMarkAllocationAsHeap(double Time, uint32 Callstac
 #endif
 
 	// Remove the allocation from the Live allocs.
-	FAllocationItem* Alloc = LiveAllocs[RootHeap]->Remove(Address); // we take ownership of Alloc
+	FAllocationItem* Alloc = RootHeap.LiveAllocs->Remove(Address); // we take ownership of Alloc
 	if (Alloc)
 	{
 		check(Address == Alloc->Address);
@@ -1652,7 +1779,7 @@ void FAllocationsProvider::EditMarkAllocationAsHeap(double Time, uint32 Callstac
 		Alloc->RootHeap = static_cast<uint8>(Heap);
 
 		// Re-add it to the Live allocs as a heap allocation.
-		LiveAllocs[RootHeap]->AddHeap(Alloc); // the Live allocs takes ownership of Alloc
+		RootHeap.LiveAllocs->AddHeap(Alloc); // the Live allocs takes ownership of Alloc
 
 		// Update stats. Remove this allocation from the total.
 		TotalAllocatedMemory -= Alloc->GetSize();
@@ -1680,7 +1807,7 @@ void FAllocationsProvider::EditUnmarkAllocationAsHeap(double Time, uint32 Callst
 
 	INSIGHTS_WATCH_API_LOGF(Address, TEXT("UnmarkAllocAsHeap 0x%llX : Heap=%u CallstackId=%u"), Time, Address, Heap, CallstackId);
 
-	if (Heap >= (uint32)HeapSpecs.Num())
+	if (!IsValidHeap(Heap))
 	{
 		++HeapErrors;
 		if (HeapErrors <= MaxLogMessagesPerErrorType)
@@ -1689,12 +1816,12 @@ void FAllocationsProvider::EditUnmarkAllocationAsHeap(double Time, uint32 Callst
 		}
 		return;
 	}
-	const FHeapSpec& HeapSpec = HeapSpecs[Heap];
 
-	const HeapId RootHeap = FindRootHeap(Heap);
+	const FHeapSpec& HeapSpec = GetHeapSpecUnchecked(Heap);
+	const FRootHeap& RootHeap = FindParentRootHeapUnchecked(Heap);
+
 #if 0 // TODO
-	ensure(RootHeap != Heap);
-	if (Heap == RootHeap)
+	if (Heap == RootHeap.HeapSpec->Id)
 	{
 		++HeapWarnings;
 		if (HeapWarnings <= MaxLogMessagesPerWarningType)
@@ -1705,7 +1832,7 @@ void FAllocationsProvider::EditUnmarkAllocationAsHeap(double Time, uint32 Callst
 #endif
 
 	// Remove the heap allocation from the Live allocs.
-	FAllocationItem* Alloc = LiveAllocs[RootHeap]->RemoveHeap(Address); // we take ownership of Alloc
+	FAllocationItem* Alloc = RootHeap.LiveAllocs->RemoveHeap(Address); // we take ownership of Alloc
 	if (Alloc)
 	{
 		check(Address == Alloc->Address);
@@ -1735,7 +1862,7 @@ void FAllocationsProvider::EditUnmarkAllocationAsHeap(double Time, uint32 Callst
 			TArray<FAllocInHeap> AllocsInHeap;
 			const uint64 HeapSize = Alloc->GetSize();
 			const uint64 EndAddress = Alloc->Address + HeapSize;
-			LiveAllocs[RootHeap]->Enumerate(Address, EndAddress, [&AllocatedSizeInHeap, &AllocsInHeap](const FAllocationItem& ChildAlloc)
+			RootHeap.LiveAllocs->Enumerate(Address, EndAddress, [&AllocatedSizeInHeap, &AllocsInHeap](const FAllocationItem& ChildAlloc)
 			{
 				const uint64 ChildAllocSize = ChildAlloc.GetSize();
 				AllocatedSizeInHeap += ChildAllocSize;
@@ -1777,7 +1904,7 @@ void FAllocationsProvider::EditUnmarkAllocationAsHeap(double Time, uint32 Callst
 				for (const FAllocInHeap& AllocInHeap : AllocsInHeap)
 				{
 					INSIGHTS_WATCH_INDIRECT_API_LOGF(TEXT("Free"), AllocInHeap.Address, Time);
-					EditFree(Time, CallstackId, AllocInHeap.Address, RootHeap);
+					EditFree(Time, CallstackId, AllocInHeap.Address, RootHeap.HeapSpec->Id);
 				}
 			}
 		}
@@ -1787,16 +1914,16 @@ void FAllocationsProvider::EditUnmarkAllocationAsHeap(double Time, uint32 Callst
 		const uint32 AllocCallstackId = Alloc->CallstackId;
 
 		// Re-add this allocation to the Live allocs.
-		LiveAllocs[RootHeap]->Add(Alloc); // the Live allocs takes ownership of Alloc
+		RootHeap.LiveAllocs->Add(Alloc); // the Live allocs takes ownership of Alloc
 
 		// We cannot just unmark the allocation as heap, there is no timestamp support, instead fake a "free"
 		// event and an "alloc" event. Make sure the new allocation retains the tag from the original.
 		CurrentTracker = 1;
 		TagTracker.PushTagFromPtr(CurrentSystemThreadId, CurrentTracker, Alloc->Tag);
 		INSIGHTS_WATCH_INDIRECT_API_LOGF(TEXT("Free"), Address, Time);
-		EditFree(Time, CallstackId, Address, RootHeap);
+		EditFree(Time, CallstackId, Address, RootHeap.HeapSpec->Id);
 		INSIGHTS_WATCH_INDIRECT_API_LOGF(TEXT("Alloc"), Address, Time);
-		EditAlloc(Time, AllocCallstackId, Address, Size, Alignment, RootHeap);
+		EditAlloc(Time, AllocCallstackId, Address, Size, Alignment, RootHeap.HeapSpec->Id);
 		TagTracker.PopTagFromPtr(CurrentSystemThreadId, CurrentTracker);
 		CurrentTracker = 0;
 	}
@@ -1898,18 +2025,6 @@ void FAllocationsProvider::AdvanceTimelines(double Time)
 
 ////////////////////////////////////////////////////////////////////////////////////////////////////
 
-HeapId FAllocationsProvider::FindRootHeap(HeapId Heap) const
-{
-	const FHeapSpec* RootHeap = &HeapSpecs[Heap];
-	while (RootHeap->Parent != nullptr)
-	{
-		RootHeap = RootHeap->Parent;
-	}
-	return RootHeap->Id;
-}
-
-////////////////////////////////////////////////////////////////////////////////////////////////////
-
 void FAllocationsProvider::EditPushTag(uint32 ThreadId, uint8 Tracker, TagIdType Tag)
 {
 	EditAccessCheck();
@@ -1932,8 +2047,22 @@ void FAllocationsProvider::EditPushTagFromPtr(uint32 ThreadId, uint8 Tracker, ui
 {
 	EditAccessCheck();
 
-	FLiveAllocCollection* Allocs = LiveAllocs[RootHeapId];
-	FAllocationItem* Alloc = Allocs ? Allocs->FindRef(Ptr) : nullptr;
+	const FAllocationItem* Alloc = nullptr;
+
+	if (IsValidRootHeap(RootHeapId))
+	{
+		const FRootHeap& RootHeap = GetRootHeapUnchecked(RootHeapId);
+		Alloc = RootHeap.LiveAllocs->FindRef(Ptr);
+	}
+	else
+	{
+		++MiscErrors;
+		if (MiscErrors <= MaxLogMessagesPerErrorType)
+		{
+			UE_LOG(LogTraceServices, Error, TEXT("[MemAlloc] Invalid root heap (%u) for MemoryScopePtr or ReallocFree event!"), RootHeapId);
+		}
+	}
+
 	const TagIdType Tag = Alloc ? Alloc->Tag : 0; // If ptr is not found use "Untagged"
 	TagTracker.PushTagFromPtr(ThreadId, Tracker, Tag);
 
@@ -1942,7 +2071,7 @@ void FAllocationsProvider::EditPushTagFromPtr(uint32 ThreadId, uint8 Tracker, ui
 		++MiscErrors;
 		if (MiscErrors <= MaxLogMessagesPerErrorType)
 		{
-			UE_LOG(LogTraceServices, Error, TEXT("[MemAlloc] Invalid Ptr (0x%llX) for MemoryScopePtr or ReallocFree event!"), Ptr);
+			UE_LOG(LogTraceServices, Error, TEXT("[MemAlloc] Invalid address (0x%llX) for MemoryScopePtr or ReallocFree event!"), Ptr);
 		}
 	}
 }
@@ -2011,12 +2140,12 @@ void FAllocationsProvider::EditOnAnalysisCompleted(double Time)
 	DebugPrint();
 #endif
 
-	for (uint32 RootHeap = 0; RootHeap < MaxRootHeaps; ++RootHeap)
+	for (uint32 RootHeapIdx = 0; RootHeapIdx < MaxRootHeaps; ++RootHeapIdx)
 	{
-		const FSbTree* Tree = SbTree[RootHeap];
-		if (Tree)
+		FRootHeap* RootHeap = RootHeaps[RootHeapIdx];
+		if (RootHeap)
 		{
-			Tree->Validate();
+			RootHeap->SbTree->Validate();
 		}
 	}
 
@@ -2051,22 +2180,64 @@ void FAllocationsProvider::EditOnAnalysisCompleted(double Time)
 	}
 
 	uint64 TotalEventCount = 0;
-	for (uint32 RootHeap = 0; RootHeap < MaxRootHeaps; ++RootHeap)
+	for (uint32 RootHeapIdx = 0; RootHeapIdx < MaxRootHeaps; ++RootHeapIdx)
+	{
+		FRootHeap* RootHeap = RootHeaps[RootHeapIdx];
+		if (RootHeap)
+		{
+			TotalEventCount += RootHeap->EventIndex;
+		}
+	}
+	uint32 NumHeapSpecs = 0;
+	for (FHeapSpec* HeapSpec : HeapSpecs)
 	{
-		TotalEventCount += EventIndex[RootHeap];
+		if (HeapSpec)
+		{
+			++NumHeapSpecs;
+		}
 	}
-	UE_LOG(LogTraceServices, Log, TEXT("[MemAlloc] Analysis Completed (%llu events, %llu allocs, %llu frees, %llu heaps)"), TotalEventCount, AllocCount, FreeCount, HeapCount);
+	UE_LOG(LogTraceServices, Log, TEXT("[MemAlloc] Analysis Completed (%llu events, %llu allocs, %llu frees, %llu heaps, %d heap specs)"), TotalEventCount, AllocCount, FreeCount, HeapCount, NumHeapSpecs);
 }
 
 ////////////////////////////////////////////////////////////////////////////////////////////////////
 
 void FAllocationsProvider::EnumerateRootHeaps(TFunctionRef<void(HeapId Id, const FHeapSpec&)> Callback) const
 {
-	for (const FHeapSpec& Spec : HeapSpecs)
+	ReadAccessCheck();
+
+	for (uint32 RootHeapIdx = 0; RootHeapIdx < MaxRootHeaps; ++RootHeapIdx)
 	{
-		if (Spec.Parent == nullptr && Spec.Name != nullptr)
+		FRootHeap* RootHeap = RootHeaps[RootHeapIdx];
+		if (RootHeap)
 		{
-			Callback(Spec.Id, Spec);
+			check(RootHeap->HeapSpec != nullptr);
+			const FHeapSpec& HeapSpec = *RootHeap->HeapSpec;
+			check(HeapSpec.Parent == nullptr);
+			if (HeapSpec.Name != nullptr)
+			{
+				Callback(HeapSpec.Id, HeapSpec);
+			}
+		}
+	}
+}
+
+////////////////////////////////////////////////////////////////////////////////////////////////////
+
+void FAllocationsProvider::EnumerateHeaps(TFunctionRef<void(HeapId Id, const FHeapSpec&)> Callback) const
+{
+	ReadAccessCheck();
+
+	const uint32 NumHeapSpecs = HeapSpecs.Num();
+	for (uint32 HeapSpecIdx = 0; HeapSpecIdx < NumHeapSpecs; ++HeapSpecIdx)
+	{
+		const FHeapSpec* HeapSpecPtr = HeapSpecs[HeapSpecIdx];
+		if (HeapSpecPtr)
+		{
+			const FHeapSpec& HeapSpec = *HeapSpecPtr;
+			if (HeapSpec.Name != nullptr)
+			{
+				Callback(HeapSpec.Id, HeapSpec);
+			}
 		}
 	}
 }
@@ -2337,12 +2508,12 @@ void FAllocationsProvider::EnumerateTags(TFunctionRef<void(const TCHAR*, const T
 
 void FAllocationsProvider::DebugPrint() const
 {
-	for (uint32 RootHeap = 0; RootHeap < MaxRootHeaps; ++RootHeap)
+	for (uint32 RootHeapIdx = 0; RootHeapIdx < MaxRootHeaps; ++RootHeapIdx)
 	{
-		const FSbTree* Tree = SbTree[RootHeap];
-		if (Tree)
+		FRootHeap* RootHeap = RootHeaps[RootHeapIdx];
+		if (RootHeap)
 		{
-			Tree->DebugPrint();
+			RootHeap->SbTree->DebugPrint();
 		}
 	}
 }
@@ -2376,11 +2547,13 @@ const IAllocationsProvider::FQueryStatus FAllocationsProvider::PollQuery(FQueryH
 void FAllocationsProvider::EnumerateLiveAllocs(TFunctionRef<void(const FAllocationItem& Alloc)> Callback) const
 {
 	ReadAccessCheck();
-	for (const FLiveAllocCollection* Allocs : LiveAllocs)
+
+	for (uint32 RootHeapIdx = 0; RootHeapIdx < MaxRootHeaps; ++RootHeapIdx)
 	{
-		if (Allocs)
+		FRootHeap* RootHeap = RootHeaps[RootHeapIdx];
+		if (RootHeap)
 		{
-			Allocs->Enumerate(Callback);
+			RootHeap->LiveAllocs->Enumerate(Callback);
 		}
 	}
 }
diff --git a/Engine/Source/Developer/TraceServices/Private/Model/AllocationsProvider.h b/Engine/Source/Developer/TraceServices/Private/Model/AllocationsProvider.h
--- a/Engine/Source/Developer/TraceServices/Private/Model/AllocationsProvider.h
+++ b/Engine/Source/Developer/TraceServices/Private/Model/AllocationsProvider.h
@@ -266,6 +266,17 @@ class FAllocationsProvider : public IAllocationsProvider
 private:
 	static constexpr double DefaultTimelineSampleGranularity = 0.0001; // 0.1ms
 
+	// Number of supported root heaps
+	static constexpr uint32 MaxRootHeaps = 16;
+
+	struct FRootHeap
+	{
+		FHeapSpec* HeapSpec = nullptr;
+		FLiveAllocCollection* LiveAllocs = nullptr;
+		FSbTree* SbTree = nullptr;
+		uint32 EventIndex = 0;
+	};
+
 public:
 	explicit FAllocationsProvider(IAnalysisSession& InSession, FMetadataProvider& InMetadataProvider);
 	virtual ~FAllocationsProvider();
@@ -282,8 +293,16 @@ public:
 
 	virtual bool IsInitialized() const override { ReadAccessCheck(); return bInitialized; }
 
-	virtual int32 GetTimelineNumPoints() const override { ReadAccessCheck(); return static_cast<int32>(Timeline.Num()); }
+	virtual void EnumerateTags(TFunctionRef<void(const TCHAR*, const TCHAR*, TagIdType, TagIdType)> Callback) const override;
+
+	virtual const TCHAR* GetTagName(TagIdType Tag) const override { ReadAccessCheck(); return TagTracker.GetTagString(Tag); }
+	virtual const TCHAR* GetTagFullPath(TagIdType Tag) const override { ReadAccessCheck(); return TagTracker.GetTagFullPath(Tag); }
+	bool HasTagFromPtrScope(uint32 ThreadId, uint8 Tracker) const { ReadAccessCheck(); return TagTracker.HasTagFromPtrScope(ThreadId, Tracker); }
+
 	virtual void EnumerateRootHeaps(TFunctionRef<void(HeapId Id, const FHeapSpec&)> Callback) const override;
+	virtual void EnumerateHeaps(TFunctionRef<void(HeapId Id, const FHeapSpec&)> Callback) const override;
+
+	virtual int32 GetTimelineNumPoints() const override { ReadAccessCheck(); return static_cast<int32>(Timeline.Num()); }
 	virtual void GetTimelineIndexRange(double StartTime, double EndTime, int32& StartIndex, int32& EndIndex) const override;
 	virtual void EnumerateMinTotalAllocatedMemoryTimeline(int32 StartIndex, int32 EndIndex, TFunctionRef<void(double Time, double Duration, uint64 Value)> Callback) const override;
 	virtual void EnumerateMaxTotalAllocatedMemoryTimeline(int32 StartIndex, int32 EndIndex, TFunctionRef<void(double Time, double Duration, uint64 Value)> Callback) const override;
@@ -291,21 +310,16 @@ public:
 	virtual void EnumerateMaxLiveAllocationsTimeline(int32 StartIndex, int32 EndIndex, TFunctionRef<void(double Time, double Duration, uint32 Value)> Callback) const override;
 	virtual void EnumerateAllocEventsTimeline(int32 StartIndex, int32 EndIndex, TFunctionRef<void(double Time, double Duration, uint32 Value)> Callback) const override;
 	virtual void EnumerateFreeEventsTimeline(int32 StartIndex, int32 EndIndex, TFunctionRef<void(double Time, double Duration, uint32 Value)> Callback) const override;
-	virtual void EnumerateTags(TFunctionRef<void(const TCHAR*, const TCHAR*, TagIdType, TagIdType)> Callback) const override;
 
 	virtual FQueryHandle StartQuery(const FQueryParams& Params) const override;
 	virtual void CancelQuery(FQueryHandle Query) const override;
 	virtual const FQueryStatus PollQuery(FQueryHandle Query) const override;
 
-	const FSbTree* GetSbTreeUnchecked(HeapId Heap) const { ReadAccessCheck(); return SbTree[Heap]; }
+	const FSbTree* GetSbTreeUnchecked(HeapId Heap) const { ReadAccessCheck(); return RootHeaps[Heap]->SbTree; }
 
 	void EnumerateLiveAllocs(TFunctionRef<void(const FAllocationItem& Alloc)> Callback) const;
 	uint32 GetNumLiveAllocs() const;
 
-	virtual const TCHAR* GetTagName(TagIdType Tag) const override { ReadAccessCheck(); return TagTracker.GetTagString(Tag); }
-	virtual const TCHAR* GetTagFullPath(TagIdType Tag) const override { ReadAccessCheck(); return TagTracker.GetTagFullPath(Tag); }
-	bool HasTagFromPtrScope(uint32 ThreadId, uint8 Tracker) const { ReadAccessCheck(); return TagTracker.HasTagFromPtrScope(ThreadId, Tracker); }
-
 	void DebugPrint() const;
 
 	//////////////////////////////////////////////////
@@ -341,7 +355,31 @@ private:
 	void UpdateHistogramByEventDistance(uint32 EventDistance);
 	void AdvanceTimelines(double Time);
 	void AddHeapSpec(HeapId Id, HeapId ParentId, const FStringView& Name, EMemoryTraceHeapFlags Flags);
-	HeapId FindRootHeap(HeapId Heap) const;
+
+	bool IsValidHeap(HeapId HeapId) const
+	{
+		return HeapId < (uint32)HeapSpecs.Num() && HeapSpecs[HeapId] != nullptr;
+	}
+
+	FHeapSpec& GetHeapSpecUnchecked(HeapId Id) const
+	{
+		return *HeapSpecs[Id];
+	}
+
+	FHeapSpec& GetOrCreateHeapSpec(HeapId Id);
+
+	bool IsValidRootHeap(HeapId RootHeapId) const
+	{
+		return RootHeapId < MaxRootHeaps && RootHeaps[RootHeapId] != nullptr;
+	}
+
+	FRootHeap& GetRootHeapUnchecked(HeapId RootHeapId) const
+	{
+		return *RootHeaps[RootHeapId];
+	}
+
+	FRootHeap& FindParentRootHeapUnchecked(HeapId HeapId) const;
+	void CreateRootHeap(HeapId Id);
 
 private:
 	mutable FProviderLock Lock;
@@ -349,9 +387,6 @@ private:
 	IAnalysisSession& Session;
 	FMetadataProvider& MetadataProvider;
 
-	// Number of supported root heaps
-	constexpr static uint8 MaxRootHeaps = 16;
-
 	double InitTime = 0;
 	uint8 MinAlignment = 0;
 	uint8 SizeShift = 0;
@@ -364,11 +399,8 @@ private:
 	FTagTracker TagTracker;
 	uint8 CurrentTracker = 0;
 
-	TArray<FHeapSpec> HeapSpecs;
-
-	uint32 EventIndex[MaxRootHeaps] = { 0 };
-	FSbTree* SbTree[MaxRootHeaps] = { nullptr };
-	FLiveAllocCollection* LiveAllocs[MaxRootHeaps] = { nullptr };
+	TArray<FHeapSpec*> HeapSpecs;
+	FRootHeap* RootHeaps[MaxRootHeaps] = { nullptr };
 
 	uint64 AllocCount = 0;
 	uint64 FreeCount = 0;
diff --git a/Engine/Source/Developer/TraceServices/Public/TraceServices/Model/AllocationsProvider.h b/Engine/Source/Developer/TraceServices/Public/TraceServices/Model/AllocationsProvider.h
--- a/Engine/Source/Developer/TraceServices/Public/TraceServices/Model/AllocationsProvider.h
+++ b/Engine/Source/Developer/TraceServices/Public/TraceServices/Model/AllocationsProvider.h
@@ -123,6 +123,17 @@ public:
 
 	virtual bool IsInitialized() const = 0;
 
+	// Enumerates the discovered tags.
+	virtual void EnumerateTags(TFunctionRef<void(const TCHAR*, const TCHAR*, TagIdType, TagIdType)> Callback) const = 0;
+
+	// Returns the display name of the specified LLM tag.
+	// Lifetime of returned string matches the session lifetime.
+	virtual const TCHAR* GetTagName(TagIdType Tag) const = 0;
+	virtual const TCHAR* GetTagFullPath(TagIdType Tag) const = 0;
+
+	virtual void EnumerateRootHeaps(TFunctionRef<void(HeapId Id, const FHeapSpec&)> Callback) const = 0;
+	virtual void EnumerateHeaps(TFunctionRef<void(HeapId Id, const FHeapSpec&)> Callback) const = 0;
+
 	// Returns the number of points in each timeline (Min/Max Total Allocated Memory, Min/Max Live Allocations, Total Alloc Events, Total Free Events).
 	virtual int32 GetTimelineNumPoints() const = 0;
 
@@ -130,8 +141,6 @@ public:
 	// Index values are in range { -1, 0, .. , N-1, N }, where N = GetTimelineNumPoints().
 	virtual void GetTimelineIndexRange(double StartTime, double EndTime, int32& StartIndex, int32& EndIndex) const = 0;
 
-	virtual void EnumerateRootHeaps(TFunctionRef<void(HeapId Id, const FHeapSpec&)> Callback) const = 0;
-
 	// Enumerates the Min Total Allocated Memory timeline points in the inclusive index interval [StartIndex, EndIndex].
 	virtual void EnumerateMinTotalAllocatedMemoryTimeline(int32 StartIndex, int32 EndIndex, TFunctionRef<void(double Time, double Duration, uint64 Value)> Callback) const = 0;
 
@@ -150,17 +159,9 @@ public:
 	// Enumerates the Free Events timeline points in the inclusive index interval [StartIndex, EndIndex].
 	virtual void EnumerateFreeEventsTimeline(int32 StartIndex, int32 EndIndex, TFunctionRef<void(double Time, double Duration, uint32 Value)> Callback) const = 0;
 
-	// Enumerates the discovered tags.
-	virtual void EnumerateTags(TFunctionRef<void(const TCHAR*, const TCHAR*, TagIdType, TagIdType)> Callback) const = 0;
-
 	virtual FQueryHandle StartQuery(const FQueryParams& Params) const = 0;
 	virtual void CancelQuery(FQueryHandle Query) const = 0;
 	virtual const FQueryStatus PollQuery(FQueryHandle Query) const = 0;
-
-	// Returns the display name of the specified LLM tag.
-	// Lifetime of returned string matches the session lifetime.
-	virtual const TCHAR* GetTagName(TagIdType Tag) const = 0;
-	virtual const TCHAR* GetTagFullPath(TagIdType Tag) const = 0;
 };
 
 TRACESERVICES_API FName GetAllocationsProviderName();
