[Insights] AllocationsProvider: Added more debug checks for using read API.

diff --git a/Engine/Source/Developer/TraceServices/Private/Model/AllocationsProvider.cpp b/Engine/Source/Developer/TraceServices/Private/Model/AllocationsProvider.cpp
--- a/Engine/Source/Developer/TraceServices/Private/Model/AllocationsProvider.cpp
+++ b/Engine/Source/Developer/TraceServices/Private/Model/AllocationsProvider.cpp
@@ -415,6 +415,7 @@ uint32 IAllocationsProvider::FAllocation::GetAlignment() const
 }
 
 ////////////////////////////////////////////////////////////////////////////////////////////////////
+
 uint32 IAllocationsProvider::FAllocation::GetThreadId() const
 {
 	const auto* Inner = (const FAllocationItem*)this;
@@ -422,6 +423,7 @@ uint32 IAllocationsProvider::FAllocation::GetThreadId() const
 }
 
 ////////////////////////////////////////////////////////////////////////////////////////////////////
+
 uint32 IAllocationsProvider::FAllocation::GetCallstackId() const
 {
 	const auto* Inner = (const FAllocationItem*)this;
@@ -429,6 +431,7 @@ uint32 IAllocationsProvider::FAllocation::GetCallstackId() const
 }
 
 ////////////////////////////////////////////////////////////////////////////////////////////////////
+
 uint32 IAllocationsProvider::FAllocation::GetFreeCallstackId() const
 {
 	const auto* Inner = (const FAllocationItem*)this;
@@ -459,6 +462,7 @@ HeapId IAllocationsProvider::FAllocation::GetRootHeap() const
 }
 
 ////////////////////////////////////////////////////////////////////////////////////////////////////
+
 bool IAllocationsProvider::FAllocation::IsHeap() const
 {
 	const auto* Inner = (const FAllocationItem*)this;
@@ -2297,6 +2301,8 @@ void FAllocationsProvider::GetTimelineIndexRange(double StartTime, double EndTim
 
 void FAllocationsProvider::EnumerateMinTotalAllocatedMemoryTimeline(int32 StartIndex, int32 EndIndex, TFunctionRef<void(double Time, double Duration, uint64 Value)> Callback) const
 {
+	ReadAccessCheck();
+
 	const int32 NumPoints = static_cast<int32>(Timeline.Num());
 	StartIndex = FMath::Max(StartIndex, 0);
 	EndIndex = FMath::Min(EndIndex + 1, NumPoints); // make it exclusive
@@ -2331,6 +2337,8 @@ void FAllocationsProvider::EnumerateMinTotalAllocatedMemoryTimeline(int32 StartI
 
 void FAllocationsProvider::EnumerateMaxTotalAllocatedMemoryTimeline(int32 StartIndex, int32 EndIndex, TFunctionRef<void(double Time, double Duration, uint64 Value)> Callback) const
 {
+	ReadAccessCheck();
+
 	const int32 NumPoints = static_cast<int32>(Timeline.Num());
 	StartIndex = FMath::Max(StartIndex, 0);
 	EndIndex = FMath::Min(EndIndex + 1, NumPoints); // make it exclusive
@@ -2365,6 +2373,8 @@ void FAllocationsProvider::EnumerateMaxTotalAllocatedMemoryTimeline(int32 StartI
 
 void FAllocationsProvider::EnumerateMinLiveAllocationsTimeline(int32 StartIndex, int32 EndIndex, TFunctionRef<void(double Time, double Duration, uint32 Value)> Callback) const
 {
+	ReadAccessCheck();
+
 	const int32 NumPoints = static_cast<int32>(Timeline.Num());
 	StartIndex = FMath::Max(StartIndex, 0);
 	EndIndex = FMath::Min(EndIndex + 1, NumPoints); // make it exclusive
@@ -2399,6 +2409,8 @@ void FAllocationsProvider::EnumerateMinLiveAllocationsTimeline(int32 StartIndex,
 
 void FAllocationsProvider::EnumerateMaxLiveAllocationsTimeline(int32 StartIndex, int32 EndIndex, TFunctionRef<void(double Time, double Duration, uint32 Value)> Callback) const
 {
+	ReadAccessCheck();
+
 	const int32 NumPoints = static_cast<int32>(Timeline.Num());
 	StartIndex = FMath::Max(StartIndex, 0);
 	EndIndex = FMath::Min(EndIndex + 1, NumPoints); // make it exclusive
@@ -2433,6 +2445,8 @@ void FAllocationsProvider::EnumerateMaxLiveAllocationsTimeline(int32 StartIndex,
 
 void FAllocationsProvider::EnumerateAllocEventsTimeline(int32 StartIndex, int32 EndIndex, TFunctionRef<void(double Time, double Duration, uint32 Value)> Callback) const
 {
+	ReadAccessCheck();
+
 	const int32 NumPoints = static_cast<int32>(Timeline.Num());
 	StartIndex = FMath::Max(StartIndex, 0);
 	EndIndex = FMath::Min(EndIndex + 1, NumPoints); // make it exclusive
@@ -2467,6 +2481,8 @@ void FAllocationsProvider::EnumerateAllocEventsTimeline(int32 StartIndex, int32
 
 void FAllocationsProvider::EnumerateFreeEventsTimeline(int32 StartIndex, int32 EndIndex, TFunctionRef<void(double Time, double Duration, uint32 Value)> Callback) const
 {
+	ReadAccessCheck();
+
 	const int32 NumPoints = static_cast<int32>(Timeline.Num());
 	StartIndex = FMath::Max(StartIndex, 0);
 	EndIndex = FMath::Min(EndIndex + 1, NumPoints); // make it exclusive
@@ -2501,6 +2517,8 @@ void FAllocationsProvider::EnumerateFreeEventsTimeline(int32 StartIndex, int32 E
 
 void FAllocationsProvider::EnumerateTags(TFunctionRef<void(const TCHAR*, const TCHAR*, TagIdType, TagIdType)> Callback) const
 {
+	ReadAccessCheck();
+
 	TagTracker.EnumerateTags(Callback);
 }
 
@@ -2508,6 +2526,8 @@ void FAllocationsProvider::EnumerateTags(TFunctionRef<void(const TCHAR*, const T
 
 void FAllocationsProvider::DebugPrint() const
 {
+	ReadAccessCheck();
+
 	for (uint32 RootHeapIdx = 0; RootHeapIdx < MaxRootHeaps; ++RootHeapIdx)
 	{
 		FRootHeap* RootHeap = RootHeaps[RootHeapIdx];
@@ -2563,6 +2583,7 @@ void FAllocationsProvider::EnumerateLiveAllocs(TFunctionRef<void(const FAllocati
 uint32 FAllocationsProvider::GetNumLiveAllocs() const
 {
 	ReadAccessCheck();
+
 	return TotalLiveAllocations;
 }
 
