[Insights] Add RegionsTraceAnalysis and RegionProvider

diff --git a/Engine/Source/Developer/TraceInsights/Private/Insights/Tests/MiscTests.cpp b/Engine/Source/Developer/TraceInsights/Private/Insights/Tests/MiscTests.cpp
new file mode 100644
--- /dev/null
+++ b/Engine/Source/Developer/TraceInsights/Private/Insights/Tests/MiscTests.cpp
@@ -0,0 +1,83 @@
+﻿// Copyright Epic Games, Inc. All Rights Reserved.
+#include "CoreMinimal.h"
+#include "Misc/AutomationTest.h"
+//awful hack since SlabAllocator is private
+#include "Developer/TraceServices/Private/Common/SlabAllocator.h"
+#include "Common/PagedArray.h"
+
+IMPLEMENT_SIMPLE_AUTOMATION_TEST(FPagedArrayFilteringTest, "Insights.FPagedArrayFilteringTest", EAutomationTestFlags::ApplicationContextMask | EAutomationTestFlags::SmokeFilter)
+
+bool operator!=(const FInt32Interval& lhs, const FInt32Interval& rhs)
+{
+	return lhs.Min != rhs.Min || lhs.Max != rhs.Max;
+}
+
+bool FPagedArrayFilteringTest::RunTest(const FString& Parameters)
+{
+	using namespace TraceServices;
+	
+	TArray<int32> Integers {0,1,2,3,4,5,6,7,8,9,10};
+	TArray<int32> Integer {1};
+	
+	//first element >= Value
+	TestEqual(TEXT("First index of element found in range"), Algo::LowerBound(Integers, 4),4);
+	TestEqual(TEXT("First index of element not in range"), Algo::LowerBound(Integers, 100),11);
+	//first element > Value
+	TestEqual(TEXT("Upper bound value found in range"), Algo::UpperBound(Integers, 4), 5);
+	TestEqual(TEXT("Upper bound value not found in range"), Algo::UpperBound(Integers, 100), 11);
+	
+	TestEqual(TEXT("Upper bound value found in single-item range"), Algo::UpperBound(Integer, 10), 1);
+
+	FSlabAllocator alloc(32 << 20);
+	int32 PageSize = 4;
+
+	struct FTimeRegion
+	{
+		double BeginTime;
+		double EndTime;
+	};
+	
+	TPagedArray<FTimeRegion> EmptyLane(alloc, PageSize);
+
+	TPagedArray<FTimeRegion> OneItemLane(alloc, PageSize);
+	OneItemLane.EmplaceBack(FTimeRegion{ 1.0, 2.0 });
+
+	TPagedArray<FTimeRegion> TestLane(alloc, PageSize);
+	TestLane.EmplaceBack(FTimeRegion{ 1.0, 2.0 });
+	TestLane.EmplaceBack(FTimeRegion{ 3.0, 4.0 });
+	TestLane.EmplaceBack(FTimeRegion{ 5.0, 6.0 });
+	TestLane.EmplaceBack(FTimeRegion{ 7.0, 8.0 });
+	TestLane.EmplaceBack(FTimeRegion{ 9.0, 10.0 });
+	TestLane.EmplaceBack(FTimeRegion{ 11.0, 12.0 });
+	TestLane.EmplaceBack(FTimeRegion{ 13.0, 14.0 });
+	
+	FInt32Interval Result;
+	auto BeginProj = [](const FTimeRegion& r ){ return r.BeginTime;};
+	auto EndProj = [](const FTimeRegion& r ){ return r.EndTime;};
+	
+	Result = GetElementRangeOverlappingGivenRange<FTimeRegion>(EmptyLane, -10.0, 0.0,BeginProj, EndProj);
+	TestEqual(TEXT("No overlaps for empty range"), Result, FInt32Interval{-1,-1});
+
+	Result = GetElementRangeOverlappingGivenRange<FTimeRegion>(OneItemLane, -10.0, 10.0,BeginProj, EndProj);
+	TestEqual(TEXT("Overlap is interval is larger than element-range"), Result, {0,0});
+	
+	Result = GetElementRangeOverlappingGivenRange<FTimeRegion>(TestLane, -10.0, 0.0,BeginProj, EndProj);
+	TestEqual(TEXT("No overlaps for interval before element-range"), Result, {-1,-1});
+
+	Result = GetElementRangeOverlappingGivenRange<FTimeRegion>(TestLane, 100.0, 200.0,BeginProj, EndProj);
+	TestEqual(TEXT("No overlaps for interval after element-range"), Result, {-1,-1});
+
+	//partially overlapping begin
+	Result = GetElementRangeOverlappingGivenRange<FTimeRegion>(TestLane, 0, 5,BeginProj, EndProj);
+	TestEqual(TEXT("Overlaps for interval earlier to inside element-range"), Result, {0,2});	
+	
+	// full inside
+	Result = GetElementRangeOverlappingGivenRange<FTimeRegion>(TestLane, 4, 8,BeginProj, EndProj);
+	TestEqual(TEXT("Overlaps for interval inside element-range"), Result, {2,3});
+
+	//partially overlapping end
+	Result = GetElementRangeOverlappingGivenRange<FTimeRegion>(TestLane, 11.3, 100,BeginProj, EndProj);
+	TestEqual(TEXT("Overlaps for interval inside to after element-range"), Result, {5,6});
+
+	return !HasAnyErrors();
+}
\ No newline at end of file
diff --git a/Engine/Source/Developer/TraceServices/Private/AnalysisService.cpp b/Engine/Source/Developer/TraceServices/Private/AnalysisService.cpp
--- a/Engine/Source/Developer/TraceServices/Private/AnalysisService.cpp
+++ b/Engine/Source/Developer/TraceServices/Private/AnalysisService.cpp
@@ -16,6 +16,7 @@
 #include "Model/LogPrivate.h"
 #include "Model/MemoryPrivate.h"
 #include "Model/NetProfilerProvider.h"
+#include "Model/RegionsPrivate.h"
 #include "Model/ScreenshotProviderPrivate.h"
 #include "Model/ThreadsPrivate.h"
 #include "Trace/Analysis.h"
@@ -241,6 +242,9 @@ TSharedPtr<const IAnalysisSession> FAnalysisService::StartAnalysis(uint32 TraceI
 	TSharedPtr<FBookmarkProvider> BookmarkProvider = MakeShared<FBookmarkProvider>(*Session);
 	Session->AddProvider(FBookmarkProvider::ProviderName, TSharedPtr<IBookmarkProvider>(BookmarkProvider), TSharedPtr<IEditableBookmarkProvider>(BookmarkProvider));
 
+	TSharedPtr<FRegionProvider> RegionProvider = MakeShared<FRegionProvider>(*Session);
+	Session->AddProvider(FRegionProvider::ProviderName, TSharedPtr<IRegionProvider>(RegionProvider), TSharedPtr<IEditableRegionProvider>(RegionProvider));
+
 	TSharedPtr<FLogProvider> LogProvider = MakeShared<FLogProvider>(*Session);
 	Session->AddProvider(FLogProvider::ProviderName, TSharedPtr<ILogProvider>(LogProvider), TSharedPtr<IEditableLogProvider>(LogProvider));
 
@@ -259,7 +263,7 @@ TSharedPtr<const IAnalysisSession> FAnalysisService::StartAnalysis(uint32 TraceI
 	TSharedPtr<FScreenshotProvider> ScreenshotProvider = MakeShared<FScreenshotProvider>(*Session);
 	Session->AddProvider(FScreenshotProvider::ProviderName, ScreenshotProvider);
 
-	Session->AddAnalyzer(new FMiscTraceAnalyzer(*Session, *ThreadProvider, *LogProvider, *FrameProvider, *ChannelProvider, *ScreenshotProvider));
+	Session->AddAnalyzer(new FMiscTraceAnalyzer(*Session, *ThreadProvider, *LogProvider, *FrameProvider, *ChannelProvider, *ScreenshotProvider, *RegionProvider));
 	Session->AddAnalyzer(new FBookmarksAnalyzer(*Session, *BookmarkProvider, LogProvider.Get()));
 	Session->AddAnalyzer(new FLogTraceAnalyzer(*Session, *LogProvider));
 
diff --git a/Engine/Source/Developer/TraceServices/Private/Analyzers/MiscTraceAnalysis.cpp b/Engine/Source/Developer/TraceServices/Private/Analyzers/MiscTraceAnalysis.cpp
--- a/Engine/Source/Developer/TraceServices/Private/Analyzers/MiscTraceAnalysis.cpp
+++ b/Engine/Source/Developer/TraceServices/Private/Analyzers/MiscTraceAnalysis.cpp
@@ -2,12 +2,15 @@
 
 #include "MiscTraceAnalysis.h"
 
+#include "Common/ProviderLock.h"
 #include "Common/Utils.h"
 #include "HAL/LowLevelMemTracker.h"
+#include "Logging/MessageLog.h"
 #include "Model/BookmarksPrivate.h"
 #include "Model/Channel.h"
 #include "Model/FramesPrivate.h"
 #include "Model/LogPrivate.h"
+#include "Model/RegionsPrivate.h"
 #include "Model/ScreenshotProviderPrivate.h"
 #include "Model/ThreadsPrivate.h"
 #include "TraceServices/Model/AnalysisSession.h"
@@ -20,13 +23,15 @@ FMiscTraceAnalyzer::FMiscTraceAnalyzer(IAnalysisSession& InSession,
 									   FLogProvider& InLogProvider,
 									   FFrameProvider& InFrameProvider,
 									   FChannelProvider& InChannelProvider,
-									   FScreenshotProvider& InScreenshotProvider)
+									   FScreenshotProvider& InScreenshotProvider,
+									   FRegionProvider& InRegionProvider)
 	: Session(InSession)
 	, ThreadProvider(InThreadProvider)
 	, LogProvider(InLogProvider)
 	, FrameProvider(InFrameProvider)
 	, ChannelProvider(InChannelProvider)
 	, ScreenshotProvider(InScreenshotProvider)
+	, RegionProvider(InRegionProvider)
 {
 	// Todo: update this to use provider locking instead of session locking
 	// FProviderEditScopeLock LogProviderLock (LogProvider);
@@ -55,6 +60,15 @@ void FMiscTraceAnalyzer::OnAnalysisBegin(const FOnAnalysisContext& Context)
 	Builder.RouteEvent(RouteId_ChannelToggle, "Trace", "ChannelToggle");
 	Builder.RouteEvent(RouteId_ScreenshotHeader, "Misc", "ScreenshotHeader");
 	Builder.RouteEvent(RouteId_ScreenshotChunk, "Misc", "ScreenshotChunk");
+	
+	Builder.RouteEvent(RouteId_RegionBegin, "Misc", "RegionBegin");
+	Builder.RouteEvent(RouteId_RegionEnd, "Misc", "RegionEnd");
+}
+
+void FMiscTraceAnalyzer::OnAnalysisEnd()
+{
+	FProviderEditScopeLock RegionProviderScopedLock(static_cast<IEditableProvider&>(RegionProvider));
+	RegionProvider.OnAnalysisSessionEnded();
 }
 
 void FMiscTraceAnalyzer::OnThreadInfo(const FThreadInfo& ThreadInfo)
@@ -228,6 +242,29 @@ bool FMiscTraceAnalyzer::OnEvent(uint16 RouteId, EStyle Style, const FOnEventCon
 		ThreadState->ThreadGroupStack.Pop();
 		break;
 	}
+	
+	case RouteId_RegionBegin:
+	{
+		uint64 Cycle = EventData.GetValue<uint64>("Cycle");
+		FString Name;
+		EventData.GetString("RegionName", Name);
+
+		FProviderEditScopeLock RegionProviderScopedLock(static_cast<const IRegionProvider&>(RegionProvider));
+		RegionProvider.AppendRegionBegin(*Name, Context.EventTime.AsSeconds(Cycle));
+		break;
+	}
+	
+	case RouteId_RegionEnd:
+	{
+		uint64 EndCycle = EventData.GetValue<uint64>("Cycle");
+		FString Name = TEXT("Invalid");
+		EventData.GetString("RegionName", Name);
+
+		FProviderEditScopeLock RegionProviderScopedLock(static_cast<const IRegionProvider&>(RegionProvider));
+		RegionProvider.AppendRegionEnd(*Name, Context.EventTime.AsSeconds(EndCycle));
+		break;
+	}
+		
 	//
 	// End retired events
 	}
diff --git a/Engine/Source/Developer/TraceServices/Private/Analyzers/MiscTraceAnalysis.h b/Engine/Source/Developer/TraceServices/Private/Analyzers/MiscTraceAnalysis.h
--- a/Engine/Source/Developer/TraceServices/Private/Analyzers/MiscTraceAnalysis.h
+++ b/Engine/Source/Developer/TraceServices/Private/Analyzers/MiscTraceAnalysis.h
@@ -16,6 +16,7 @@ class FLogProvider;
 class FFrameProvider;
 class FChannelProvider;
 class FScreenshotProvider;
+class FRegionProvider;
 
 class FMiscTraceAnalyzer
 	: public UE::Trace::IAnalyzer
@@ -26,8 +27,10 @@ public:
 					   FLogProvider& LogProvider,
 					   FFrameProvider& FrameProvider, 
 					   FChannelProvider& ChannelProvider,
-					   FScreenshotProvider& ScreenshotProvider);
+					   FScreenshotProvider& ScreenshotProvider,
+					   FRegionProvider& RegionProvider);
 	virtual void OnAnalysisBegin(const FOnAnalysisContext& Context) override;
+	virtual void OnAnalysisEnd() override;
 	virtual void OnThreadInfo(const FThreadInfo& ThreadInfo) override;
 	virtual bool OnEvent(uint16 RouteId, EStyle Style, const FOnEventContext& Context) override;
 
@@ -49,6 +52,8 @@ private:
 		RouteId_ChannelToggle,
 		RouteId_ScreenshotHeader,
 		RouteId_ScreenshotChunk,
+		RouteId_RegionBegin,
+		RouteId_RegionEnd
 	};
 
 	struct FThreadState
@@ -66,6 +71,8 @@ private:
 	FFrameProvider& FrameProvider;
 	FChannelProvider& ChannelProvider;
 	FScreenshotProvider& ScreenshotProvider;
+	FRegionProvider& RegionProvider;
+	
 	TMap<uint32, TSharedRef<FThreadState>> ThreadStateMap;
 	uint64 LastFrameCycle[TraceFrameType_Count] = { 0, 0 };
 	uint64 ScreenshotLogCategoryId = uint64(-1);
diff --git a/Engine/Source/Developer/TraceServices/Private/Model/Bookmarks.cpp b/Engine/Source/Developer/TraceServices/Private/Model/Bookmarks.cpp
--- a/Engine/Source/Developer/TraceServices/Private/Model/Bookmarks.cpp
+++ b/Engine/Source/Developer/TraceServices/Private/Model/Bookmarks.cpp
@@ -4,11 +4,16 @@
 #include "Common/FormatArgs.h"
 #include "Model/BookmarksPrivate.h"
 #include "AnalysisServicePrivate.h"
+#include "RegionsPrivate.h"
+#include "Common/ProviderLock.h"
+#include "TraceServices/Model/Regions.h"
 
 namespace TraceServices
 {
 
 const FName FBookmarkProvider::ProviderName("BookmarkProvider");
+static constexpr FStringView RegionStartToken = TEXTVIEW("RegionStart:");
+static constexpr FStringView RegionEndToken = TEXTVIEW("RegionEnd:");
 
 FBookmarkProvider::FBookmarkProvider(IAnalysisSession& InSession)
 	: Session(InSession)
@@ -50,6 +55,9 @@ void FBookmarkProvider::AppendBookmark(uint64 BookmarkPoint, double Time, const
 	Bookmark->Time = Time;
 	Bookmark->Text = Session.StoreString(FormatBuffer);
 	Bookmarks.Add(Bookmark);
+
+	CheckBookmarkForRegion(Bookmark);
+	
 	Session.UpdateDurationSeconds(Time);
 }
 
@@ -60,9 +68,37 @@ void FBookmarkProvider::AppendBookmark(uint64 BookmarkPoint, double Time, const
 	Bookmark->Time = Time;
 	Bookmark->Text = Text;
 	Bookmarks.Add(Bookmark);
+
+	CheckBookmarkForRegion(Bookmark);
+	
 	Session.UpdateDurationSeconds(Time);
 }
 
+void FBookmarkProvider::CheckBookmarkForRegion(const TSharedRef<FBookmarkInternal> Bookmark) const
+{
+	const FStringView Text = Bookmark->Text;
+	IEditableRegionProvider* EditableRegionProvider = Session.EditProvider<IEditableRegionProvider>(FRegionProvider::ProviderName);
+
+	if (!EditableRegionProvider)
+	{
+		return;
+	}
+	
+	if (Text.StartsWith(RegionStartToken))
+	{
+		// StringView.GetData() is not necessarily null-terminated. Since we started from a null terminated string
+		// and only called RightChop() we should still be fine.
+		
+		FProviderEditScopeLock _(*EditableRegionProvider);
+		EditableRegionProvider->AppendRegionBegin(Text.RightChop(RegionStartToken.Len()).GetData(), Bookmark->Time);
+	}
+	if (Text.StartsWith(RegionEndToken))
+	{
+		FProviderEditScopeLock _(*EditableRegionProvider);
+		EditableRegionProvider->AppendRegionEnd(Text.RightChop(RegionEndToken.Len()).GetData(), Bookmark->Time);
+	}
+}
+
 void FBookmarkProvider::EnumerateBookmarks(double IntervalStart, double IntervalEnd, TFunctionRef<void(const FBookmark &)> Callback) const
 {
 	Session.ReadAccessCheck();
diff --git a/Engine/Source/Developer/TraceServices/Private/Model/BookmarksPrivate.h b/Engine/Source/Developer/TraceServices/Private/Model/BookmarksPrivate.h
--- a/Engine/Source/Developer/TraceServices/Private/Model/BookmarksPrivate.h
+++ b/Engine/Source/Developer/TraceServices/Private/Model/BookmarksPrivate.h
@@ -44,6 +44,10 @@ public:
 	virtual void AppendBookmark(uint64 BookmarkPoint, double Time, const TCHAR* Text) override;
 
 private:
+	// check whether a bookmark is actually a region marker from the deprecated/5.1 regions implementation.
+	// will forward to IEditableRegionProvider in this case
+	void CheckBookmarkForRegion(const TSharedRef<FBookmarkInternal> Shared) const;
+	
 	IAnalysisSession& Session;
 	TMap<uint64, TSharedPtr<FBookmarkSpec>> SpecMap;
 	TArray<TSharedRef<FBookmarkInternal>> Bookmarks;
diff --git a/Engine/Source/Developer/TraceServices/Private/Model/Regions.cpp b/Engine/Source/Developer/TraceServices/Private/Model/Regions.cpp
new file mode 100644
--- /dev/null
+++ b/Engine/Source/Developer/TraceServices/Private/Model/Regions.cpp
@@ -0,0 +1,204 @@
+﻿// Copyright Epic Games, Inc. All Rights Reserved.
+
+#include "TraceServices/Model/Regions.h"
+#include "Common/FormatArgs.h"
+#include "Model/RegionsPrivate.h"
+#include "AnalysisServicePrivate.h"
+#include "Common/Utils.h"
+#include "Logging/MessageLog.h"
+
+#define LOCTEXT_NAMESPACE "RegionProvider"
+
+namespace TraceServices
+{
+
+thread_local FProviderLock::FThreadLocalState GRegionsProviderLockState;
+	
+const FName FRegionProvider::ProviderName("RegionProvider");
+
+FRegionProvider::FRegionProvider(IAnalysisSession& InSession)
+	: Session(InSession)
+{
+}
+	
+uint64 FRegionProvider::GetRegionCount() const
+{
+	ReadAccessCheck();
+
+	uint64 RegionCount = 0;
+	for (const FRegionLane& Lane : Lanes)
+	{
+		RegionCount += Lane.Num();
+	}
+	return RegionCount;
+}
+
+const FRegionLane* FRegionProvider::GetLane(int32 index) const
+{
+	ReadAccessCheck();
+	if (index < Lanes.Num())
+	{
+		return &(Lanes[index]);
+	}
+	return nullptr;
+}
+
+void FRegionProvider::AppendRegionBegin(const TCHAR* Name, double Time)
+{
+	EditAccessCheck();
+	
+	FTimeRegion Region;
+	Region.BeginTime = Time;
+	Region.Text = Session.StoreString(Name);
+	Region.Depth = CalculateRegionDepth(Region);
+
+	if (Region.Depth == Lanes.Num())
+	{
+		Lanes.Emplace(Session.GetLinearAllocator());
+	}
+	
+	Lanes[Region.Depth].Regions.EmplaceBack(Region);
+	FTimeRegion* OpenRegion = &(Lanes[Region.Depth].Regions.Last());
+	OpenRegions.Add(Region.Text, OpenRegion);
+
+	{
+		FAnalysisSessionEditScope _(Session);
+		Session.UpdateDurationSeconds(Time);
+	}
+	
+	UpdateCounter++;
+}
+
+void FRegionProvider::AppendRegionEnd(const TCHAR* Name, double Time)
+{
+	EditAccessCheck();
+
+	if (FTimeRegion** OpenRegion = OpenRegions.Find(Name))
+	{
+		(*OpenRegion)->EndTime = Time;
+		OpenRegions.Remove(Name);
+		{
+			FAnalysisSessionEditScope _(Session);
+			Session.UpdateDurationSeconds(Time);
+		}
+		UpdateCounter++;
+	}
+	else
+	{
+		if (FMessageLog* Log = Session.GetLog())
+		{
+			Log->Warning(LOCTEXT("StrayRegionEnd","A region end event was encountered without having seen a matching region start event first."));
+		}
+		else
+		{
+			UE_LOG(LogTraceServices, Warning, TEXT("A region end event (%s) was encountered without having seen a matching region start event first."), Name)
+		}
+	}
+}
+
+void FRegionProvider::OnAnalysisSessionEnded()
+{
+	EditAccessCheck();
+	
+	for (const auto& KV : OpenRegions)
+	{
+		const FTimeRegion* Region = KV.Value;
+		if (FMessageLog* Log = Session.GetLog())
+		{
+			Log->Warning(LOCTEXT("StrayRegionBegin","A region begin event was never closed."));
+		}
+		else
+		{
+			UE_LOG(LogTraceServices, Warning, TEXT("A region begin event (%s) was never closed."), Region->Text)
+		}		
+	}
+}
+
+int32 FRegionProvider::CalculateRegionDepth(const FTimeRegion& Region) const
+{
+	TArray<const FTimeRegion*> Overlaps;
+
+	constexpr int32 DepthLimit = 100;
+	int32 NewDepth = 0;
+	// find first free lane/depth
+	while (NewDepth < DepthLimit)
+	{
+		if (!Lanes.IsValidIndex(NewDepth))
+		{
+			break;
+		}
+		
+		const FTimeRegion& LastRegion = Lanes[NewDepth].Regions.Last();
+		if (LastRegion.EndTime <= Region.BeginTime)
+		{
+			break;
+		}
+		NewDepth++;
+	}
+
+	ensureMsgf(NewDepth < DepthLimit, TEXT("Regions are nested too deep."));
+
+	return NewDepth;
+}
+	
+void FRegionProvider::EnumerateLanes(TFunctionRef<void(const FRegionLane&, int32)> Callback) const
+{
+	for (int i = 0 ; i < Lanes.Num(); ++i)
+	{
+		Callback(Lanes[i], i);
+	}
+}
+	
+bool FRegionProvider::EnumerateRegions(double IntervalStart, double IntervalEnd, TFunctionRef<bool(const FTimeRegion &)> Callback) const
+{
+	ReadAccessCheck();
+	
+	if (IntervalStart > IntervalEnd)
+	{
+		return false;
+	}
+
+	for (const FRegionLane& Lane : Lanes)
+	{
+		if (!Lane.EnumerateRegions(IntervalStart, IntervalEnd, Callback))
+			return false;
+	}
+
+	return true;
+}
+
+bool FRegionLane::EnumerateRegions(double IntervalStart, double IntervalEnd, TFunctionRef<bool(const FTimeRegion&)> Callback) const
+{
+	const FInt32Interval OverlapRange = GetElementRangeOverlappingGivenRange<FTimeRegion>(Regions, IntervalStart, IntervalEnd,
+		[](const FTimeRegion& r){return r.BeginTime;},
+		[](const FTimeRegion& r){return r.EndTime;});
+	
+	if (OverlapRange.Min == -1 )
+	{
+		return true;
+	}
+	
+	for (int32 Index = OverlapRange.Min; Index <= OverlapRange.Max; ++Index)
+	{
+		if (!Callback(Regions[Index]))
+		{
+			return false;
+		}
+	}
+
+	return true;
+}
+	
+const IRegionProvider& ReadRegionProvider(const IAnalysisSession& Session)
+{
+	return *Session.ReadProvider<IRegionProvider>(FRegionProvider::ProviderName);
+}
+
+IEditableRegionProvider& EditRegionProvider(IAnalysisSession& Session)
+{
+	return *Session.EditProvider<IEditableRegionProvider>(FRegionProvider::ProviderName);
+}
+
+} // namespace TraceServices
+
+#undef LOCTEXT_NAMESPACE
\ No newline at end of file
diff --git a/Engine/Source/Developer/TraceServices/Private/Model/RegionsPrivate.h b/Engine/Source/Developer/TraceServices/Private/Model/RegionsPrivate.h
new file mode 100644
--- /dev/null
+++ b/Engine/Source/Developer/TraceServices/Private/Model/RegionsPrivate.h
@@ -0,0 +1,73 @@
+﻿// Copyright Epic Games, Inc. All Rights Reserved.
+
+#pragma once
+
+#include "Common/PagedArray.h"
+#include "Common/ProviderLock.h"
+#include "TraceServices/Model/Regions.h"
+#include "Templates/SharedPointer.h"
+
+
+namespace TraceServices
+{
+extern thread_local FProviderLock::FThreadLocalState GRegionsProviderLockState;
+
+class FAnalysisSessionLock;
+class FStringStore;
+
+class FRegionProvider
+	: public IRegionProvider
+	, public IEditableRegionProvider
+{
+public:
+	static const FName ProviderName;
+
+	// Provider locking API
+	virtual void BeginEdit() const override       { Lock.BeginWrite(GRegionsProviderLockState); }
+	virtual void EndEdit() const override         { Lock.EndWrite(GRegionsProviderLockState); }
+	virtual void EditAccessCheck() const override { Lock.WriteAccessCheck(GRegionsProviderLockState); }
+	virtual void BeginRead() const override       { Lock.BeginRead(GRegionsProviderLockState); }
+	virtual void EndRead() const override         { Lock.EndRead(GRegionsProviderLockState); }
+	virtual void ReadAccessCheck() const override { Lock.ReadAccessCheck(GRegionsProviderLockState); }
+
+	explicit FRegionProvider(IAnalysisSession& Session);
+	virtual ~FRegionProvider() override {}
+
+	// implement IRegionProvider API
+	virtual uint64 GetRegionCount() const override;
+	
+	virtual uint64 GetUpdateCounter() const override { ReadAccessCheck(); return UpdateCounter; }
+
+	virtual bool EnumerateRegions(double IntervalStart, double IntervalEnd, TFunctionRef<bool(const FTimeRegion&)> Callback) const override;
+	virtual void EnumerateLanes(TFunctionRef<void(const FRegionLane&, const int32)> Callback) const override;
+	virtual int32 GetLaneCount()  const override { ReadAccessCheck(); return Lanes.Num(); }
+
+	/**
+	 * Direct Access to a certain lane at a given index/depth.
+	 * Do not store the pointer returned from this function as it could be accessed outside the ProviderEditScope.
+	 */
+	virtual const FRegionLane* GetLane(int32 index) const override;
+
+
+	// implement IEditableRegionProvider API
+	virtual void AppendRegionBegin(const TCHAR* Name, double Time) override;
+	virtual void AppendRegionEnd(const TCHAR* Name, double Time) override;
+	virtual void OnAnalysisSessionEnded() override;
+
+private:
+
+	mutable FProviderLock Lock;
+	
+	// Update the depth member of a region to allow overlapping regions to be displayed on separate lanes.
+	int32 CalculateRegionDepth(const FTimeRegion& Item) const;
+	
+	IAnalysisSession& Session;
+	// open Regions inside lanes
+	TMap<FStringView, FTimeRegion*> OpenRegions;
+	// closed regions
+	TArray<FRegionLane> Lanes;
+	// counter incremented each time region data changes during analysis
+	uint64 UpdateCounter = -1;
+};
+
+} // namespace TraceServices
diff --git a/Engine/Source/Developer/TraceServices/Public/Common/PagedArray.h b/Engine/Source/Developer/TraceServices/Public/Common/PagedArray.h
--- a/Engine/Source/Developer/TraceServices/Public/Common/PagedArray.h
+++ b/Engine/Source/Developer/TraceServices/Public/Common/PagedArray.h
@@ -756,4 +756,98 @@ inline SIZE_T GetNum(const TPagedArray<ItemType, PageType>& PagedArray)
 	return PagedArray.NumPages();
 }
 
+/**
+ *	Use binary search to find the first and last element inside a TPagedArray that overlaps a given input interval.
+ *	This requires the elements in the array to be sorted by the value returned from the Projection.
+ *	Example usage for Timeline events would require a projection that returns Item.StartTime and the resulting range.Min
+ *	will point to the last element where Item.End > StartTime and range.Max to the last element where
+ *	Item.StartTime < EndTime.
+ */
+template<typename ItemType, typename PageType>	
+FInt32Interval GetElementRangeOverlappingGivenRange(const TPagedArray<ItemType, PageType>& PagedArray,
+	double StartTime, double EndTime,
+	TFunctionRef<double(const ItemType&)> ItemStartProjection,
+	TFunctionRef<double(const ItemType&)> ItemEndProjection
+	)
+{
+	FInt32Interval Result = { -1, -1 };
+	check(EndTime >= StartTime);
+
+	const PageType* PageData = PagedArray.GetPages();
+	if (!PageData)
+	{
+		return Result;
+	}
+	
+	const int32 NumPoints = static_cast<int32>(PagedArray.Num());
+	const int32 NumPages = static_cast<int32>(PagedArray.NumPages());
+	const int32 PageSize = static_cast<int32>(PagedArray.GetPageSize());
+
+	if (ItemStartProjection(PagedArray.First()) > EndTime || ItemEndProjection(PagedArray.Last()) <= StartTime)
+	{
+		return Result;
+	}
+	
+	TArrayView<const PageType, int32> Pages = MakeArrayView(PageData, NumPages);
+
+	// find the page before the first page that's already inside the range (thus -1)
+	int32 StartPageIndex = Algo::UpperBoundBy(Pages, StartTime,
+		[&ItemEndProjection](const PageType& Page) { return ItemEndProjection(Page.Items[0]); });
+	StartPageIndex -= 1;
+
+	if (StartPageIndex < 0)
+	{
+		Result.Min = 0;
+		StartPageIndex = 0;
+	}
+	else 
+	{
+		// if we went past the end with StartPageIndex we still have to search the last page
+		const PageType& Page = PageData[StartPageIndex];
+		TArrayView<ItemType> PageValues(Page.Items, Page.Count);
+		const int32 Index = Algo::UpperBoundBy(PageValues, StartTime, ItemEndProjection);
+		check(Index >= 0);
+		// Index may point past the end of the page, but the first item in the next page could be valid
+		// as long as we're not in the last page
+		if (Index >= Page.Count && StartPageIndex + 1 == NumPages )
+		{
+			return Result;
+		}
+	
+		Result.Min = StartPageIndex * PageSize + Index;
+		check(Index <= NumPoints);
+	}
+	
+	TArrayView<const PageType, int32> RemainingPages = MakeArrayView(&PageData[StartPageIndex], NumPages - StartPageIndex);
+
+	int32 EndPageIndex = Algo::UpperBoundBy(RemainingPages, EndTime,
+		[&ItemStartProjection](const PageType& Page) { return ItemStartProjection(Page.Items[0]); });
+	// EndPageIndex needs to be remapped back to the full page range
+	EndPageIndex += StartPageIndex;
+
+	// past the end means we still have to search the last page, so no early out here
+	check(EndPageIndex > 0)
+	// find the page before the first page that outside the range (thus -1)
+	EndPageIndex -= 1;
+	
+	{
+		const PageType& Page = PageData[EndPageIndex];
+		TArrayView<ItemType> PageValues(Page.Items, Page.Count);
+		int32 Index = Algo::UpperBoundBy(PageValues, EndTime, ItemStartProjection);
+		check(Index >= 0);
+		// Index may point past the end of the page, that means include last element in result
+		if (Index >= PageValues.Num())
+		{
+			Index = PageValues.Num() - 1;
+		}
+		else
+		{
+			Index -= 1;
+		}
+		Result.Max = EndPageIndex * PageSize + Index;
+	}
+
+	return Result;
+}
+
 } // namespace TraceServices
diff --git a/Engine/Source/Developer/TraceServices/Public/TraceServices/Model/AnalysisSession.h b/Engine/Source/Developer/TraceServices/Public/TraceServices/Model/AnalysisSession.h
--- a/Engine/Source/Developer/TraceServices/Public/TraceServices/Model/AnalysisSession.h
+++ b/Engine/Source/Developer/TraceServices/Public/TraceServices/Model/AnalysisSession.h
@@ -54,6 +54,11 @@ public:
 	virtual uint32 GetTraceId() const = 0;
 	virtual bool IsAnalysisComplete() const = 0;
 	virtual double GetDurationSeconds() const = 0;
+	/**
+	 * Update the internal estimation of the session duration with a new timestamp.
+	 * This function should be called with any timestamps from individual events so the full session duration can be
+	 * estimated.
+	 */
 	virtual void UpdateDurationSeconds(double Duration) = 0;
 	virtual FMessageLog* GetLog() const = 0;
 
diff --git a/Engine/Source/Developer/TraceServices/Public/TraceServices/Model/Regions.h b/Engine/Source/Developer/TraceServices/Public/TraceServices/Model/Regions.h
new file mode 100644
--- /dev/null
+++ b/Engine/Source/Developer/TraceServices/Public/TraceServices/Model/Regions.h
@@ -0,0 +1,127 @@
+﻿// Copyright Epic Games, Inc. All Rights Reserved.
+
+#pragma  once
+
+#include "Common/PagedArray.h"
+#include "HAL/Platform.h"
+#include "Templates/Function.h"
+#include "TraceServices/Model/AnalysisSession.h"
+
+template <typename FuncType> class TFunctionRef;
+
+namespace TraceServices
+{
+
+class FRegionProvider;
+
+struct TRACESERVICES_API FTimeRegion
+{
+	double BeginTime = std::numeric_limits<double>::infinity();
+	double EndTime = std::numeric_limits<double>::infinity();
+	const TCHAR* Text = nullptr;
+	int8_t Depth = -1;
+};
+
+class TRACESERVICES_API FRegionLane
+{
+public:
+
+	FRegionLane(ILinearAllocator& InAllocator) : Regions(InAllocator, 512) {}
+
+	/**
+	 * Call Callback for every region overlapping the interval defined by IntervalStart and IntervalEnd
+	 * @param Callback a callback called for each region. Return false to abort iteration.
+	 * @returns true if the enumeration finished, false if it was aborted by the callback returning false
+	 */
+	bool EnumerateRegions(double IntervalStart, double IntervalEnd, TFunctionRef<bool(const FTimeRegion&)> Callback) const;
+	int32 Num() const { return Regions.Num(); }
+private:
+	TPagedArray<FTimeRegion> Regions;
+	friend class FRegionProvider;
+};
+
+	
+class TRACESERVICES_API IRegionProvider
+	: public IProvider
+{
+public:
+
+	
+	virtual ~IRegionProvider() override = default;
+	
+	/**
+	 * @return the amount of currently known regions (including open-ended ones)
+	 */
+	virtual uint64 GetRegionCount() const = 0;
+	
+	/**
+	 * @return the amount of currently known regions (including open-ended ones)
+	 */
+	virtual int32 GetLaneCount() const = 0;
+
+	/**
+	 * @return a pointer to the lane at depth index or nullptr if index > GetLaneCount()-1
+	 */
+	virtual const FRegionLane* GetLane(int32 index) const = 0; 
+	
+	/**
+	 * @return A monotonically increasing counter that that changes each time new data is added to the provider.
+	 * This can be used to detect when to update any (UI-)state dependent on the provider during analysis.
+	 */
+	virtual uint64 GetUpdateCounter() const = 0;
+	
+	/**
+	 * Enumerates all regions that overlap a certain time interval. Will enumerate by depth but does not expose lanes.
+	 * @param Callback a callback called for each region. Return false to abort iteration.
+	 * @returns true if the enumeration finished, false if it was aborted by the callback returning false
+	 */
+	virtual bool EnumerateRegions(double IntervalStart, double IntervalEnd, TFunctionRef<bool(const FTimeRegion&)> Callback) const = 0;
+	/**
+	 * Will call Callback(Lane, Depth) for each lane in order.
+	 */
+	virtual void EnumerateLanes(TFunctionRef<void(const FRegionLane&, const int32)> Callback) const = 0;
+
+	// Provider locking API
+	virtual void BeginEdit() const override = 0;
+	virtual void EndEdit() const override = 0;
+	virtual void BeginRead() const override = 0;
+	virtual void EndRead() const override = 0;
+	
+};
+
+/*
+* The interface to a provider that can consume mutations of region events from a session.
+*/
+class IEditableRegionProvider
+	: public IEditableProvider
+{
+public:
+	virtual ~IEditableRegionProvider() override = default;
+
+	/*
+	* Append a new begin event of a region from the trace session.
+	*
+	* @param Name		The string name of the region.
+	* @param Time		The time in seconds of the begin event of this region.
+	*/
+	virtual void AppendRegionBegin(const TCHAR* Name, double Time) = 0;
+	
+	/*
+	* Append a new end event of a region from the trace session.
+	*
+	* @param Name		The string name of the region.
+	* @param Time		The time in seconds of the end event of this region.
+	*/
+	virtual void AppendRegionEnd(const TCHAR* Name, double Time) = 0;
+
+	/**
+	 * Called from the analyzer once all events have been processed.
+	 * Allows postprocessing and error reporting for regions that were never closed.
+	 */
+	virtual void OnAnalysisSessionEnded() = 0;
+};
+
+TRACESERVICES_API const IRegionProvider& ReadRegionProvider(const IAnalysisSession& Session);
+TRACESERVICES_API IEditableRegionProvider& EditRegionProvider(IAnalysisSession& Session);
+
+} // namespace TraceServices
