[Insights] Fix NoPCH build caused by missing include.

diff --git a/Engine/Source/Developer/TraceServices/Public/Common/PagedArray.h b/Engine/Source/Developer/TraceServices/Public/Common/PagedArray.h
--- a/Engine/Source/Developer/TraceServices/Public/Common/PagedArray.h
+++ b/Engine/Source/Developer/TraceServices/Public/Common/PagedArray.h
@@ -4,6 +4,8 @@
 
 #include "TraceServices/Containers/Allocators.h"
 #include "Containers/Array.h"
+#include "Containers/ArrayView.h"
+#include "Math/Interval.h"
 
 #ifndef TRACESERVICES_PAGED_ARRAY_ITERATOR_V2
 #define TRACESERVICES_PAGED_ARRAY_ITERATOR_V2 0 // enables a simpler implementation of TPagedArrayIterator, for debug purposes
