[Insights] Further NonUnity fixes for Regions.cpp

diff --git a/Engine/Source/Developer/TraceServices/Private/Model/Regions.cpp b/Engine/Source/Developer/TraceServices/Private/Model/Regions.cpp
--- a/Engine/Source/Developer/TraceServices/Private/Model/Regions.cpp
+++ b/Engine/Source/Developer/TraceServices/Private/Model/Regions.cpp
@@ -6,6 +6,7 @@
 #include "AnalysisServicePrivate.h"
 #include "Common/Utils.h"
 #include "Logging/MessageLog.h"
+#include "Internationalization/Internationalization.h"
 
 #define LOCTEXT_NAMESPACE "RegionProvider"
 
