[Insights] - Refactored the IProvider / IEditableProvider interface to avoid diamond problem.
  - RegionProvider: Added warning when a begin region event is encountered while there is an open region with same name. Also added limitation to max 100 warnings (to avoid spam).
  - MemoryInsights: Fixed crash when grouping by heaps or by tags (not using the correct read lock).

diff --git a/Engine/Source/Developer/TraceInsights/Private/Insights/MemoryProfiler/ViewModels/MemAllocGroupingByHeap.cpp b/Engine/Source/Developer/TraceInsights/Private/Insights/MemoryProfiler/ViewModels/MemAllocGroupingByHeap.cpp
--- a/Engine/Source/Developer/TraceInsights/Private/Insights/MemoryProfiler/ViewModels/MemAllocGroupingByHeap.cpp
+++ b/Engine/Source/Developer/TraceInsights/Private/Insights/MemoryProfiler/ViewModels/MemAllocGroupingByHeap.cpp
@@ -1,12 +1,13 @@
 ﻿// Copyright Epic Games, Inc. All Rights Reserved.
 
 #include "MemAllocGroupingByHeap.h"
-#include "CallstackFormatting.h"
+
 #include "Common/ProviderLock.h" // TraceServices
 #include "Internationalization/Internationalization.h"
 
 // Insights
 #include "Insights/Common/AsyncOperationProgress.h"
+#include "Insights/MemoryProfiler/ViewModels/CallstackFormatting.h"
 #include "Insights/MemoryProfiler/ViewModels/MemAllocNode.h"
 
 #define LOCTEXT_NAMESPACE "Insights::FMemAllocGroupingByHeap"
diff --git a/Engine/Source/Developer/TraceInsights/Private/Insights/MemoryProfiler/ViewModels/MemAllocGroupingByTag.cpp b/Engine/Source/Developer/TraceInsights/Private/Insights/MemoryProfiler/ViewModels/MemAllocGroupingByTag.cpp
--- a/Engine/Source/Developer/TraceInsights/Private/Insights/MemoryProfiler/ViewModels/MemAllocGroupingByTag.cpp
+++ b/Engine/Source/Developer/TraceInsights/Private/Insights/MemoryProfiler/ViewModels/MemAllocGroupingByTag.cpp
@@ -1,8 +1,12 @@
 ﻿// Copyright Epic Games, Inc. All Rights Reserved.
+
 #include "MemAllocGroupingByTag.h"
-#include "MemAllocNode.h"
 
+#include "Common/ProviderLock.h" // TraceServices
+
+// Insights
 #include "Insights/Common/AsyncOperationProgress.h"
+#include "Insights/MemoryProfiler/ViewModels/MemAllocNode.h"
 
 #define LOCTEXT_NAMESPACE "Insights::FMemAllocGroupingByTag"
 
@@ -14,7 +18,7 @@ namespace Insights
 ////////////////////////////////////////////////////////////////////////////////////////////////////
 
 INSIGHTS_IMPLEMENT_RTTI(FMemAllocGroupingByTag);
-	
+
 ////////////////////////////////////////////////////////////////////////////////////////////////////
 
 FMemAllocGroupingByTag::FMemAllocGroupingByTag(const TraceServices::IAllocationsProvider& InTagProvider)
@@ -35,7 +39,7 @@ void FMemAllocGroupingByTag::GroupNodes(const TArray<FTableTreeNodePtr>& Nodes,
                                         IAsyncOperationProgress& InAsyncOperationProgress) const
 {
 	using namespace TraceServices;
-	
+
 	ParentGroup.ClearChildren();
 
 	constexpr TagIdType InvalidTagId = ~0u;
@@ -111,19 +115,22 @@ void FMemAllocGroupingByTag::GroupNodes(const TArray<FTableTreeNodePtr>& Nodes,
 	FScopedMemory _(TagNodes);
 
 	// Create tag nodes.
-	TagProvider.EnumerateTags([&](const TCHAR* Name, const TCHAR* FullName, TagIdType Id, TagIdType ParentId)
 	{
-		const FName NodeName(Name);
+		TraceServices::FProviderReadScopeLock TagProviderReadScopeLock(TagProvider);
+		TagProvider.EnumerateTags([&](const TCHAR* Name, const TCHAR* FullName, TagIdType Id, TagIdType ParentId)
+			{
+				const FName NodeName(Name);
 #if INSIGHTS_MERGE_MEM_TAGS_BY_NAME
-		FTableTreeNodePtr TreeNode = MakeShared<FTableTreeNode>(NodeName, InParentTable);
+				FTableTreeNodePtr TreeNode = MakeShared<FTableTreeNode>(NodeName, InParentTable);
 #else
-		const FName NodeNameEx(Name, int32(Id));
-		FTableTreeNodePtr TreeNode = MakeShared<FTableTreeNode>(NodeNameEx, InParentTable);
+				const FName NodeNameEx(Name, int32(Id));
+				FTableTreeNodePtr TreeNode = MakeShared<FTableTreeNode>(NodeNameEx, InParentTable);
 #endif
-		FTagNode* TagNode = new FTagNode(NodeName, Id, ParentId, TreeNode);
-		TagNodes.Add(TagNode);
-		IdToNodeMap.Add(Id, TagNode);
-	});
+				FTagNode* TagNode = new FTagNode(NodeName, Id, ParentId, TreeNode);
+				TagNodes.Add(TagNode);
+				IdToNodeMap.Add(Id, TagNode);
+			});
+	}
 
 	// Create the Untagged node (if not already present).
 	constexpr TagIdType UntaggedNodeId = 0;
@@ -250,7 +257,7 @@ void FMemAllocGroupingByTag::GroupNodes(const TArray<FTableTreeNodePtr>& Nodes,
 		}
 	}
 }
-	
+
 ////////////////////////////////////////////////////////////////////////////////////////////////////
 
 #undef INSIGHTS_MERGE_MEM_TAGS_BY_NAME
diff --git a/Engine/Source/Developer/TraceInsights/Private/Insights/MemoryProfiler/Widgets/SMemAllocTableTreeView.cpp b/Engine/Source/Developer/TraceInsights/Private/Insights/MemoryProfiler/Widgets/SMemAllocTableTreeView.cpp
--- a/Engine/Source/Developer/TraceInsights/Private/Insights/MemoryProfiler/Widgets/SMemAllocTableTreeView.cpp
+++ b/Engine/Source/Developer/TraceInsights/Private/Insights/MemoryProfiler/Widgets/SMemAllocTableTreeView.cpp
@@ -414,7 +414,7 @@ void SMemAllocTableTreeView::UpdateQuery(TraceServices::IAllocationsProvider::EQ
 							{
 								if (Type == AssetMetadataType)
 								{
-									TraceServices::IDefinitionProvider::FReadScopeLock DefinitionProviderReadLock(*DefinitionProvider);
+									TraceServices::FProviderReadScopeLock DefinitionProviderReadLock(*DefinitionProvider);
 									const auto Reader = Schema->Reader();
 									const auto AssetNameRef = Reader.GetValueAs<UE::Trace::FEventRef32>((uint8*)Data, 0);
 									const auto AssetName = DefinitionProvider->Get<TraceServices::FStringDefinition>(*AssetNameRef);
@@ -1298,19 +1298,21 @@ void SMemAllocTableTreeView::PopulateLLMTagSuggestionList(const FString& Text, T
 		return;
 	}
 
-	TraceServices::FProviderReadScopeLock _(*AllocationsProvider);
-
 	// Use a Set to avoid duplicate tag names.
 	TSet<FString> Suggestions;
-	AllocationsProvider->EnumerateTags([&Suggestions, &Text](const TCHAR* Display, const TCHAR* FullPath, TraceServices::TagIdType CurrentTag, TraceServices::TagIdType ParentTag)
+
 	{
-		if (Text.IsEmpty() || FCString::Stristr(FullPath, *Text))
+		TraceServices::FProviderReadScopeLock _(*AllocationsProvider);
+		AllocationsProvider->EnumerateTags([&Suggestions, &Text](const TCHAR* Display, const TCHAR* FullPath, TraceServices::TagIdType CurrentTag, TraceServices::TagIdType ParentTag)
 		{
-			Suggestions.Add(FullPath);
-		}
+			if (Text.IsEmpty() || FCString::Stristr(FullPath, *Text))
+			{
+				Suggestions.Add(FullPath);
+			}
 
-		return true;
-	});
+			return true;
+		});
+	}
 
 	OutSuggestions = Suggestions.Array();
 	OutSuggestions.Sort();
@@ -1716,10 +1718,10 @@ void SMemAllocTableTreeView::ExportMemorySnapshot() const
 		ReportMessageLog.Notify();
 		return;
 	}
-	
+
 	UTF16CHAR BOM = UNICODE_BOM;
 	ExportFileHandle->Write((uint8*)&BOM, sizeof(UTF16CHAR));
-	
+
 	static constexpr TCHAR LineEnd = TEXT('\n');
 	static constexpr TCHAR QuotationMarkBegin = TEXT('\"');
 	static constexpr TCHAR QuotationMarkEnd = TEXT('\"');
@@ -1776,7 +1778,7 @@ void SMemAllocTableTreeView::ExportMemorySnapshot() const
 					OutData << QuotationMarkEnd;
 					return;
 				}
-				
+
 				const uint32 NumCallstackFrames = Alloc->Callstack->Num();
 				check(NumCallstackFrames <= 256);
 				OutData << QuotationMarkBegin;
diff --git a/Engine/Source/Developer/TraceServices/Private/AnalysisService.cpp b/Engine/Source/Developer/TraceServices/Private/AnalysisService.cpp
--- a/Engine/Source/Developer/TraceServices/Private/AnalysisService.cpp
+++ b/Engine/Source/Developer/TraceServices/Private/AnalysisService.cpp
@@ -158,7 +158,7 @@ const IProvider* FAnalysisSession::ReadProviderPrivate(const FName& InName) cons
 	}
 }
 
-IProvider* FAnalysisSession::EditProviderPrivate(const FName& InName)
+IEditableProvider* FAnalysisSession::EditProviderPrivate(const FName& InName)
 {
 	const auto* FindIt = Providers.Find(InName);
 	if (FindIt)
@@ -235,27 +235,23 @@ TSharedPtr<const IAnalysisSession> FAnalysisService::StartAnalysis(uint32 TraceI
 
 	FAnalysisSessionEditScope _(*Session);
 
-	// Note that we upcast to the interface we wish to register in cases where a provider implements multiple interfaces
-	// We aren't using virtual inheritance since we don't care about the vtable complexity, and there is no state to worry about
-	// We follow the single data inheritance but multiple interface inheritance idiom, and this necessitates the upcast via TSharedPtr<IInterface>
-
 	TSharedPtr<FBookmarkProvider> BookmarkProvider = MakeShared<FBookmarkProvider>(*Session);
-	Session->AddProvider(FBookmarkProvider::ProviderName, TSharedPtr<IBookmarkProvider>(BookmarkProvider), TSharedPtr<IEditableBookmarkProvider>(BookmarkProvider));
+	Session->AddProvider(GetBookmarkProviderName(), BookmarkProvider, BookmarkProvider);
 
 	TSharedPtr<FRegionProvider> RegionProvider = MakeShared<FRegionProvider>(*Session);
-	Session->AddProvider(FRegionProvider::ProviderName, TSharedPtr<IRegionProvider>(RegionProvider), TSharedPtr<IEditableRegionProvider>(RegionProvider));
+	Session->AddProvider(GetRegionProviderName(), RegionProvider, RegionProvider);
 
 	TSharedPtr<FLogProvider> LogProvider = MakeShared<FLogProvider>(*Session);
-	Session->AddProvider(FLogProvider::ProviderName, TSharedPtr<ILogProvider>(LogProvider), TSharedPtr<IEditableLogProvider>(LogProvider));
+	Session->AddProvider(FLogProvider::ProviderName, LogProvider, LogProvider);
 
 	TSharedPtr<FThreadProvider> ThreadProvider = MakeShared<FThreadProvider>(*Session);
-	Session->AddProvider(FThreadProvider::ProviderName, TSharedPtr<IThreadProvider>(ThreadProvider), TSharedPtr<IEditableThreadProvider>(ThreadProvider));
+	Session->AddProvider(FThreadProvider::ProviderName, ThreadProvider, ThreadProvider);
 
 	TSharedPtr<FFrameProvider> FrameProvider = MakeShared<FFrameProvider>(*Session);
 	Session->AddProvider(FFrameProvider::ProviderName, FrameProvider);
 
 	TSharedPtr<FCounterProvider> CounterProvider = MakeShared<FCounterProvider>(*Session, *FrameProvider);
-	Session->AddProvider(FCounterProvider::ProviderName, TSharedPtr<ICounterProvider>(CounterProvider), TSharedPtr<IEditableCounterProvider>(CounterProvider));
+	Session->AddProvider(FCounterProvider::ProviderName, CounterProvider, CounterProvider);
 
 	TSharedPtr<FChannelProvider> ChannelProvider = MakeShared<FChannelProvider>();
 	Session->AddProvider(FChannelProvider::ProviderName, ChannelProvider);
@@ -263,15 +259,14 @@ TSharedPtr<const IAnalysisSession> FAnalysisService::StartAnalysis(uint32 TraceI
 	TSharedPtr<FScreenshotProvider> ScreenshotProvider = MakeShared<FScreenshotProvider>(*Session);
 	Session->AddProvider(FScreenshotProvider::ProviderName, ScreenshotProvider);
 
+	TSharedPtr<FDefinitionProvider> DefProvider = MakeShared<FDefinitionProvider>(&Session.Get());
+	Session->AddProvider(FDefinitionProvider::ProviderName, DefProvider, DefProvider);
+
 	Session->AddAnalyzer(new FMiscTraceAnalyzer(*Session, *ThreadProvider, *LogProvider, *FrameProvider, *ChannelProvider, *ScreenshotProvider, *RegionProvider));
 	Session->AddAnalyzer(new FBookmarksAnalyzer(*Session, *BookmarkProvider, LogProvider.Get()));
 	Session->AddAnalyzer(new FLogTraceAnalyzer(*Session, *LogProvider));
-
 	Session->AddAnalyzer(new FStringsAnalyzer(*Session));
 
-	TSharedPtr<FDefinitionProvider> DefProvider = MakeShared<FDefinitionProvider>(&Session.Get());
-	Session->AddProvider(FDefinitionProvider::ProviderName, DefProvider, DefProvider);
-
 	ModuleService.OnAnalysisBegin(*Session);
 
 	Session->Start();
diff --git a/Engine/Source/Developer/TraceServices/Private/AnalysisServicePrivate.h b/Engine/Source/Developer/TraceServices/Private/AnalysisServicePrivate.h
--- a/Engine/Source/Developer/TraceServices/Private/AnalysisServicePrivate.h
+++ b/Engine/Source/Developer/TraceServices/Private/AnalysisServicePrivate.h
@@ -70,7 +70,7 @@ public:
 
 private:
 	virtual const IProvider* ReadProviderPrivate(const FName& Name) const override;
-	virtual IProvider* EditProviderPrivate(const FName& Name) override;
+	virtual IEditableProvider* EditProviderPrivate(const FName& Name) override;
 
 	mutable FAnalysisSessionLock Lock;
 
diff --git a/Engine/Source/Developer/TraceServices/Private/Analyzers/MiscTraceAnalysis.cpp b/Engine/Source/Developer/TraceServices/Private/Analyzers/MiscTraceAnalysis.cpp
--- a/Engine/Source/Developer/TraceServices/Private/Analyzers/MiscTraceAnalysis.cpp
+++ b/Engine/Source/Developer/TraceServices/Private/Analyzers/MiscTraceAnalysis.cpp
@@ -60,7 +60,7 @@ void FMiscTraceAnalyzer::OnAnalysisBegin(const FOnAnalysisContext& Context)
 	Builder.RouteEvent(RouteId_ChannelToggle, "Trace", "ChannelToggle");
 	Builder.RouteEvent(RouteId_ScreenshotHeader, "Misc", "ScreenshotHeader");
 	Builder.RouteEvent(RouteId_ScreenshotChunk, "Misc", "ScreenshotChunk");
-	
+
 	Builder.RouteEvent(RouteId_RegionBegin, "Misc", "RegionBegin");
 	Builder.RouteEvent(RouteId_RegionEnd, "Misc", "RegionEnd");
 }
@@ -161,7 +161,7 @@ bool FMiscTraceAnalyzer::OnEvent(uint16 RouteId, EStyle Style, const FOnEventCon
 		Screenshot->Id = Id;
 
 		EventData.GetString("Name", Screenshot->Name);
-		
+
 		uint64 Cycle = EventData.GetValue<uint64>("Cycle");
 		Screenshot->Timestamp = Context.EventTime.AsSeconds(Cycle);
 
@@ -194,6 +194,7 @@ bool FMiscTraceAnalyzer::OnEvent(uint16 RouteId, EStyle Style, const FOnEventCon
 
 		break;
 	}
+
 	// Begin retired events
 	//
 	case RouteId_RegisterGameThread:
@@ -242,31 +243,30 @@ bool FMiscTraceAnalyzer::OnEvent(uint16 RouteId, EStyle Style, const FOnEventCon
 		ThreadState->ThreadGroupStack.Pop();
 		break;
 	}
-	
+	//
+	// End retired events
+
 	case RouteId_RegionBegin:
 	{
 		uint64 Cycle = EventData.GetValue<uint64>("Cycle");
 		FString Name;
 		EventData.GetString("RegionName", Name);
 
-		FProviderEditScopeLock RegionProviderScopedLock(static_cast<const IRegionProvider&>(RegionProvider));
+		FProviderEditScopeLock RegionProviderScopedLock(RegionProvider);
 		RegionProvider.AppendRegionBegin(*Name, Context.EventTime.AsSeconds(Cycle));
 		break;
 	}
-	
+
 	case RouteId_RegionEnd:
 	{
-		uint64 EndCycle = EventData.GetValue<uint64>("Cycle");
+		uint64 Cycle = EventData.GetValue<uint64>("Cycle");
 		FString Name = TEXT("Invalid");
 		EventData.GetString("RegionName", Name);
 
-		FProviderEditScopeLock RegionProviderScopedLock(static_cast<const IRegionProvider&>(RegionProvider));
-		RegionProvider.AppendRegionEnd(*Name, Context.EventTime.AsSeconds(EndCycle));
+		FProviderEditScopeLock RegionProviderScopedLock(RegionProvider);
+		RegionProvider.AppendRegionEnd(*Name, Context.EventTime.AsSeconds(Cycle));
 		break;
 	}
-		
-	//
-	// End retired events
 	}
 
 	return true;
diff --git a/Engine/Source/Developer/TraceServices/Private/Analyzers/StringsAnalyzer.cpp b/Engine/Source/Developer/TraceServices/Private/Analyzers/StringsAnalyzer.cpp
--- a/Engine/Source/Developer/TraceServices/Private/Analyzers/StringsAnalyzer.cpp
+++ b/Engine/Source/Developer/TraceServices/Private/Analyzers/StringsAnalyzer.cpp
@@ -1,6 +1,7 @@
 ﻿// Copyright Epic Games, Inc. All Rights Reserved.
 
 #include "StringsAnalyzer.h"
+#include "Common/ProviderLock.h"
 #include "Common/Utils.h"
 #include "TraceServices/Model/Definitions.h"
 #include "TraceServices/Model/Strings.h"
@@ -27,7 +28,7 @@ bool FStringsAnalyzer::OnEvent(uint16 RouteId, EStyle Style, const FOnEventConte
 	case RouteId_FName:
 		{
 			const FEventData& EventData = Context.EventData;
-			IDefinitionProvider* DefinitionProvider = GetDefinitionProvider(Session);
+			IDefinitionProvider* DefinitionProvider = EditDefinitionProvider(Session);
 			check(DefinitionProvider != nullptr);
 
 			FWideStringView DisplayWide;
@@ -48,7 +49,7 @@ bool FStringsAnalyzer::OnEvent(uint16 RouteId, EStyle Style, const FOnEventConte
 				return true;
 			}
 
-			IDefinitionProvider::FEditScopeLock Lock(*DefinitionProvider);
+			FProviderEditScopeLock Lock(*DefinitionProvider);
 			FStringDefinition* Instance = DefinitionProvider->Create<FStringDefinition>();
 			Instance->Display = Display;
 
diff --git a/Engine/Source/Developer/TraceServices/Private/Model/AllocationsProvider.h b/Engine/Source/Developer/TraceServices/Private/Model/AllocationsProvider.h
--- a/Engine/Source/Developer/TraceServices/Private/Model/AllocationsProvider.h
+++ b/Engine/Source/Developer/TraceServices/Private/Model/AllocationsProvider.h
@@ -261,7 +261,9 @@ private:
 
 ////////////////////////////////////////////////////////////////////////////////////////////////////
 
-class FAllocationsProvider : public IAllocationsProvider
+class FAllocationsProvider
+	: public IAllocationsProvider
+	, public IEditableProvider
 {
 private:
 	static constexpr double DefaultTimelineSampleGranularity = 0.0001; // 0.1ms
@@ -281,16 +283,13 @@ public:
 	explicit FAllocationsProvider(IAnalysisSession& InSession, FMetadataProvider& InMetadataProvider);
 	virtual ~FAllocationsProvider();
 
-	virtual void BeginEdit() const override       { Lock.BeginWrite(GAllocationsProviderLockState); }
-	virtual void EndEdit() const override         { Lock.EndWrite(GAllocationsProviderLockState); }
-	virtual void EditAccessCheck() const override { Lock.WriteAccessCheck(GAllocationsProviderLockState); }
+	//////////////////////////////////////////////////
+	// Read operations
+
 	virtual void BeginRead() const override       { Lock.BeginRead(GAllocationsProviderLockState); }
 	virtual void EndRead() const override         { Lock.EndRead(GAllocationsProviderLockState); }
 	virtual void ReadAccessCheck() const override { Lock.ReadAccessCheck(GAllocationsProviderLockState); }
 
-	//////////////////////////////////////////////////
-	// Read operations
-
 	virtual bool IsInitialized() const override { ReadAccessCheck(); return bInitialized; }
 
 	virtual void EnumerateTags(TFunctionRef<void(const TCHAR*, const TCHAR*, TagIdType, TagIdType)> Callback) const override;
@@ -325,6 +324,10 @@ public:
 	//////////////////////////////////////////////////
 	// Edit operations
 
+	virtual void BeginEdit() const override       { Lock.BeginWrite(GAllocationsProviderLockState); }
+	virtual void EndEdit() const override         { Lock.EndWrite(GAllocationsProviderLockState); }
+	virtual void EditAccessCheck() const override { Lock.WriteAccessCheck(GAllocationsProviderLockState); }
+
 	void EditInit(double Time, uint8 MinAlignment);
 
 	void EditAlloc(double Time, uint32 CallstackId, uint64 Address, uint64 Size, uint32 Alignment, HeapId RootHeap);
diff --git a/Engine/Source/Developer/TraceServices/Private/Model/Bookmarks.cpp b/Engine/Source/Developer/TraceServices/Private/Model/Bookmarks.cpp
--- a/Engine/Source/Developer/TraceServices/Private/Model/Bookmarks.cpp
+++ b/Engine/Source/Developer/TraceServices/Private/Model/Bookmarks.cpp
@@ -11,19 +11,18 @@
 namespace TraceServices
 {
 
-const FName FBookmarkProvider::ProviderName("BookmarkProvider");
 static constexpr FStringView RegionStartToken = TEXTVIEW("RegionStart:");
 static constexpr FStringView RegionEndToken = TEXTVIEW("RegionEnd:");
 
 FBookmarkProvider::FBookmarkProvider(IAnalysisSession& InSession)
 	: Session(InSession)
 {
-
 }
 
 FBookmarkSpec& FBookmarkProvider::GetSpec(uint64 BookmarkPoint)
 {
 	Session.WriteAccessCheck();
+
 	TSharedPtr<FBookmarkSpec>* Found = SpecMap.Find(BookmarkPoint);
 	if (Found)
 	{
@@ -40,6 +39,7 @@ FBookmarkSpec& FBookmarkProvider::GetSpec(uint64 BookmarkPoint)
 void FBookmarkProvider::UpdateBookmarkSpec(uint64 BookmarkPoint, const TCHAR* FormatString, const TCHAR* File, int32 Line)
 {
 	Session.WriteAccessCheck();
+
 	FBookmarkSpec& BookmarkSpec = GetSpec(BookmarkPoint);
 	BookmarkSpec.FormatString = FormatString;
 	BookmarkSpec.File = File;
@@ -49,6 +49,7 @@ void FBookmarkProvider::UpdateBookmarkSpec(uint64 BookmarkPoint, const TCHAR* Fo
 void FBookmarkProvider::AppendBookmark(uint64 BookmarkPoint, double Time, const uint8* FormatArgs)
 {
 	Session.WriteAccessCheck();
+
 	FBookmarkSpec Spec = GetSpec(BookmarkPoint);
 	FFormatArgsHelper::Format(FormatBuffer, FormatBufferSize - 1, TempBuffer, FormatBufferSize - 1, Spec.FormatString, FormatArgs);
 	TSharedRef<FBookmarkInternal> Bookmark = MakeShared<FBookmarkInternal>();
@@ -57,38 +58,39 @@ void FBookmarkProvider::AppendBookmark(uint64 BookmarkPoint, double Time, const
 	Bookmarks.Add(Bookmark);
 
 	CheckBookmarkForRegion(Bookmark);
-	
+
 	Session.UpdateDurationSeconds(Time);
 }
 
 void FBookmarkProvider::AppendBookmark(uint64 BookmarkPoint, double Time, const TCHAR* Text)
 {
 	Session.WriteAccessCheck();
+
 	TSharedRef<FBookmarkInternal> Bookmark = MakeShared<FBookmarkInternal>();
 	Bookmark->Time = Time;
 	Bookmark->Text = Text;
 	Bookmarks.Add(Bookmark);
 
 	CheckBookmarkForRegion(Bookmark);
-	
+
 	Session.UpdateDurationSeconds(Time);
 }
 
 void FBookmarkProvider::CheckBookmarkForRegion(const TSharedRef<FBookmarkInternal> Bookmark) const
 {
-	const FStringView Text = Bookmark->Text;
-	IEditableRegionProvider* EditableRegionProvider = Session.EditProvider<IEditableRegionProvider>(FRegionProvider::ProviderName);
-
+	IEditableRegionProvider* EditableRegionProvider = Session.EditProvider<IEditableRegionProvider>(GetRegionProviderName());
 	if (!EditableRegionProvider)
 	{
 		return;
 	}
-	
+
+	const FStringView Text = Bookmark->Text;
+
 	if (Text.StartsWith(RegionStartToken))
 	{
 		// StringView.GetData() is not necessarily null-terminated. Since we started from a null terminated string
 		// and only called RightChop() we should still be fine.
-		
+
 		FProviderEditScopeLock _(*EditableRegionProvider);
 		EditableRegionProvider->AppendRegionBegin(Text.RightChop(RegionStartToken.Len()).GetData(), Bookmark->Time);
 	}
@@ -102,6 +104,7 @@ void FBookmarkProvider::CheckBookmarkForRegion(const TSharedRef<FBookmarkInterna
 void FBookmarkProvider::EnumerateBookmarks(double IntervalStart, double IntervalEnd, TFunctionRef<void(const FBookmark &)> Callback) const
 {
 	Session.ReadAccessCheck();
+
 	if (IntervalStart > IntervalEnd)
 	{
 		return;
@@ -134,14 +137,20 @@ void FBookmarkProvider::EnumerateBookmarks(double IntervalStart, double Interval
 	}
 }
 
+FName GetBookmarkProviderName()
+{
+	static FName Name(TEXT("BookmarkProvider"));
+	return Name;
+}
+
 const IBookmarkProvider& ReadBookmarkProvider(const IAnalysisSession& Session)
 {
-	return *Session.ReadProvider<IBookmarkProvider>(FBookmarkProvider::ProviderName);
+	return *Session.ReadProvider<IBookmarkProvider>(GetBookmarkProviderName());
 }
 
 IEditableBookmarkProvider& EditBookmarkProvider(IAnalysisSession& Session)
 {
-	return *Session.EditProvider<IEditableBookmarkProvider>(FBookmarkProvider::ProviderName);
+	return *Session.EditProvider<IEditableBookmarkProvider>(GetBookmarkProviderName());
 }
 
 } // namespace TraceServices
diff --git a/Engine/Source/Developer/TraceServices/Private/Model/BookmarksPrivate.h b/Engine/Source/Developer/TraceServices/Private/Model/BookmarksPrivate.h
--- a/Engine/Source/Developer/TraceServices/Private/Model/BookmarksPrivate.h
+++ b/Engine/Source/Developer/TraceServices/Private/Model/BookmarksPrivate.h
@@ -29,8 +29,6 @@ class FBookmarkProvider
 	, public IEditableBookmarkProvider
 {
 public:
-	static const FName ProviderName;
-
 	explicit FBookmarkProvider(IAnalysisSession& Session);
 	virtual ~FBookmarkProvider() {}
 
diff --git a/Engine/Source/Developer/TraceServices/Private/Model/CookProfilerProviderPrivate.h b/Engine/Source/Developer/TraceServices/Private/Model/CookProfilerProviderPrivate.h
--- a/Engine/Source/Developer/TraceServices/Private/Model/CookProfilerProviderPrivate.h
+++ b/Engine/Source/Developer/TraceServices/Private/Model/CookProfilerProviderPrivate.h
@@ -13,23 +13,32 @@ extern thread_local FProviderLock::FThreadLocalState GCookProviderLockState;
 
 class FCookProfilerProvider
 	: public ICookProfilerProvider
+	, public IEditableCookProfilerProvider
 {
 public:
 	explicit FCookProfilerProvider(IAnalysisSession& Session);
 	virtual ~FCookProfilerProvider() {}
 
-	virtual void BeginEdit() const override       { Lock.BeginWrite(GCookProviderLockState); }
-	virtual void EndEdit() const override         { Lock.EndWrite(GCookProviderLockState); }
-	virtual void EditAccessCheck() const override { Lock.WriteAccessCheck(GCookProviderLockState); }
+	//////////////////////////////////////////////////
+	// Read operations
+
 	virtual void BeginRead() const override       { Lock.BeginRead(GCookProviderLockState); }
 	virtual void EndRead() const override         { Lock.EndRead(GCookProviderLockState); }
 	virtual void ReadAccessCheck() const override { Lock.ReadAccessCheck(GCookProviderLockState); }
 
-	FPackageData* EditPackage(uint64 Id);
-
-	// Functions to query the data.
-	virtual void EnumeratePackages(double StartTime, double EndTime, EnumeratePackagesCallback Callback) const override;
 	virtual uint32 GetNumPackages() const;
+	virtual void EnumeratePackages(double StartTime, double EndTime, EnumeratePackagesCallback Callback) const override;
+
+	//////////////////////////////////////////////////
+	// Edit operations
+
+	virtual void BeginEdit() const override       { Lock.BeginWrite(GCookProviderLockState); }
+	virtual void EndEdit() const override         { Lock.EndWrite(GCookProviderLockState); }
+	virtual void EditAccessCheck() const override { Lock.WriteAccessCheck(GCookProviderLockState); }
+
+	virtual FPackageData* EditPackage(uint64 Id) override;
+
+	//////////////////////////////////////////////////
 
 private:
 	uint32 FindOrAddPackage(uint64 Id);
diff --git a/Engine/Source/Developer/TraceServices/Private/Model/DefinitionProvider.cpp b/Engine/Source/Developer/TraceServices/Private/Model/DefinitionProvider.cpp
--- a/Engine/Source/Developer/TraceServices/Private/Model/DefinitionProvider.cpp
+++ b/Engine/Source/Developer/TraceServices/Private/Model/DefinitionProvider.cpp
@@ -56,16 +56,16 @@ void* FDefinitionProvider::Allocate(uint32 Size, uint32 Alignment)
 
 ////////////////////////////////////////////////////////////////////////////////////////////////////
 
-IDefinitionProvider* GetDefinitionProvider(IAnalysisSession& Session)
+const IDefinitionProvider* ReadDefinitionProvider(const IAnalysisSession& Session)
 {
-	return Session.EditProvider<IDefinitionProvider>(FDefinitionProvider::ProviderName);
+	return Session.ReadProvider<IDefinitionProvider>(FDefinitionProvider::ProviderName);
 }
 
 ////////////////////////////////////////////////////////////////////////////////////////////////////
 
-const IDefinitionProvider* ReadDefinitionProvider(const IAnalysisSession& Session)
+IDefinitionProvider* EditDefinitionProvider(IAnalysisSession& Session)
 {
-	return Session.ReadProvider<IDefinitionProvider>(FDefinitionProvider::ProviderName);
+	return Session.EditProvider<IDefinitionProvider>(FDefinitionProvider::ProviderName);
 }
 
 ////////////////////////////////////////////////////////////////////////////////////////////////////
diff --git a/Engine/Source/Developer/TraceServices/Private/Model/DefinitionProvider.h b/Engine/Source/Developer/TraceServices/Private/Model/DefinitionProvider.h
--- a/Engine/Source/Developer/TraceServices/Private/Model/DefinitionProvider.h
+++ b/Engine/Source/Developer/TraceServices/Private/Model/DefinitionProvider.h
@@ -19,10 +19,13 @@ public:
 
 	FDefinitionProvider(IAnalysisSession* InSession);
 
-	virtual void BeginEdit() const override { Lock.WriteLock(); }
-	virtual void EndEdit() const override { Lock.WriteUnlock(); }
 	virtual void BeginRead() const override { Lock.ReadLock(); }
 	virtual void EndRead() const override { Lock.ReadUnlock(); }
+	virtual void ReadAccessCheck() const override { }
+
+	virtual void BeginEdit() const override { Lock.WriteLock(); }
+	virtual void EndEdit() const override { Lock.WriteUnlock(); }
+	virtual void EditAccessCheck() const override { }
 
 private:
 	virtual void AddEntry(uint64 Hash, const void* Ptr) override;
@@ -30,12 +33,13 @@ private:
 	virtual void* Allocate(uint32 Size, uint32 Alignment) override;
 
 private:
+	mutable FRWLock Lock;
+
 	TArray<TUniquePtr<uint8>> Pages;
 	static constexpr uint32 PageSize = 1024;
 	uint32 PageRemain;
 
 	TMap<uint64, const void*> Definitions;
-	mutable FRWLock Lock;
 };
 
 } // namespace TraceServices
diff --git a/Engine/Source/Developer/TraceServices/Private/Model/MetadataProvider.cpp b/Engine/Source/Developer/TraceServices/Private/Model/MetadataProvider.cpp
--- a/Engine/Source/Developer/TraceServices/Private/Model/MetadataProvider.cpp
+++ b/Engine/Source/Developer/TraceServices/Private/Model/MetadataProvider.cpp
@@ -672,6 +672,13 @@ const IMetadataProvider* ReadMetadataProvider(const IAnalysisSession& Session)
 
 ////////////////////////////////////////////////////////////////////////////////////////////////////
 
+IEditableMetadataProvider* EditMetadataProvider(IAnalysisSession& Session)
+{
+	return Session.EditProvider<IEditableMetadataProvider>(GetMetadataProviderName());
+}
+
+////////////////////////////////////////////////////////////////////////////////////////////////////
+
 } // namespace TraceServices
 
 #undef METADATA_PROVIDER_ERROR
diff --git a/Engine/Source/Developer/TraceServices/Private/Model/MetadataProvider.h b/Engine/Source/Developer/TraceServices/Private/Model/MetadataProvider.h
--- a/Engine/Source/Developer/TraceServices/Private/Model/MetadataProvider.h
+++ b/Engine/Source/Developer/TraceServices/Private/Model/MetadataProvider.h
@@ -34,7 +34,9 @@ private:
 
 ////////////////////////////////////////////////////////////////////////////////////////////////////
 
-class FMetadataProvider : public IMetadataProvider
+class FMetadataProvider
+	: public IMetadataProvider
+	, public IEditableMetadataProvider
 {
 public:
 	static constexpr uint32 MaxMetadataSize = 0xFFFF;
@@ -95,43 +97,46 @@ public:
 	explicit FMetadataProvider(IAnalysisSession& InSession);
 	virtual ~FMetadataProvider();
 
-	virtual void BeginEdit() const override       { Lock.BeginWrite(GMetadataProviderLockState); }
-	virtual void EndEdit() const override         { Lock.EndWrite(GMetadataProviderLockState); }
-	virtual void EditAccessCheck() const override { Lock.WriteAccessCheck(GMetadataProviderLockState); }
+	//////////////////////////////////////////////////
+	// Read operations
+
 	virtual void BeginRead() const override       { Lock.BeginRead(GMetadataProviderLockState); }
 	virtual void EndRead() const override         { Lock.EndRead(GMetadataProviderLockState); }
 	virtual void ReadAccessCheck() const override { Lock.ReadAccessCheck(GMetadataProviderLockState); }
 
+	virtual uint16 GetRegisteredMetadataType(FName InName) const override;
+	virtual FName GetRegisteredMetadataName(uint16 InType) const override;
+	virtual const FMetadataSchema* GetRegisteredMetadataSchema(uint16) const override;
+
+	virtual uint32 GetMetadataStackSize(uint32 InThreadId, uint32 InMetadataId) const override;
+	virtual bool GetMetadata(uint32 InThreadId, uint32 InMetadataId, uint32 InStackDepth, uint16& OutType, const void*& OutData, uint32& OutSize) const override;
+	virtual void EnumerateMetadata(uint32 InThreadId, uint32 InMetadataId, TFunctionRef<bool(uint32 StackDepth, uint16 Type, const void* Data, uint32 Size)> Callback) const override;
+
 	//////////////////////////////////////////////////
 	// Edit operations
 
-	uint16 RegisterMetadataType(const TCHAR* InName, const FMetadataSchema& InSchema);
+	virtual void BeginEdit() const override       { Lock.BeginWrite(GMetadataProviderLockState); }
+	virtual void EndEdit() const override         { Lock.EndWrite(GMetadataProviderLockState); }
+	virtual void EditAccessCheck() const override { Lock.WriteAccessCheck(GMetadataProviderLockState); }
+
+	virtual uint16 RegisterMetadataType(const TCHAR* InName, const FMetadataSchema& InSchema) override;
 
-	void PushScopedMetadata(uint32 InThreadId, uint16 InType, const void* InData, uint32 InSize);
-	void PopScopedMetadata(uint32 InThreadId, uint16 InType);
+	virtual void PushScopedMetadata(uint32 InThreadId, uint16 InType, const void* InData, uint32 InSize) override;
+	virtual void PopScopedMetadata(uint32 InThreadId, uint16 InType) override;
 
-	void BeginClearStackScope(uint32 InThreadId);
-	void EndClearStackScope(uint32 InThreadId);
+	virtual void BeginClearStackScope(uint32 InThreadId) override;
+	virtual void EndClearStackScope(uint32 InThreadId) override;
 
-	void SaveStack(uint32 InThreadId, uint32 InSavedStackId);
-	void BeginRestoreSavedStackScope(uint32 InThreadId, uint32 InSavedStackId);
-	void EndRestoreSavedStackScope(uint32 InThreadId);
+	virtual void SaveStack(uint32 InThreadId, uint32 InSavedStackId) override;
+	virtual void BeginRestoreSavedStackScope(uint32 InThreadId, uint32 InSavedStackId) override;
+	virtual void EndRestoreSavedStackScope(uint32 InThreadId) override;
 
 	// Pins the metadata stack and returns an id for it.
-	uint32 PinAndGetId(uint32 InThreadId);
+	virtual uint32 PinAndGetId(uint32 InThreadId) override;
 
-	void OnAnalysisCompleted();
+	virtual void OnAnalysisCompleted() override;
 
 	//////////////////////////////////////////////////
-	// Read operations
-
-	virtual uint16 GetRegisteredMetadataType(FName InName) const override;
-	virtual FName GetRegisteredMetadataName(uint16 InType) const override;
-	virtual const FMetadataSchema* GetRegisteredMetadataSchema(uint16) const override;
-
-	virtual uint32 GetMetadataStackSize(uint32 InThreadId, uint32 InMetadataId) const override;
-	virtual bool GetMetadata(uint32 InThreadId, uint32 InMetadataId, uint32 InStackDepth, uint16& OutType, const void*& OutData, uint32& OutSize) const override;
-	virtual void EnumerateMetadata(uint32 InThreadId, uint32 InMetadataId, TFunctionRef<bool(uint32 StackDepth, uint16 Type, const void* Data, uint32 Size)> Callback) const override;
 
 private:
 	FMetadataThread& GetOrAddThread(uint32 InThreadId);
diff --git a/Engine/Source/Developer/TraceServices/Private/Model/Regions.cpp b/Engine/Source/Developer/TraceServices/Private/Model/Regions.cpp
--- a/Engine/Source/Developer/TraceServices/Private/Model/Regions.cpp
+++ b/Engine/Source/Developer/TraceServices/Private/Model/Regions.cpp
@@ -1,12 +1,13 @@
-﻿// Copyright Epic Games, Inc. All Rights Reserved.
+// Copyright Epic Games, Inc. All Rights Reserved.
 
 #include "TraceServices/Model/Regions.h"
-#include "Common/FormatArgs.h"
-#include "Model/RegionsPrivate.h"
+
 #include "AnalysisServicePrivate.h"
+#include "Common/FormatArgs.h"
 #include "Common/Utils.h"
-#include "Logging/MessageLog.h"
 #include "Internationalization/Internationalization.h"
+#include "Logging/MessageLog.h"
+#include "Model/RegionsPrivate.h"
 
 #define LOCTEXT_NAMESPACE "RegionProvider"
 
@@ -14,14 +15,12 @@ namespace TraceServices
 {
 
 thread_local FProviderLock::FThreadLocalState GRegionsProviderLockState;
-	
-const FName FRegionProvider::ProviderName("RegionProvider");
 
 FRegionProvider::FRegionProvider(IAnalysisSession& InSession)
 	: Session(InSession)
 {
 }
-	
+
 uint64 FRegionProvider::GetRegionCount() const
 {
 	ReadAccessCheck();
@@ -37,6 +36,7 @@ uint64 FRegionProvider::GetRegionCount() const
 const FRegionLane* FRegionProvider::GetLane(int32 index) const
 {
 	ReadAccessCheck();
+
 	if (index < Lanes.Num())
 	{
 		return &(Lanes[index]);
@@ -47,27 +47,46 @@ const FRegionLane* FRegionProvider::GetLane(int32 index) const
 void FRegionProvider::AppendRegionBegin(const TCHAR* Name, double Time)
 {
 	EditAccessCheck();
-	
-	FTimeRegion Region;
-	Region.BeginTime = Time;
-	Region.Text = Session.StoreString(Name);
-	Region.Depth = CalculateRegionDepth(Region);
 
-	if (Region.Depth == Lanes.Num())
+	if (FTimeRegion** OpenRegion = OpenRegions.Find(Name))
+	{
+		++NumWarnings;
+		if (NumWarnings <= MaxWarningMessages)
+		{
+			if (FMessageLog* Log = Session.GetLog())
+			{
+				Log->Warning(FText::Format(LOCTEXT("DoubleRegionBegin", "A region begin event ({0}) was encountered while a region with same name is already open."), FText::FromString(FString(Name))));
+			}
+			else
+			{
+				UE_LOG(LogTraceServices, Warning, TEXT("[Regions] A region begin event (%s) was encountered while a region with same name is already open."), Name)
+			}
+		}
+	}
+	else
 	{
-		Lanes.Emplace(Session.GetLinearAllocator());
+		FTimeRegion Region;
+		Region.BeginTime = Time;
+		Region.Text = Session.StoreString(Name);
+		Region.Depth = CalculateRegionDepth(Region);
+
+		if (Region.Depth == Lanes.Num())
+		{
+			Lanes.Emplace(Session.GetLinearAllocator());
+		}
+
+		Lanes[Region.Depth].Regions.EmplaceBack(Region);
+		FTimeRegion* NewOpenRegion = &(Lanes[Region.Depth].Regions.Last());
+
+		OpenRegions.Add(Region.Text, NewOpenRegion);
+		UpdateCounter++;
 	}
-	
-	Lanes[Region.Depth].Regions.EmplaceBack(Region);
-	FTimeRegion* OpenRegion = &(Lanes[Region.Depth].Regions.Last());
-	OpenRegions.Add(Region.Text, OpenRegion);
 
+	// Update session time
 	{
 		FAnalysisSessionEditScope _(Session);
 		Session.UpdateDurationSeconds(Time);
 	}
-	
-	UpdateCounter++;
 }
 
 void FRegionProvider::AppendRegionEnd(const TCHAR* Name, double Time)
@@ -77,42 +96,62 @@ void FRegionProvider::AppendRegionEnd(const TCHAR* Name, double Time)
 	if (FTimeRegion** OpenRegion = OpenRegions.Find(Name))
 	{
 		(*OpenRegion)->EndTime = Time;
+
 		OpenRegions.Remove(Name);
-		{
-			FAnalysisSessionEditScope _(Session);
-			Session.UpdateDurationSeconds(Time);
-		}
 		UpdateCounter++;
 	}
 	else
 	{
-		if (FMessageLog* Log = Session.GetLog())
-		{
-			Log->Warning(LOCTEXT("StrayRegionEnd","A region end event was encountered without having seen a matching region start event first."));
-		}
-		else
+		++NumWarnings;
+		if (NumWarnings <= MaxWarningMessages)
 		{
-			UE_LOG(LogTraceServices, Warning, TEXT("A region end event (%s) was encountered without having seen a matching region start event first."), Name)
+			if (FMessageLog* Log = Session.GetLog())
+			{
+				Log->Warning(FText::Format(LOCTEXT("StrayRegionEnd", "A region end event ({0}) was encountered without having seen a matching region start event first."), FText::FromString(FString(Name))));
+			}
+			else
+			{
+				UE_LOG(LogTraceServices, Warning, TEXT("[Regions] A region end event (%s) was encountered without having seen a matching region start event first."), Name)
+			}
 		}
 	}
+
+	// Update session time
+	{
+		FAnalysisSessionEditScope _(Session);
+		Session.UpdateDurationSeconds(Time);
+	}
 }
 
 void FRegionProvider::OnAnalysisSessionEnded()
 {
 	EditAccessCheck();
-	
+
 	for (const auto& KV : OpenRegions)
 	{
 		const FTimeRegion* Region = KV.Value;
-		if (FMessageLog* Log = Session.GetLog())
+
+		++NumWarnings;
+		if (NumWarnings <= MaxWarningMessages)
 		{
-			Log->Warning(LOCTEXT("StrayRegionBegin","A region begin event was never closed."));
+			if (FMessageLog* Log = Session.GetLog())
+			{
+				Log->Warning(FText::Format(LOCTEXT("StrayRegionBegin", "A region begin event ({0}) was never closed."), FText::FromString(FString(Region->Text))));
+			}
+			else
+			{
+				UE_LOG(LogTraceServices, Warning, TEXT("[Regions] A region begin event (%s) was never closed."), Region->Text)
+			}
 		}
-		else
-		{
-			UE_LOG(LogTraceServices, Warning, TEXT("A region begin event (%s) was never closed."), Region->Text)
-		}		
 	}
+
+	if (NumWarnings > 0 || NumErrors > 0)
+	{
+		UE_LOG(LogTraceServices, Error, TEXT("[Regions] %u warnings; %u errors"), NumWarnings, NumErrors);
+	}
+
+	uint64 TotalRegionCount = GetRegionCount();
+	UE_LOG(LogTraceServices, Log, TEXT("[Regions] Analysis completed (%llu regions, %d lanes)."), TotalRegionCount, Lanes.Num());
 }
 
 int32 FRegionProvider::CalculateRegionDepth(const FTimeRegion& Region) const
@@ -128,7 +167,7 @@ int32 FRegionProvider::CalculateRegionDepth(const FTimeRegion& Region) const
 		{
 			break;
 		}
-		
+
 		const FTimeRegion& LastRegion = Lanes[NewDepth].Regions.Last();
 		if (LastRegion.EndTime <= Region.BeginTime)
 		{
@@ -141,19 +180,21 @@ int32 FRegionProvider::CalculateRegionDepth(const FTimeRegion& Region) const
 
 	return NewDepth;
 }
-	
+
 void FRegionProvider::EnumerateLanes(TFunctionRef<void(const FRegionLane&, int32)> Callback) const
 {
-	for (int i = 0 ; i < Lanes.Num(); ++i)
+	ReadAccessCheck();
+
+	for (int32 LaneIndex = 0; LaneIndex < Lanes.Num(); ++LaneIndex)
 	{
-		Callback(Lanes[i], i);
+		Callback(Lanes[LaneIndex], LaneIndex);
 	}
 }
-	
-bool FRegionProvider::EnumerateRegions(double IntervalStart, double IntervalEnd, TFunctionRef<bool(const FTimeRegion &)> Callback) const
+
+bool FRegionProvider::EnumerateRegions(double IntervalStart, double IntervalEnd, TFunctionRef<bool(const FTimeRegion&)> Callback) const
 {
 	ReadAccessCheck();
-	
+
 	if (IntervalStart > IntervalEnd)
 	{
 		return false;
@@ -162,7 +203,9 @@ bool FRegionProvider::EnumerateRegions(double IntervalStart, double IntervalEnd,
 	for (const FRegionLane& Lane : Lanes)
 	{
 		if (!Lane.EnumerateRegions(IntervalStart, IntervalEnd, Callback))
+		{
 			return false;
+		}
 	}
 
 	return true;
@@ -171,14 +214,14 @@ bool FRegionProvider::EnumerateRegions(double IntervalStart, double IntervalEnd,
 bool FRegionLane::EnumerateRegions(double IntervalStart, double IntervalEnd, TFunctionRef<bool(const FTimeRegion&)> Callback) const
 {
 	const FInt32Interval OverlapRange = GetElementRangeOverlappingGivenRange<FTimeRegion>(Regions, IntervalStart, IntervalEnd,
-		[](const FTimeRegion& r){return r.BeginTime;},
-		[](const FTimeRegion& r){return r.EndTime;});
-	
-	if (OverlapRange.Min == -1 )
+		[](const FTimeRegion& r) { return r.BeginTime; },
+		[](const FTimeRegion& r) { return r.EndTime; });
+
+	if (OverlapRange.Min == -1)
 	{
 		return true;
 	}
-	
+
 	for (int32 Index = OverlapRange.Min; Index <= OverlapRange.Max; ++Index)
 	{
 		if (!Callback(Regions[Index]))
@@ -189,17 +232,23 @@ bool FRegionLane::EnumerateRegions(double IntervalStart, double IntervalEnd, TFu
 
 	return true;
 }
-	
+
+FName GetRegionProviderName()
+{
+	static FName Name(TEXT("RegionProvider"));
+	return Name;
+}
+
 const IRegionProvider& ReadRegionProvider(const IAnalysisSession& Session)
 {
-	return *Session.ReadProvider<IRegionProvider>(FRegionProvider::ProviderName);
+	return *Session.ReadProvider<IRegionProvider>(GetRegionProviderName());
 }
 
 IEditableRegionProvider& EditRegionProvider(IAnalysisSession& Session)
 {
-	return *Session.EditProvider<IEditableRegionProvider>(FRegionProvider::ProviderName);
+	return *Session.EditProvider<IEditableRegionProvider>(GetRegionProviderName());
 }
 
 } // namespace TraceServices
 
-#undef LOCTEXT_NAMESPACE
\ No newline at end of file
+#undef LOCTEXT_NAMESPACE
diff --git a/Engine/Source/Developer/TraceServices/Private/Model/RegionsPrivate.h b/Engine/Source/Developer/TraceServices/Private/Model/RegionsPrivate.h
--- a/Engine/Source/Developer/TraceServices/Private/Model/RegionsPrivate.h
+++ b/Engine/Source/Developer/TraceServices/Private/Model/RegionsPrivate.h
@@ -1,4 +1,4 @@
-﻿// Copyright Epic Games, Inc. All Rights Reserved.
+// Copyright Epic Games, Inc. All Rights Reserved.
 
 #pragma once
 
@@ -7,9 +7,9 @@
 #include "TraceServices/Model/Regions.h"
 #include "Templates/SharedPointer.h"
 
-
 namespace TraceServices
 {
+
 extern thread_local FProviderLock::FThreadLocalState GRegionsProviderLockState;
 
 class FAnalysisSessionLock;
@@ -20,54 +20,63 @@ class FRegionProvider
 	, public IEditableRegionProvider
 {
 public:
-	static const FName ProviderName;
+	explicit FRegionProvider(IAnalysisSession& Session);
+	virtual ~FRegionProvider() override {}
+
+	//////////////////////////////////////////////////
+	// Read operations
 
-	// Provider locking API
-	virtual void BeginEdit() const override       { Lock.BeginWrite(GRegionsProviderLockState); }
-	virtual void EndEdit() const override         { Lock.EndWrite(GRegionsProviderLockState); }
-	virtual void EditAccessCheck() const override { Lock.WriteAccessCheck(GRegionsProviderLockState); }
 	virtual void BeginRead() const override       { Lock.BeginRead(GRegionsProviderLockState); }
 	virtual void EndRead() const override         { Lock.EndRead(GRegionsProviderLockState); }
 	virtual void ReadAccessCheck() const override { Lock.ReadAccessCheck(GRegionsProviderLockState); }
 
-	explicit FRegionProvider(IAnalysisSession& Session);
-	virtual ~FRegionProvider() override {}
-
-	// implement IRegionProvider API
 	virtual uint64 GetRegionCount() const override;
-	
-	virtual uint64 GetUpdateCounter() const override { ReadAccessCheck(); return UpdateCounter; }
+	virtual int32 GetLaneCount()  const override { ReadAccessCheck(); return Lanes.Num(); }
+
+	virtual const FRegionLane* GetLane(int32 Index) const override;
 
 	virtual bool EnumerateRegions(double IntervalStart, double IntervalEnd, TFunctionRef<bool(const FTimeRegion&)> Callback) const override;
 	virtual void EnumerateLanes(TFunctionRef<void(const FRegionLane&, const int32)> Callback) const override;
-	virtual int32 GetLaneCount()  const override { ReadAccessCheck(); return Lanes.Num(); }
 
-	/**
-	 * Direct Access to a certain lane at a given index/depth.
-	 * Do not store the pointer returned from this function as it could be accessed outside the ProviderEditScope.
-	 */
-	virtual const FRegionLane* GetLane(int32 index) const override;
+	virtual uint64 GetUpdateCounter() const override { ReadAccessCheck(); return UpdateCounter; }
 
+	//////////////////////////////////////////////////
+	// Edit operations
+
+	virtual void BeginEdit() const override       { Lock.BeginWrite(GRegionsProviderLockState); }
+	virtual void EndEdit() const override         { Lock.EndWrite(GRegionsProviderLockState); }
+	virtual void EditAccessCheck() const override { Lock.WriteAccessCheck(GRegionsProviderLockState); }
 
-	// implement IEditableRegionProvider API
 	virtual void AppendRegionBegin(const TCHAR* Name, double Time) override;
 	virtual void AppendRegionEnd(const TCHAR* Name, double Time) override;
+
 	virtual void OnAnalysisSessionEnded() override;
 
-private:
+	//////////////////////////////////////////////////
 
-	mutable FProviderLock Lock;
-	
+private:
 	// Update the depth member of a region to allow overlapping regions to be displayed on separate lanes.
 	int32 CalculateRegionDepth(const FTimeRegion& Item) const;
-	
+
+private:
+	mutable FProviderLock Lock;
+
 	IAnalysisSession& Session;
-	// open Regions inside lanes
+
+	// Open regions inside lanes
 	TMap<FStringView, FTimeRegion*> OpenRegions;
-	// closed regions
+
+	// Closed regions
 	TArray<FRegionLane> Lanes;
-	// counter incremented each time region data changes during analysis
+
+	// Counter incremented each time region data changes during analysis
 	uint64 UpdateCounter = -1;
+
+	static constexpr uint32 MaxWarningMessages = 100;
+	static constexpr uint32 MaxErrorMessages = 100;
+
+	uint32 NumWarnings = 0;
+	uint32 NumErrors = 0;
 };
 
 } // namespace TraceServices
diff --git a/Engine/Source/Developer/TraceServices/Private/Modules/CookProfilerModule.cpp b/Engine/Source/Developer/TraceServices/Private/Modules/CookProfilerModule.cpp
--- a/Engine/Source/Developer/TraceServices/Private/Modules/CookProfilerModule.cpp
+++ b/Engine/Source/Developer/TraceServices/Private/Modules/CookProfilerModule.cpp
@@ -20,7 +20,7 @@ void FCookProfilerModule::GetModuleInfo(FModuleInfo& OutModuleInfo)
 void FCookProfilerModule::OnAnalysisBegin(IAnalysisSession& Session)
 {
 	TSharedPtr<FCookProfilerProvider> CookProfilerProvider = MakeShared<FCookProfilerProvider>(Session);
-	Session.AddProvider(CookProfilerProviderName, CookProfilerProvider);
+	Session.AddProvider(CookProfilerProviderName, CookProfilerProvider, CookProfilerProvider);
 	Session.AddAnalyzer(new FCookAnalyzer(Session, *CookProfilerProvider));
 }
 
@@ -34,4 +34,9 @@ const ICookProfilerProvider* ReadCookProfilerProvider(const IAnalysisSession& Se
 	return Session.ReadProvider<ICookProfilerProvider>(CookProfilerProviderName);
 }
 
+IEditableCookProfilerProvider* EditCookProfilerProvider(IAnalysisSession& Session)
+{
+	return Session.EditProvider<IEditableCookProfilerProvider>(CookProfilerProviderName);
+}
+
 } // namespace TraceServices
diff --git a/Engine/Source/Developer/TraceServices/Private/Modules/MemoryModule.cpp b/Engine/Source/Developer/TraceServices/Private/Modules/MemoryModule.cpp
--- a/Engine/Source/Developer/TraceServices/Private/Modules/MemoryModule.cpp
+++ b/Engine/Source/Developer/TraceServices/Private/Modules/MemoryModule.cpp
@@ -40,12 +40,12 @@ void FMemoryModule::OnAnalysisBegin(IAnalysisSession& Session)
 
 	// Metadata
 	TSharedPtr<FMetadataProvider> MetadataProvider = MakeShared<FMetadataProvider>(Session);
-	Session.AddProvider(GetMetadataProviderName(), MetadataProvider);
+	Session.AddProvider(GetMetadataProviderName(), MetadataProvider, MetadataProvider);
 	Session.AddAnalyzer(new FMetadataAnalysis(Session, MetadataProvider.Get()));
 
 	// Allocations
 	TSharedPtr<FAllocationsProvider> AllocationsProvider = MakeShared<FAllocationsProvider>(Session, *MetadataProvider);
-	Session.AddProvider(GetAllocationsProviderName(), AllocationsProvider);
+	Session.AddProvider(GetAllocationsProviderName(), AllocationsProvider, AllocationsProvider);
 	Session.AddAnalyzer(new FAllocationsAnalyzer(Session, *AllocationsProvider, *MetadataProvider));
 }
 
diff --git a/Engine/Source/Developer/TraceServices/Private/Modules/TimingProfilerModule.cpp b/Engine/Source/Developer/TraceServices/Private/Modules/TimingProfilerModule.cpp
--- a/Engine/Source/Developer/TraceServices/Private/Modules/TimingProfilerModule.cpp
+++ b/Engine/Source/Developer/TraceServices/Private/Modules/TimingProfilerModule.cpp
@@ -28,7 +28,7 @@ void FTimingProfilerModule::OnAnalysisBegin(IAnalysisSession& InSession)
 
 	// see comment in FAnalysisService::StartAnalysis for details
 	TSharedPtr<FTimingProfilerProvider> TimingProfilerProvider = MakeShared<FTimingProfilerProvider>(Session);
-	Session.AddProvider(TimingProfilerProviderName, TSharedPtr<ITimingProfilerProvider>(TimingProfilerProvider), TSharedPtr<IEditableTimingProfilerProvider>(TimingProfilerProvider));
+	Session.AddProvider(TimingProfilerProviderName, TimingProfilerProvider, TimingProfilerProvider);
 	Session.AddAnalyzer(new FCpuProfilerAnalyzer(Session, *TimingProfilerProvider, *EditableThreadProvider));
 	Session.AddAnalyzer(new FGpuProfilerAnalyzer(Session, *TimingProfilerProvider));
 }
diff --git a/Engine/Source/Developer/TraceServices/Public/Common/ProviderLock.h b/Engine/Source/Developer/TraceServices/Public/Common/ProviderLock.h
--- a/Engine/Source/Developer/TraceServices/Public/Common/ProviderLock.h
+++ b/Engine/Source/Developer/TraceServices/Public/Common/ProviderLock.h
@@ -11,7 +11,7 @@ namespace TraceServices
 
 struct TRACESERVICES_API FProviderEditScopeLock
 {
-	FProviderEditScopeLock(const IProvider& InProvider)
+	FProviderEditScopeLock(const IEditableProvider& InProvider)
 		: Provider(InProvider)
 	{
 		Provider.BeginEdit();
@@ -22,7 +22,7 @@ struct TRACESERVICES_API FProviderEditScopeLock
 		Provider.EndEdit();
 	}
 
-	const IProvider& Provider;
+	const IEditableProvider& Provider;
 };
 
 struct TRACESERVICES_API FProviderReadScopeLock
diff --git a/Engine/Source/Developer/TraceServices/Public/TraceServices/Model/AllocationsProvider.h b/Engine/Source/Developer/TraceServices/Public/TraceServices/Model/AllocationsProvider.h
--- a/Engine/Source/Developer/TraceServices/Public/TraceServices/Model/AllocationsProvider.h
+++ b/Engine/Source/Developer/TraceServices/Public/TraceServices/Model/AllocationsProvider.h
@@ -12,13 +12,9 @@
 namespace TraceServices
 {
 
-////////////////////////////////////////////////////////////////////////////////////////////////////
-
 // Id type for tags
 typedef uint32 TagIdType;
 	
-////////////////////////////////////////////////////////////////////////////////////////////////////
-
 class IAllocationsProvider : public IProvider
 {
 public:
@@ -116,10 +112,9 @@ public:
 public:
 	virtual ~IAllocationsProvider() = default;
 
-	virtual void BeginEdit() const = 0;
-	virtual void EndEdit() const = 0;
 	virtual void BeginRead() const = 0;
 	virtual void EndRead() const = 0;
+	virtual void ReadAccessCheck() const = 0;
 
 	virtual bool IsInitialized() const = 0;
 
diff --git a/Engine/Source/Developer/TraceServices/Public/TraceServices/Model/AnalysisSession.h b/Engine/Source/Developer/TraceServices/Public/TraceServices/Model/AnalysisSession.h
--- a/Engine/Source/Developer/TraceServices/Public/TraceServices/Model/AnalysisSession.h
+++ b/Engine/Source/Developer/TraceServices/Public/TraceServices/Model/AnalysisSession.h
@@ -29,18 +29,19 @@ class IProvider
 public:
 	virtual ~IProvider() = default;
 
-	virtual void BeginEdit() const { unimplemented(); }
-	virtual void EndEdit() const { unimplemented(); }
-	virtual void EditAccessCheck() const { unimplemented(); }
-
 	virtual void BeginRead() const { unimplemented(); }
 	virtual void EndRead() const { unimplemented(); }
 	virtual void ReadAccessCheck() const { unimplemented(); }
 };
 
 class IEditableProvider
-	: public IProvider
 {
+public:
+	virtual ~IEditableProvider() = default;
+
+	virtual void BeginEdit() const { unimplemented(); }
+	virtual void EndEdit() const { unimplemented(); }
+	virtual void EditAccessCheck() const { unimplemented(); }
 };
 
 class IAnalysisSession
@@ -78,14 +79,14 @@ public:
 	
 	virtual void AddAnalyzer(UE::Trace::IAnalyzer* Analyzer) = 0;
 
-	UE_DEPRECATED(5.1, "Please use the TSharedPtr overload for registering Provider objects")
-	virtual void AddProvider(const FName& Name, IProvider* Provider) { AddProvider(Name, TSharedPtr<IProvider>(Provider), nullptr); }
 	virtual void AddProvider(const FName& Name, TSharedPtr<IProvider> Provider, TSharedPtr<IEditableProvider> EditableProvider = nullptr) = 0;
+
 	template<typename ProviderType>
 	const ProviderType* ReadProvider(const FName& Name) const
 	{
 		return static_cast<const ProviderType*>(ReadProviderPrivate(Name));
 	}
+
 	template<typename ProviderType>
 	ProviderType* EditProvider(const FName& Name)
 	{
@@ -94,7 +95,7 @@ public:
 
 private:
 	virtual const IProvider* ReadProviderPrivate(const FName& Name) const = 0;
-	virtual IProvider* EditProviderPrivate(const FName& Name) = 0;
+	virtual IEditableProvider* EditProviderPrivate(const FName& Name) = 0;
 };
 
 struct FAnalysisSessionReadScope
diff --git a/Engine/Source/Developer/TraceServices/Public/TraceServices/Model/Bookmarks.h b/Engine/Source/Developer/TraceServices/Public/TraceServices/Model/Bookmarks.h
--- a/Engine/Source/Developer/TraceServices/Public/TraceServices/Model/Bookmarks.h
+++ b/Engine/Source/Developer/TraceServices/Public/TraceServices/Model/Bookmarks.h
@@ -65,6 +65,7 @@ public:
 	virtual void AppendBookmark(uint64 BookmarkPoint, double Time, const TCHAR* Text) = 0;
 };
 
+TRACESERVICES_API FName GetBookmarkProviderName();
 TRACESERVICES_API const IBookmarkProvider& ReadBookmarkProvider(const IAnalysisSession& Session);
 TRACESERVICES_API IEditableBookmarkProvider& EditBookmarkProvider(IAnalysisSession& Session);
 
diff --git a/Engine/Source/Developer/TraceServices/Public/TraceServices/Model/CookProfilerProvider.h b/Engine/Source/Developer/TraceServices/Public/TraceServices/Model/CookProfilerProvider.h
--- a/Engine/Source/Developer/TraceServices/Public/TraceServices/Model/CookProfilerProvider.h
+++ b/Engine/Source/Developer/TraceServices/Public/TraceServices/Model/CookProfilerProvider.h
@@ -22,16 +22,34 @@ struct FPackageData
 	const TCHAR* AssetClass;
 };
 
-class ICookProfilerProvider
-	: public IProvider
+class ICookProfilerProvider : public IProvider
 {
 public:
 	typedef TFunctionRef<bool(const FPackageData& /*Package*/)> EnumeratePackagesCallback;
 
 	virtual ~ICookProfilerProvider() = default;
-	virtual void EnumeratePackages(double StartTime, double EndTime, EnumeratePackagesCallback Callback) const = 0;
+
+	virtual void BeginRead() const = 0;
+	virtual void EndRead() const = 0;
+	virtual void ReadAccessCheck() const = 0;
+
 	virtual uint32 GetNumPackages() const = 0;
+	virtual void EnumeratePackages(double StartTime, double EndTime, EnumeratePackagesCallback Callback) const = 0;
+};
+
+class IEditableCookProfilerProvider : public IEditableProvider
+{
+public:
+	virtual ~IEditableCookProfilerProvider() = default;
+
+	virtual void BeginEdit() const = 0;
+	virtual void EndEdit() const = 0;
+	virtual void EditAccessCheck() const = 0;
+
+	virtual FPackageData* EditPackage(uint64 Id) = 0;
 };
 
 TRACESERVICES_API const ICookProfilerProvider* ReadCookProfilerProvider(const IAnalysisSession& Session);
+TRACESERVICES_API IEditableCookProfilerProvider* EditCookProfilerProvider(IAnalysisSession& Session);
+
 } // namespace TraceServices
diff --git a/Engine/Source/Developer/TraceServices/Public/TraceServices/Model/Definitions.h b/Engine/Source/Developer/TraceServices/Public/TraceServices/Model/Definitions.h
--- a/Engine/Source/Developer/TraceServices/Public/TraceServices/Model/Definitions.h
+++ b/Engine/Source/Developer/TraceServices/Public/TraceServices/Model/Definitions.h
@@ -12,7 +12,7 @@ namespace TraceServices
  * Allows users to publish "definitions", structs representing a definition event. This data can then be queried when
  * resolving a reference field in another event.
  *
- * \note If your definition type in turn references other definitions it encouraged to resolve those during analysis to
+ * \note If your definition type reference other definitions it is encouraged to resolve those during analysis to
  * avoid having to resolve complex chains on every lookup.
  *
  * Example usage:
@@ -28,41 +28,10 @@ namespace TraceServices
  *		DefinitionProvider->Register<FMyDefinitionClass>(Instance, Id);
  * \endcode
  */
-class IDefinitionProvider : public IEditableProvider
+class IDefinitionProvider
+	: public IProvider
+	, public IEditableProvider
 {
-public:
-	struct TRACESERVICES_API FEditScopeLock
-	{
-		FEditScopeLock(const IDefinitionProvider& InMetadataProvider)
-			: MetadataProvider(InMetadataProvider)
-		{
-			MetadataProvider.BeginEdit();
-		}
-
-		~FEditScopeLock()
-		{
-			MetadataProvider.EndEdit();
-		}
-
-		const IDefinitionProvider& MetadataProvider;
-	};
-
-	struct TRACESERVICES_API FReadScopeLock
-	{
-		FReadScopeLock(const IDefinitionProvider& InMetadataProvider)
-			: MetadataProvider(InMetadataProvider)
-		{
-			MetadataProvider.BeginRead();
-		}
-
-		~FReadScopeLock()
-		{
-			MetadataProvider.EndRead();
-		}
-
-		const IDefinitionProvider& MetadataProvider;
-	};
-
 public:
 	virtual ~IDefinitionProvider() = default;
 
@@ -137,7 +106,7 @@ protected:
 	TMap<uint32, StringifierFn> Stringifiers;
 };
 
-TRACESERVICES_API IDefinitionProvider* GetDefinitionProvider(IAnalysisSession& Session);
 TRACESERVICES_API const IDefinitionProvider* ReadDefinitionProvider(const IAnalysisSession& Session);
+TRACESERVICES_API IDefinitionProvider* EditDefinitionProvider(IAnalysisSession& Session);
 
 } // namespace TraceServices
diff --git a/Engine/Source/Developer/TraceServices/Public/TraceServices/Model/MetadataProvider.h b/Engine/Source/Developer/TraceServices/Public/TraceServices/Model/MetadataProvider.h
--- a/Engine/Source/Developer/TraceServices/Public/TraceServices/Model/MetadataProvider.h
+++ b/Engine/Source/Developer/TraceServices/Public/TraceServices/Model/MetadataProvider.h
@@ -11,7 +11,6 @@
 namespace TraceServices
 {
 
-////////////////////////////////////////////////////////////////////////////////////////////////////
 class FMetadataSchema
 {
 public:
@@ -266,8 +265,6 @@ private:
 	TArray<FField> Fields;
 };
 
-////////////////////////////////////////////////////////////////////////////////////////////////////
-
 class IMetadataProvider : public IProvider
 {
 public:
@@ -276,10 +273,9 @@ public:
 
 	virtual ~IMetadataProvider() = default;
 
-	virtual void BeginEdit() const = 0;
-	virtual void EndEdit() const = 0;
 	virtual void BeginRead() const = 0;
 	virtual void EndRead() const = 0;
+	virtual void ReadAccessCheck() const = 0;
 
 	virtual uint16 GetRegisteredMetadataType(FName Name) const = 0;
 	virtual FName GetRegisteredMetadataName(uint16 Type) const = 0;
@@ -290,7 +286,35 @@ public:
 	virtual void EnumerateMetadata(uint32 InThreadId, uint32 InMetadataId, TFunctionRef<bool(uint32 StackDepth, uint16 Type, const void* Data, uint32 Size)> Callback) const = 0;
 };
 
+class IEditableMetadataProvider : public IEditableProvider
+{
+public:
+	virtual ~IEditableMetadataProvider() = default;
+
+	virtual void BeginEdit() const = 0;
+	virtual void EndEdit() const = 0;
+	virtual void EditAccessCheck() const = 0;
+
+	virtual uint16 RegisterMetadataType(const TCHAR* InName, const FMetadataSchema& InSchema) = 0;
+
+	virtual void PushScopedMetadata(uint32 InThreadId, uint16 InType, const void* InData, uint32 InSize) = 0;
+	virtual void PopScopedMetadata(uint32 InThreadId, uint16 InType) = 0;
+
+	virtual void BeginClearStackScope(uint32 InThreadId) = 0;
+	virtual void EndClearStackScope(uint32 InThreadId) = 0;
+
+	virtual void SaveStack(uint32 InThreadId, uint32 InSavedStackId) = 0;
+	virtual void BeginRestoreSavedStackScope(uint32 InThreadId, uint32 InSavedStackId) = 0;
+	virtual void EndRestoreSavedStackScope(uint32 InThreadId) = 0;
+
+	// Pins the metadata stack and returns an id for it.
+	virtual uint32 PinAndGetId(uint32 InThreadId) = 0;
+
+	virtual void OnAnalysisCompleted() { }
+};
+
 TRACESERVICES_API FName GetMetadataProviderName();
 TRACESERVICES_API const IMetadataProvider* ReadMetadataProvider(const IAnalysisSession& Session);
+TRACESERVICES_API IEditableMetadataProvider* EditMetadataProvider(IAnalysisSession& Session);
 
 } // namespace TraceServices
diff --git a/Engine/Source/Developer/TraceServices/Public/TraceServices/Model/Modules.h b/Engine/Source/Developer/TraceServices/Public/TraceServices/Model/Modules.h
--- a/Engine/Source/Developer/TraceServices/Public/TraceServices/Model/Modules.h
+++ b/Engine/Source/Developer/TraceServices/Public/TraceServices/Model/Modules.h
@@ -7,7 +7,8 @@
 #include "Templates/UnrealTemplate.h"
 #include <atomic>
 
-namespace TraceServices {
+namespace TraceServices
+{
 
 ////////////////////////////////////////////////////////////////////////////////
 /**
@@ -154,7 +155,9 @@ struct FModule
 };
 
 ////////////////////////////////////////////////////////////////////////////////
-class IModuleProvider : public IEditableProvider
+class IModuleProvider
+	: public IProvider
+	, public IEditableProvider
 {
 public:
 	struct FStats
diff --git a/Engine/Source/Developer/TraceServices/Public/TraceServices/Model/Regions.h b/Engine/Source/Developer/TraceServices/Public/TraceServices/Model/Regions.h
--- a/Engine/Source/Developer/TraceServices/Public/TraceServices/Model/Regions.h
+++ b/Engine/Source/Developer/TraceServices/Public/TraceServices/Model/Regions.h
@@ -1,8 +1,8 @@
-﻿// Copyright Epic Games, Inc. All Rights Reserved.
+// Copyright Epic Games, Inc. All Rights Reserved.
 
-#pragma  once
+#pragma once
 
-#include "Common/PagedArray.h"
+#include "Common/PagedArray.h" // TraceServices
 #include "HAL/Platform.h"
 #include "Templates/Function.h"
 #include "TraceServices/Model/AnalysisSession.h"
@@ -12,8 +12,6 @@ template <typename FuncType> class TFunctionRef;
 namespace TraceServices
 {
 
-class FRegionProvider;
-
 struct TRACESERVICES_API FTimeRegion
 {
 	double BeginTime = std::numeric_limits<double>::infinity();
@@ -24,69 +22,64 @@ struct TRACESERVICES_API FTimeRegion
 
 class TRACESERVICES_API FRegionLane
 {
-public:
+	friend class FRegionProvider;
 
+public:
 	FRegionLane(ILinearAllocator& InAllocator) : Regions(InAllocator, 512) {}
 
+	int32 Num() const { return Regions.Num(); }
+
 	/**
 	 * Call Callback for every region overlapping the interval defined by IntervalStart and IntervalEnd
 	 * @param Callback a callback called for each region. Return false to abort iteration.
 	 * @returns true if the enumeration finished, false if it was aborted by the callback returning false
 	 */
 	bool EnumerateRegions(double IntervalStart, double IntervalEnd, TFunctionRef<bool(const FTimeRegion&)> Callback) const;
-	int32 Num() const { return Regions.Num(); }
+
 private:
 	TPagedArray<FTimeRegion> Regions;
-	friend class FRegionProvider;
 };
 
-	
 class TRACESERVICES_API IRegionProvider
 	: public IProvider
 {
 public:
-
-	
 	virtual ~IRegionProvider() override = default;
-	
+
 	/**
 	 * @return the amount of currently known regions (including open-ended ones)
 	 */
 	virtual uint64 GetRegionCount() const = 0;
-	
+
 	/**
-	 * @return the amount of currently known regions (including open-ended ones)
+	 * @return the number of lanes
 	 */
 	virtual int32 GetLaneCount() const = 0;
 
 	/**
-	 * @return a pointer to the lane at depth index or nullptr if index > GetLaneCount()-1
+	 * Direct access to a certain lane at a given index/depth.
+	 * The pointer is valid only in the current read scope.
+	 * @return a pointer to the lane at the specified depth index or nullptr if Index > GetLaneCount()-1
 	 */
-	virtual const FRegionLane* GetLane(int32 index) const = 0; 
-	
-	/**
-	 * @return A monotonically increasing counter that that changes each time new data is added to the provider.
-	 * This can be used to detect when to update any (UI-)state dependent on the provider during analysis.
-	 */
-	virtual uint64 GetUpdateCounter() const = 0;
-	
+	virtual const FRegionLane* GetLane(int32 Index) const = 0;
+
 	/**
 	 * Enumerates all regions that overlap a certain time interval. Will enumerate by depth but does not expose lanes.
 	 * @param Callback a callback called for each region. Return false to abort iteration.
 	 * @returns true if the enumeration finished, false if it was aborted by the callback returning false
 	 */
 	virtual bool EnumerateRegions(double IntervalStart, double IntervalEnd, TFunctionRef<bool(const FTimeRegion&)> Callback) const = 0;
+
 	/**
 	 * Will call Callback(Lane, Depth) for each lane in order.
 	 */
 	virtual void EnumerateLanes(TFunctionRef<void(const FRegionLane&, const int32)> Callback) const = 0;
 
-	// Provider locking API
-	virtual void BeginEdit() const override = 0;
-	virtual void EndEdit() const override = 0;
-	virtual void BeginRead() const override = 0;
-	virtual void EndRead() const override = 0;
-	
+	/**
+	 * @return A monotonically increasing counter that that changes each time new data is added to the provider.
+	 * This can be used to detect when to update any (UI-)state dependent on the provider during analysis.
+	 */
+	virtual uint64 GetUpdateCounter() const = 0;
 };
 
 /*
@@ -105,7 +98,7 @@ public:
 	* @param Time		The time in seconds of the begin event of this region.
 	*/
 	virtual void AppendRegionBegin(const TCHAR* Name, double Time) = 0;
-	
+
 	/*
 	* Append a new end event of a region from the trace session.
 	*
@@ -121,6 +114,9 @@ public:
 	virtual void OnAnalysisSessionEnded() = 0;
 };
 
+////////////////////////////////////////////////////////////////////////////////////////////////////
+
+TRACESERVICES_API FName GetRegionProviderName();
 TRACESERVICES_API const IRegionProvider& ReadRegionProvider(const IAnalysisSession& Session);
 TRACESERVICES_API IEditableRegionProvider& EditRegionProvider(IAnalysisSession& Session);
 
