[Insights] Fixed typo in IEditableLogProvider.

diff --git a/Engine/Source/Developer/TraceServices/Private/Model/Log.cpp b/Engine/Source/Developer/TraceServices/Private/Model/Log.cpp
--- a/Engine/Source/Developer/TraceServices/Private/Model/Log.cpp
+++ b/Engine/Source/Developer/TraceServices/Private/Model/Log.cpp
@@ -102,7 +102,7 @@ void FLogProvider::UpdateMessageFile(uint64 LogPoint, const TCHAR* InFile, int32
 	LogMessageSpec.Line = InLine;
 }
 
-void FLogProvider::UpdateMessageVebosity(uint64 LogPoint, ELogVerbosity::Type InVerbosity)
+void FLogProvider::UpdateMessageVerbosity(uint64 LogPoint, ELogVerbosity::Type InVerbosity)
 {
 	Session.WriteAccessCheck();
 	FLogMessageSpec& LogMessageSpec = GetMessageSpec(LogPoint);
diff --git a/Engine/Source/Developer/TraceServices/Private/Model/LogPrivate.h b/Engine/Source/Developer/TraceServices/Private/Model/LogPrivate.h
--- a/Engine/Source/Developer/TraceServices/Private/Model/LogPrivate.h
+++ b/Engine/Source/Developer/TraceServices/Private/Model/LogPrivate.h
@@ -30,9 +30,9 @@ struct FLogMessageInternal
 	const TCHAR* Message = nullptr;
 };
 
-class FLogProvider :
-	public ILogProvider,
-	public IEditableLogProvider
+class FLogProvider
+	: public ILogProvider
+	, public IEditableLogProvider
 {
 public:
 	static const FName ProviderName;
@@ -40,27 +40,36 @@ public:
 	explicit FLogProvider(IAnalysisSession& Session);
 	virtual ~FLogProvider() {}
 
-	// implement IEditableLogProvider
+	//////////////////////////////////////////////////
+	// Read operations
+
+	virtual uint64 GetMessageCount() const override;
+	virtual bool ReadMessage(uint64 Index, TFunctionRef<void(const FLogMessageInfo&)> Callback) const override;
+	virtual void EnumerateMessages(double IntervalStart, double IntervalEnd, TFunctionRef<void(const FLogMessageInfo&)> Callback) const override;
+	virtual void EnumerateMessagesByIndex(uint64 Start, uint64 End, TFunctionRef<void(const FLogMessageInfo&)> Callback) const override;
+
+	virtual uint64 GetCategoryCount() const override { return Categories.Num(); }
+	virtual void EnumerateCategories(TFunctionRef<void(const FLogCategoryInfo&)> Callback) const override;
+
+	virtual const IUntypedTable& GetMessagesTable() const override { return MessagesTable; }
+
+	//////////////////////////////////////////////////
+	// Edit operations
+
 	virtual uint64 RegisterCategory() override;
 	virtual FLogCategoryInfo& GetCategory(uint64 CategoryPointer) override;
+
+	FLogMessageSpec& GetMessageSpec(uint64 LogPoint);
 	virtual void UpdateMessageCategory(uint64 LogPoint, uint64 InCategoryPointer) override;
 	virtual void UpdateMessageFormatString(uint64 LogPoint, const TCHAR* InFormatString) override;
 	virtual void UpdateMessageFile(uint64 LogPoint, const TCHAR* InFile, int32 InLine) override;
-	virtual void UpdateMessageVebosity(uint64 LogPoint, ELogVerbosity::Type InVerbosity) override;
+	virtual void UpdateMessageVerbosity(uint64 LogPoint, ELogVerbosity::Type InVerbosity) override;
 	virtual void UpdateMessageSpec(uint64 LogPoint, uint64 InCategoryPointer, const TCHAR* InFormatString, const TCHAR* InFile, int32 InLine, ELogVerbosity::Type InVerbosity) override;
 	virtual void AppendMessage(uint64 LogPoint, double Time, const uint8* FormatArgs) override;
 	virtual void AppendMessage(uint64 LogPoint, double Time, const TCHAR* Text) override;
-
-	FLogMessageSpec& GetMessageSpec(uint64 LogPoint);
 	void AppendMessage(uint64 LogPoint, double Time, const FString& Message);
 
-	virtual uint64 GetMessageCount() const override;
-	virtual bool ReadMessage(uint64 Index, TFunctionRef<void(const FLogMessageInfo&)> Callback) const override;
-	virtual void EnumerateMessages(double IntervalStart, double IntervalEnd, TFunctionRef<void(const FLogMessageInfo&)> Callback) const override;
-	virtual void EnumerateMessagesByIndex(uint64 Start, uint64 End, TFunctionRef<void(const FLogMessageInfo&)> Callback) const override;
-	virtual uint64 GetCategoryCount() const override { return Categories.Num(); }
-	virtual void EnumerateCategories(TFunctionRef<void(const FLogCategoryInfo&)> Callback) const override;
-	virtual const IUntypedTable& GetMessagesTable() const override { return MessagesTable; }
+	//////////////////////////////////////////////////
 
 private:
 	void ConstructMessage(uint64 Id, TFunctionRef<void(const FLogMessageInfo&)> Callback) const;
diff --git a/Engine/Source/Developer/TraceServices/Public/TraceServices/Model/Log.h b/Engine/Source/Developer/TraceServices/Public/TraceServices/Model/Log.h
--- a/Engine/Source/Developer/TraceServices/Public/TraceServices/Model/Log.h
+++ b/Engine/Source/Developer/TraceServices/Public/TraceServices/Model/Log.h
@@ -63,13 +63,13 @@ public:
 	*
 	* @param CategoryPointer	The unique identity (memory address) of the category instrumentation.
 	*
-	* @return A reference to the category stucture.
+	* @return A reference to the category structure.
 	*/
 	virtual FLogCategoryInfo& GetCategory(uint64 CategoryPointer) = 0;
-	
+
 	/*
 	* Update a log message's Category information.
-	* 
+	*
 	* @param LogPoint			The log message to update.
 	* @param InCategoryPointer	The category.
 	*/
@@ -98,7 +98,7 @@ public:
 	* @param LogPoint		The log message to update.
 	* @param InVerbosity	The verbosity of the message.
 	*/
-	virtual void UpdateMessageVebosity(uint64 LogPoint, ELogVerbosity::Type InVerbosity) = 0;
+	virtual void UpdateMessageVerbosity(uint64 LogPoint, ELogVerbosity::Type InVerbosity) = 0;
 
 	/*
 	* Update a log message's information.
