[Insights]
  - TraceServices: Fixed FFrameProvider::GetFrameNumberForTimestamp() to be consistent with FFrameProvider::GetFrameFromTime().
  - TraceServices: Added comments for IFrameProvider API.

diff --git a/Engine/Source/Developer/TraceServices/Private/Model/Frames.cpp b/Engine/Source/Developer/TraceServices/Private/Model/Frames.cpp
--- a/Engine/Source/Developer/TraceServices/Private/Model/Frames.cpp
+++ b/Engine/Source/Developer/TraceServices/Private/Model/Frames.cpp
@@ -1,10 +1,11 @@
 // Copyright Epic Games, Inc. All Rights Reserved.
 
-#include "TraceServices/Model/Frames.h"
 #include "Model/FramesPrivate.h"
+
+#include "Algo/BinarySearch.h"
 #include "AnalysisServicePrivate.h"
+
 #include <limits>
-#include "Algo/BinarySearch.h"
 
 namespace TraceServices
 {
@@ -23,6 +24,7 @@ FFrameProvider::FFrameProvider(IAnalysisSession& InSession)
 uint64 FFrameProvider::GetFrameCount(ETraceFrameType FrameType) const
 {
 	Session.ReadAccessCheck();
+
 	return Frames[FrameType].Num();
 }
 
@@ -43,14 +45,16 @@ void FFrameProvider::EnumerateFrames(ETraceFrameType FrameType, uint64 Start, ui
 
 bool FFrameProvider::GetFrameFromTime(ETraceFrameType FrameType, double Time, FFrame& OutFrame) const
 {
+	Session.ReadAccessCheck();
+
 	int64 LowerBound = Algo::LowerBound(FrameStartTimes[FrameType], Time);
-	if (LowerBound > 0 && LowerBound - 1 < (int64)Frames[FrameType].Num())
+	if (LowerBound <= 0)
 	{
-		OutFrame = Frames[FrameType][LowerBound - 1];
-		return true;
+		return false;
 	}
-
-	return false;
+	check(LowerBound <= (int64)Frames[FrameType].Num());
+	OutFrame = Frames[FrameType][LowerBound - 1];
+	return true;
 }
 
 const FFrame* FFrameProvider::GetFrame(ETraceFrameType FrameType, uint64 Index) const
@@ -78,7 +82,6 @@ void FFrameProvider::BeginFrame(ETraceFrameType FrameType, double Time)
 	Frame.EndTime = std::numeric_limits<double>::infinity();
 	Frame.Index = Index;
 	Frame.FrameType = FrameType;
-
 	FrameStartTimes[FrameType].Add(Time);
 
 	Session.UpdateDurationSeconds(Time);
@@ -87,31 +90,28 @@ void FFrameProvider::BeginFrame(ETraceFrameType FrameType, double Time)
 void FFrameProvider::EndFrame(ETraceFrameType FrameType, double Time)
 {
 	Session.WriteAccessCheck();
-	// If the EndFrame event is the first event that comes through
-	if (Frames[FrameType].Num() == 0)
+
+	// Ignores the EndFrame event if it is the first event that comes through.
+	if (Frames[FrameType].Num() > 0)
 	{
-		return;
+		FFrame& Frame = Frames[FrameType][Frames[FrameType].Num() - 1];
+		Frame.EndTime = Time;
 	}
-	FFrame& Frame = Frames[FrameType][Frames[FrameType].Num() - 1];
-	Frame.EndTime = Time;
+
 	Session.UpdateDurationSeconds(Time);
 }
 
 uint32 FFrameProvider::GetFrameNumberForTimestamp(ETraceFrameType FrameType, double Timestamp) const
 {
-	if (FrameStartTimes[FrameType].Num() == 0 || Timestamp < FrameStartTimes[FrameType][0])
+	Session.ReadAccessCheck();
+
+	int64 LowerBound = Algo::LowerBound(FrameStartTimes[FrameType], Timestamp);
+	if (LowerBound <= 0)
 	{
 		return 0;
 	}
-	else if (Timestamp >= FrameStartTimes[FrameType].Last())
-	{
-		return static_cast<uint32>(FrameStartTimes[FrameType].Num());
-	}
-	else
-	{
-		uint32 Index = static_cast<uint32>(Algo::LowerBound(FrameStartTimes[FrameType], Timestamp));
-		return Index + 1;
-	}
+	check(LowerBound <= (int64)Frames[FrameType].Num());
+	return uint32(LowerBound - 1);
 }
 
 const IFrameProvider& ReadFrameProvider(const IAnalysisSession& Session)
diff --git a/Engine/Source/Developer/TraceServices/Private/Model/FramesPrivate.h b/Engine/Source/Developer/TraceServices/Private/Model/FramesPrivate.h
--- a/Engine/Source/Developer/TraceServices/Private/Model/FramesPrivate.h
+++ b/Engine/Source/Developer/TraceServices/Private/Model/FramesPrivate.h
@@ -20,16 +20,24 @@ public:
 	explicit FFrameProvider(IAnalysisSession& Session);
 	virtual ~FFrameProvider() {}
 
+	//////////////////////////////////////////////////
+	// Read operations
+
 	virtual uint64 GetFrameCount(ETraceFrameType FrameType) const override;
 	virtual void EnumerateFrames(ETraceFrameType FrameType, uint64 Start, uint64 End, TFunctionRef<void(const FFrame&)> Callback) const override;
 	virtual const TArray64<double>& GetFrameStartTimes(ETraceFrameType FrameType) const override { return FrameStartTimes[FrameType]; }
 	virtual bool GetFrameFromTime(ETraceFrameType FrameType, double Time, FFrame& OutFrame) const override;
 	virtual const FFrame* GetFrame(ETraceFrameType FrameType, uint64 Index) const override;
-	virtual uint32 GetFrameNumberForTimestamp(ETraceFrameType FrameType, double Timestamp) const override;
+	virtual uint32 GetFrameNumberForTimestamp(ETraceFrameType FrameType, double Time) const override;
+
+	//////////////////////////////////////////////////
+	// Edit operations
 
 	void BeginFrame(ETraceFrameType FrameType, double Time);
 	void EndFrame(ETraceFrameType FrameType, double Time);
 
+	//////////////////////////////////////////////////
+
 private:
 	IAnalysisSession& Session;
 	TArray<TPagedArray<FFrame>> Frames;
diff --git a/Engine/Source/Developer/TraceServices/Public/TraceServices/Model/Frames.h b/Engine/Source/Developer/TraceServices/Public/TraceServices/Model/Frames.h
--- a/Engine/Source/Developer/TraceServices/Public/TraceServices/Model/Frames.h
+++ b/Engine/Source/Developer/TraceServices/Public/TraceServices/Model/Frames.h
@@ -27,12 +27,56 @@ class IFrameProvider
 {
 public:
 	virtual ~IFrameProvider() = default;
+
+	/**
+	 * Gets the number of frames, for the specified frame type.
+	 * @param FrameType - type of frames
+	 * @returns number of frames
+	 */
 	virtual uint64 GetFrameCount(ETraceFrameType FrameType) const = 0;
+
+	/**
+	 * Enumerates frames in the [Start, End) index range, for the specified frame type.
+	 * @param FrameType - type of frames
+	 * @param Start     - inclusive start index
+	 * @param End       - exclusive end index
+	 * @param Callback  - a callback called for each frame
+	 */
 	virtual void EnumerateFrames(ETraceFrameType FrameType, uint64 Start, uint64 End, TFunctionRef<void(const FFrame&)> Callback) const = 0;
+
+	/**
+	 * Gets the array of frame start times, for the specified frame type.
+	 * @param FrameType - type of frames
+	 * @returns a const reference to array of time values
+	 */
 	virtual const TArray64<double>& GetFrameStartTimes(ETraceFrameType FrameType) const = 0;
+
+	/**
+	 * Gets the first frame with start time <= specified time, for the specified frame type.
+	 * Note: This function does not check the end time of frames.
+	 * @param FrameType - type of frames
+	 * @param Time      - time value, in seconds
+	 * @param OutFrame  - the output frame, if this function returns true
+	 * @returns true if there is a frame with start time <= time or false otherwise
+	 */
 	virtual bool GetFrameFromTime(ETraceFrameType FrameType, double Time, FFrame& OutFrame) const = 0;
+
+	/**
+	 * Gets the frame at the specified index, for the specified frame type.
+	 * @param FrameType - type of frames
+	 * @param Index     - frame index
+	 * @returns a const pointer to the frame or nullptr if index is invalid
+	 */
 	virtual const FFrame* GetFrame(ETraceFrameType FrameType, uint64 Index) const = 0;
-	virtual uint32 GetFrameNumberForTimestamp(ETraceFrameType FrameType, double Timestamp) const = 0;
+
+	/**
+	 * Gets the index of the first frame with start time <= specified time, for the specified frame type.
+	 * Note: This function does not check the end time of frames.
+	 * @param FrameType - type of frames
+	 * @param Time      - time value, in seconds
+	 * @returns index of the first frame with start time <= time; also returns 0 when time < start time of the first frame or if frame count is zero
+	 */
+	virtual uint32 GetFrameNumberForTimestamp(ETraceFrameType FrameType, double Time) const = 0;
 };
 
 TRACESERVICES_API const IFrameProvider& ReadFrameProvider(const IAnalysisSession& Session);
