[Insights] Memory Insights: Added debug filtering for mem alloc events (see INSIGHTS_FILTER_EVENTS_ENABLED) in FAllocationsProvider. Also moved the histogram computations into FRootHeap class.

diff --git a/Engine/Source/Developer/TraceServices/Private/Model/AllocationsProvider.cpp b/Engine/Source/Developer/TraceServices/Private/Model/AllocationsProvider.cpp
--- a/Engine/Source/Developer/TraceServices/Private/Model/AllocationsProvider.cpp
+++ b/Engine/Source/Developer/TraceServices/Private/Model/AllocationsProvider.cpp
@@ -45,6 +45,42 @@
 
 ////////////////////////////////////////////////////////////////////////////////////////////////////
 
+// Debug functionality to discard some of the Alloc/Free/MarkAllocAsHeap/UnmarkAllocAsHeap events.
+#define INSIGHTS_FILTER_EVENTS_ENABLED 0
+
+#if INSIGHTS_FILTER_EVENTS_ENABLED
+namespace TraceServices
+{
+bool FAllocationsProvider::ShouldIgnoreEvent(double Time, uint64 Address, HeapId RootHeapId, uint32 CallstackId)
+{
+	//if (RootHeapId != EMemoryTraceRootHeap::VideoMemory)
+	//{
+	//	return true;
+	//}
+	//if (Address != 0x0ull)
+	//{
+	//	return true;
+	//}
+	//if (CallstackId != 0)
+	//{
+	//	return true;
+	//}
+	//if (Time > 10.0) // stop analysis after specified time
+	//{
+	//	EditOnAnalysisCompleted(Time);
+	//	bInitialized = false;
+	//	return true;
+	//}
+	return false;
+}
+} // namespace TraceServices
+#define INSIGHTS_FILTER_EVENT(Time, Address, RootHeapId, CallstackId) { if (ShouldIgnoreEvent(Time, Address, RootHeapId, CallstackId)) return; }
+#else // INSIGHTS_FILTER_EVENTS_ENABLED
+#define INSIGHTS_FILTER_EVENT(Time, Address, RootHeapId, CallstackId)
+#endif // INSIGHTS_FILTER_EVENTS_ENABLED
+
+////////////////////////////////////////////////////////////////////////////////////////////////////
+
 #define INSIGHTS_DEBUG_WATCH 0
 
 #if INSIGHTS_DEBUG_WATCH
@@ -61,7 +97,7 @@ namespace TraceServices
 // Action to be executed when a match is found.
 #define INSIGHTS_DEBUG_WATCH_FOUND // just log the API call
 //#define INSIGHTS_DEBUG_WATCH_FOUND return // ignore events
-//#define INSIGHTS_DEBUG_WATCH_FOUND UE_DEBUG_BREAK(); break
+//#define INSIGHTS_DEBUG_WATCH_FOUND UE_DEBUG_BREAK()
 
 ////////////////////////////////////////////////////////////////////////////////////////////////////
 
@@ -88,7 +124,7 @@ namespace TraceServices
 		if (IsAddressWatched(Address))\
 		{\
 			INSIGHTS_DEBUG_WATCH_FOUND;\
-			INSIGHTS_API_LOGF(Format, Time, __VA_ARGS__); \
+			INSIGHTS_API_LOGF(Format, Time, __VA_ARGS__);\
 		}\
 	}
 	#define INSIGHTS_WATCH_INDIRECT_API_LOGF(ApiName, Address, Time) \
@@ -96,7 +132,7 @@ namespace TraceServices
 		if (IsAddressWatched(Address))\
 		{\
 			INSIGHTS_DEBUG_WATCH_FOUND;\
-			INSIGHTS_INDIRECT_API_LOGF(ApiName, Address, Time); \
+			INSIGHTS_INDIRECT_API_LOGF(ApiName, Address, Time);\
 		}\
 	}
 #else
@@ -734,22 +770,24 @@ FAllocationItem* FShortLivingAllocs::Remove(uint64 Address)
 	--AllocCount;
 
 	// Remove node.
-	if (Node == OldestAllocNode)
-	{
-		OldestAllocNode = OldestAllocNode->Next;
-	}
-	if (Node == LastAddedAllocNode)
-	{
-		LastAddedAllocNode = LastAddedAllocNode->Prev;
-	}
 	if (Node->Prev)
 	{
 		Node->Prev->Next = Node->Next;
 	}
+	else
+	{
+		INSIGHTS_SLOW_CHECK(Node == OldestAllocNode);
+		OldestAllocNode = Node->Next;
+	}
 	if (Node->Next)
 	{
 		Node->Next->Prev = Node->Prev;
 	}
+	else
+	{
+		INSIGHTS_SLOW_CHECK(Node == LastAddedAllocNode)
+		LastAddedAllocNode = Node->Prev;
+	}
 
 	// Add the removed node to the "unused" list.
 	Node->Next = FirstUnusedNode;
@@ -1297,6 +1335,7 @@ void FAllocationsProvider::EditAlloc(double Time, uint32 CallstackId, uint64 Add
 	}
 
 	INSIGHTS_WATCH_API_LOGF(Address, TEXT("Alloc 0x%llX : Size=%llu Alignment=%u RootHeap=%u CallstackId=%u"), Time, Address, Size, Alignment, RootHeapId, CallstackId);
+	INSIGHTS_FILTER_EVENT(Time, Address, RootHeapId, CallstackId);
 
 	if (Address == 0 && RootHeapId == EMemoryTraceRootHeap::SystemMemory)
 	{
@@ -1396,7 +1435,7 @@ void FAllocationsProvider::EditAlloc(double Time, uint32 CallstackId, uint64 Add
 		Allocation.RootHeap = static_cast<uint8>(RootHeapId);
 		Allocation.Flags = EMemoryTraceHeapAllocationFlags::None;
 
-		UpdateHistogramByAllocSize(Size);
+		RootHeap.UpdateHistogramByAllocSize(Size);
 
 		// Update stats for the current timeline sample.
 		TotalAllocatedMemory += Size;
@@ -1440,6 +1479,7 @@ void FAllocationsProvider::EditFree(double Time, uint32 CallstackId, uint64 Addr
 	}
 
 	INSIGHTS_WATCH_API_LOGF(Address, TEXT("Free 0x%llX : RootHeap=%u CallstackId=%u"), Time, Address, RootHeapId, CallstackId);
+	INSIGHTS_FILTER_EVENT(Time, Address, RootHeapId, CallstackId);
 
 	if (Address == 0 && RootHeapId == EMemoryTraceRootHeap::SystemMemory)
 	{
@@ -1536,8 +1576,8 @@ void FAllocationsProvider::EditFree(double Time, uint32 CallstackId, uint64 Addr
 
 		RootHeap.SbTree->AddAlloc(AllocationPtr); // SbTree takes ownership of AllocationPtr
 
-		uint32 EventDistance = AllocationPtr->EndEventIndex - AllocationPtr->StartEventIndex;
-		UpdateHistogramByEventDistance(EventDistance);
+		const uint32 EventDistance = AllocationPtr->EndEventIndex - AllocationPtr->StartEventIndex;
+		RootHeap.UpdateHistogramByEventDistance(EventDistance);
 
 		// Update stats for the current timeline sample. (Heap allocations are already excluded)
 		if (!AllocationPtr->IsHeap())
@@ -1736,6 +1776,11 @@ void FAllocationsProvider::EditMarkAllocationAsHeap(double Time, uint32 Callstac
 {
 	EditAccessCheck();
 
+	if (!bInitialized)
+	{
+		return;
+	}
+
 	INSIGHTS_WATCH_API_LOGF(Address, TEXT("MarkAllocAsHeap 0x%llX : Heap=%u Flags=0x%X CallstackId=%u"), Time, Address, Heap, uint32(Flags), CallstackId);
 
 	if (!IsValidHeap(Heap))
@@ -1751,6 +1796,8 @@ void FAllocationsProvider::EditMarkAllocationAsHeap(double Time, uint32 Callstac
 	const FHeapSpec& HeapSpec = GetHeapSpecUnchecked(Heap);
 	const FRootHeap& RootHeap = FindParentRootHeapUnchecked(Heap);
 
+	INSIGHTS_FILTER_EVENT(Time, Address, RootHeap.HeapSpec->Id, CallstackId);
+
 #if 0 // TODO
 	if (Heap == RootHeap.HeapSpec->Id)
 	{
@@ -1781,6 +1828,10 @@ void FAllocationsProvider::EditMarkAllocationAsHeap(double Time, uint32 Callstac
 		ensure(EnumHasAnyFlags(Flags, EMemoryTraceHeapAllocationFlags::Heap));
 		Alloc->Flags = Flags | EMemoryTraceHeapAllocationFlags::Heap;
 		Alloc->RootHeap = static_cast<uint8>(Heap);
+		if (CallstackId != 0)
+		{
+			Alloc->CallstackId = CallstackId;
+		}
 
 		// Re-add it to the Live allocs as a heap allocation.
 		RootHeap.LiveAllocs->AddHeap(Alloc); // the Live allocs takes ownership of Alloc
@@ -1809,6 +1860,11 @@ void FAllocationsProvider::EditUnmarkAllocationAsHeap(double Time, uint32 Callst
 {
 	EditAccessCheck();
 
+	if (!bInitialized)
+	{
+		return;
+	}
+
 	INSIGHTS_WATCH_API_LOGF(Address, TEXT("UnmarkAllocAsHeap 0x%llX : Heap=%u CallstackId=%u"), Time, Address, Heap, CallstackId);
 
 	if (!IsValidHeap(Heap))
@@ -1824,6 +1880,8 @@ void FAllocationsProvider::EditUnmarkAllocationAsHeap(double Time, uint32 Callst
 	const FHeapSpec& HeapSpec = GetHeapSpecUnchecked(Heap);
 	const FRootHeap& RootHeap = FindParentRootHeapUnchecked(Heap);
 
+	INSIGHTS_FILTER_EVENT(Time, Address, RootHeap.HeapSpec->Id, CallstackId);
+
 #if 0 // TODO
 	if (Heap == RootHeap.HeapSpec->Id)
 	{
@@ -1943,7 +2001,7 @@ void FAllocationsProvider::EditUnmarkAllocationAsHeap(double Time, uint32 Callst
 
 ////////////////////////////////////////////////////////////////////////////////////////////////////
 
-void FAllocationsProvider::UpdateHistogramByAllocSize(uint64 Size)
+void FAllocationsProvider::FRootHeap::UpdateHistogramByAllocSize(uint64 Size)
 {
 	if (Size > MaxAllocSize)
 	{
@@ -1965,7 +2023,7 @@ void FAllocationsProvider::UpdateHistogramByAllocSize(uint64 Size)
 
 ////////////////////////////////////////////////////////////////////////////////////////////////////
 
-void FAllocationsProvider::UpdateHistogramByEventDistance(uint32 EventDistance)
+void FAllocationsProvider::FRootHeap::UpdateHistogramByEventDistance(uint32 EventDistance)
 {
 	if (EventDistance > MaxEventDistance)
 	{
@@ -2070,6 +2128,8 @@ void FAllocationsProvider::EditPushTagFromPtr(uint32 ThreadId, uint8 Tracker, ui
 	const TagIdType Tag = Alloc ? Alloc->Tag : 0; // If ptr is not found use "Untagged"
 	TagTracker.PushTagFromPtr(ThreadId, Tracker, Tag);
 
+	INSIGHTS_FILTER_EVENT(0.0, Ptr, RootHeapId, 0);
+
 	if (!Alloc)
 	{
 		++MiscErrors;
@@ -2615,6 +2675,8 @@ const IAllocationsProvider* ReadAllocationsProvider(const IAnalysisSession& Sess
 #undef INSIGHTS_REMAP_INVALID_ALLOC_EVENTS
 #undef INSIGHTS_REMAP_INVALID_FREE_EVENTS
 #undef INSIGHTS_DEBUG_METADATA
+#undef INSIGHTS_FILTER_EVENTS_ENABLED
+#undef INSIGHTS_FILTER_EVENT
 #undef INSIGHTS_DEBUG_WATCH
 #undef INSIGHTS_DEBUG_WATCH_FOUND
 #undef INSIGHTS_LOGF
diff --git a/Engine/Source/Developer/TraceServices/Private/Model/AllocationsProvider.h b/Engine/Source/Developer/TraceServices/Private/Model/AllocationsProvider.h
--- a/Engine/Source/Developer/TraceServices/Private/Model/AllocationsProvider.h
+++ b/Engine/Source/Developer/TraceServices/Private/Model/AllocationsProvider.h
@@ -100,8 +100,7 @@ private:
 		FNode* Prev;
 	};
 
-	//static const int32 MaxAllocCount = 8 * 1024; // max number of short living allocations
-	static const int32 MaxAllocCount = 64; // max number of short living allocations
+	static constexpr int32 MaxAllocCount = 64; // max number of short living allocations
 
 public:
 	FShortLivingAllocs();
@@ -218,7 +217,7 @@ public:
 	// Enumerates all allocations (including heap allocations).
 	FORCEINLINE void Enumerate(TFunctionRef<void(const FAllocationItem& Alloc)> Callback) const;
 
-	// Enumerates allocations in a sprecified address range (including heap allocations).
+	// Enumerates allocations in a specified address range (including heap allocations).
 	FORCEINLINE void Enumerate(uint64 StartAddress, uint64 EndAddress, TFunctionRef<void(const FAllocationItem& Alloc)> Callback) const;
 
 	// Adds a new allocation with specified address.
@@ -273,10 +272,22 @@ private:
 
 	struct FRootHeap
 	{
+		void UpdateHistogramByAllocSize(uint64 Size);
+		void UpdateHistogramByEventDistance(uint32 EventDistance);
+
 		FHeapSpec* HeapSpec = nullptr;
+
 		FLiveAllocCollection* LiveAllocs = nullptr;
+
 		FSbTree* SbTree = nullptr;
+
 		uint32 EventIndex = 0;
+
+		uint64 MaxAllocSize = 0;
+		uint64 AllocSizeHistogramPow2[65] = { 0 };
+
+		uint32 MaxEventDistance = 0;
+		uint32 EventDistanceHistogramPow2[33] = { 0 };
 	};
 
 public:
@@ -354,8 +365,7 @@ public:
 	//////////////////////////////////////////////////
 
 private:
-	void UpdateHistogramByAllocSize(uint64 Size);
-	void UpdateHistogramByEventDistance(uint32 EventDistance);
+	bool ShouldIgnoreEvent(double Time, uint64 Address, HeapId RootHeapId, uint32 CallstackId);
 	void AdvanceTimelines(double Time);
 	void AddHeapSpec(HeapId Id, HeapId ParentId, const FStringView& Name, EMemoryTraceHeapFlags Flags);
 
@@ -418,12 +428,6 @@ private:
 	uint64 FreeWarnings = 0;
 	uint64 FreeErrors = 0;
 
-	uint64 MaxAllocSize = 0;
-	uint64 AllocSizeHistogramPow2[65] = { 0 };
-
-	uint32 MaxEventDistance = 0;
-	uint32 EventDistanceHistogramPow2[33] = { 0 };
-
 	uint64 TotalAllocatedMemory = 0;
 	uint32 TotalLiveAllocations = 0;
 
diff --git a/Engine/Source/Developer/TraceServices/Private/Model/SbTree.cpp b/Engine/Source/Developer/TraceServices/Private/Model/SbTree.cpp
--- a/Engine/Source/Developer/TraceServices/Private/Model/SbTree.cpp
+++ b/Engine/Source/Developer/TraceServices/Private/Model/SbTree.cpp
@@ -470,18 +470,18 @@ void FSbTree::AddAlloc(const FAllocationItem* Alloc)
 
 ////////////////////////////////////////////////////////////////////////////////////////////////////
 
-int32 FSbTree::GetColumnsAtTime(double Time, int32* StartColumnPtr, int32* EndColumnPtr) const
+int32 FSbTree::GetColumnsAtTime(double Time, int32* OutStartColumn, int32* OutEndColumn) const
 {
 	// Returns the first column with T >= Time.
 	// [0 .. N]
 	const int32 LowerBoundColumn = Algo::LowerBound(ColumnStartTimes, Time);
 
-	if (StartColumnPtr)
+	if (OutStartColumn)
 	{
-		*StartColumnPtr = LowerBoundColumn - 1; // [-1 .. N-1]
+		*OutStartColumn = LowerBoundColumn - 1; // [-1 .. N-1]
 	}
 
-	if (EndColumnPtr)
+	if (OutEndColumn)
 	{
 		int32 Column = LowerBoundColumn - 1; // [-1 .. N-1]
 		const int32 Last = ColumnStartTimes.Num() - 1; // N-1
@@ -489,7 +489,7 @@ int32 FSbTree::GetColumnsAtTime(double Time, int32* StartColumnPtr, int32* EndCo
 		{
 			++Column;
 		}
-		*EndColumnPtr = Column; // [-1 .. N-1]
+		*OutEndColumn = Column; // [-1 .. N-1]
 	}
 
 	return LowerBoundColumn;
diff --git a/Engine/Source/Developer/TraceServices/Private/Model/SbTree.h b/Engine/Source/Developer/TraceServices/Private/Model/SbTree.h
--- a/Engine/Source/Developer/TraceServices/Private/Model/SbTree.h
+++ b/Engine/Source/Developer/TraceServices/Private/Model/SbTree.h
@@ -78,7 +78,7 @@ public:
 	void Validate() const;
 
 private:
-	int32 GetColumnsAtTime(double Time, int32* StartColumnPtr, int32* EndColumnPtr) const;
+	int32 GetColumnsAtTime(double Time, int32* OutStartColumn, int32* OutEndColumn) const;
 
 private:
 	ILinearAllocator& Allocator;
@@ -86,7 +86,7 @@ private:
 	// Normal cells.
 	TArray<FSbTreeCell*> Cells;
 
-	// Offsetted cells. Note: Depth 0 doesn't use offseted cells. Those will always be empty.
+	// Offsetted cells. Note: Depth 0 doesn't use offsetted cells. Those will always be empty.
 	TArray<FSbTreeCell*> OffsettedCells;
 
 	// Array with time values of the first event in each column.
