Fixed issue with NetTraceReporter encoding of bunchsize.
PR #9925: Fix net trace encoding bug (Contributed by [REDACTED_PERSON]) + versioning fix to remains backwards compatible.

diff --git a/Engine/Source/Developer/TraceServices/Private/Analyzers/NetTraceAnalyzer.cpp b/Engine/Source/Developer/TraceServices/Private/Analyzers/NetTraceAnalyzer.cpp
--- a/Engine/Source/Developer/TraceServices/Private/Analyzers/NetTraceAnalyzer.cpp
+++ b/Engine/Source/Developer/TraceServices/Private/Analyzers/NetTraceAnalyzer.cpp
@@ -20,6 +20,7 @@ enum ENetTraceAnalyzerVersion
 	ENetTraceAnalyzerVersion_Initial = 1,
 	ENetTraceAnalyzerVersion_BunchChannelIndex = 2,
 	ENetTraceAnalyzerVersion_BunchChannelInfo = 3,
+	ENetTraceAnalyzerVersion_FixedBunchSizeEncoding = 4,		
 };
 
 
@@ -297,8 +298,18 @@ void FNetTraceAnalyzer::HandlePacketContentEvent(const FOnEventContext& Context,
 			case EContentEventType::BunchEvent:
 			{
 				const uint16 DecodedNameId = IntCastChecked<uint16>(FTraceAnalyzerUtils::Decode7bit(BufferPtr));
-				const uint32 DecodedEventStartPos = IntCastChecked<uint32>(FTraceAnalyzerUtils::Decode7bit(BufferPtr));
-				const uint32 DecodedEventEndPos = IntCastChecked<uint32>(FTraceAnalyzerUtils::Decode7bit(BufferPtr)) + DecodedEventStartPos;
+
+				uint32 DecodedBunchBits = 0U;
+				if (NetTraceVersion >= ENetTraceAnalyzerVersion_FixedBunchSizeEncoding)
+				{
+					DecodedBunchBits = IntCastChecked<uint32>(FTraceAnalyzerUtils::Decode7bit(BufferPtr));
+				}
+				else
+				{
+					const uint64 DecodedEventStartPos = FTraceAnalyzerUtils::Decode7bit(BufferPtr);
+					const uint64 DecodedEventEndPos = FTraceAnalyzerUtils::Decode7bit(BufferPtr);
+					DecodedBunchBits = IntCastChecked<uint32>((uint32)DecodedEventEndPos + (uint32)DecodedEventStartPos);
+				}
 
 				const uint32* NetProfilerNameIndex = DecodedNameId ? TracedNameIdToNetProfilerNameIdMap.Find(DecodedNameId) : nullptr;
 
@@ -306,7 +317,7 @@ void FNetTraceAnalyzer::HandlePacketContentEvent(const FOnEventContext& Context,
 
 				BunchInfo.BunchInfo.Value = 0;
 				BunchInfo.HeaderBits = 0U;
-				BunchInfo.BunchBits = DecodedEventEndPos;
+				BunchInfo.BunchBits = DecodedBunchBits;
 				BunchInfo.FirstBunchEventIndex = Events.Num();
 				BunchInfo.NameIndex = NetProfilerNameIndex ? IntCastChecked<uint16>(*NetProfilerNameIndex) : 0U;
 
diff --git a/Engine/Source/Runtime/Net/Core/Private/Net/Core/Trace/NetTraceInternal.cpp b/Engine/Source/Runtime/Net/Core/Private/Net/Core/Trace/NetTraceInternal.cpp
--- a/Engine/Source/Runtime/Net/Core/Private/Net/Core/Trace/NetTraceInternal.cpp
+++ b/Engine/Source/Runtime/Net/Core/Private/Net/Core/Trace/NetTraceInternal.cpp
@@ -33,6 +33,7 @@ struct FNetTraceInternal
 		ENetTraceVersion_Initial = 1,
 		ENetTraceVersion_BunchChannelIndex = 2,
 		ENetTraceVersion_BunchChannelInfo = 3,
+		ENetTraceVersion_FixedBunchSizeEncoding = 4,		
 	};
 
 	struct FThreadBuffer
@@ -50,7 +51,7 @@ struct FNetTraceInternal
 	static inline FThreadBuffer* CreateThreadBuffer();
 
 	static inline thread_local TUniquePtr<FThreadBuffer> ThreadBuffer;
-	static constexpr ENetTraceVersion NetTraceVersion = ENetTraceVersion::ENetTraceVersion_BunchChannelInfo;
+	static constexpr ENetTraceVersion NetTraceVersion = ENetTraceVersion::ENetTraceVersion_FixedBunchSizeEncoding;
 };
 
 FNetTraceInternal::FThreadBuffer* FNetTraceInternal::CreateThreadBuffer()
diff --git a/Engine/Source/Runtime/Net/Core/Private/Net/Core/Trace/Reporters/NetTraceReporter.cpp b/Engine/Source/Runtime/Net/Core/Private/Net/Core/Trace/Reporters/NetTraceReporter.cpp
--- a/Engine/Source/Runtime/Net/Core/Private/Net/Core/Trace/Reporters/NetTraceReporter.cpp
+++ b/Engine/Source/Runtime/Net/Core/Private/Net/Core/Trace/Reporters/NetTraceReporter.cpp
@@ -247,14 +247,12 @@ void FNetTraceReporter::ReportPacketContent(FNetTracePacketContentEvent* Events,
 			case ENetTracePacketContentEventType::BunchEvent:
 			{
 				// DebugName
-				uint32 EventId = CurrentEvent.DebugNameId;
+				const uint32 EventId = CurrentEvent.DebugNameId;
 				FTraceUtils::Encode7bit(EventId, BufferPtr);
 
-				// Start, do not delta compress as we have to deal with overshoot of previous bunch
-				FTraceUtils::Encode7bit(CurrentEvent.StartPos, BufferPtr);
-
-				// End
-				FTraceUtils::Encode7bit(CurrentEvent.EndPos - CurrentEvent.StartPos, BufferPtr);
+				// BunchSize
+				const uint32 BunchSize = CurrentEvent.EndPos;
+				FTraceUtils::Encode7bit(BunchSize, BufferPtr);
 
 				// Must reset LastOffest when we begin a new bunch
 				LastOffset = 0U;
diff --git a/Engine/Source/Runtime/Net/Core/Public/Net/Core/Trace/Private/NetTraceInternal.h b/Engine/Source/Runtime/Net/Core/Public/Net/Core/Trace/Private/NetTraceInternal.h
--- a/Engine/Source/Runtime/Net/Core/Public/Net/Core/Trace/Private/NetTraceInternal.h
+++ b/Engine/Source/Runtime/Net/Core/Public/Net/Core/Trace/Private/NetTraceInternal.h
@@ -564,7 +564,7 @@ void FNetTraceCollector::Reset()
 	FNetTrace::TChooseTraceEventScope<Verbosity, decltype(Stream)>::Type ScopeName(GetObjectIdForNetTrace(HandleOrNetGUID), Stream, Collector, Verbosity)
 
 #define UE_NET_TRACE_INTERNAL_NAMED_DYNAMIC_NAME_SCOPE(ScopeName, EventName, Stream, Collector, Verbosity) \
-	auto PREPROCESSOR_JOIN(NetTraceDynamicNameFunc_, __LINE__) = [&]() { return FNetTrace::TraceName(EventName); }; \
+	auto PREPROCESSOR_JOIN(NetTraceDynamicNameFunc_, __LINE__) = [&]() { uint16 NameId = FNetTrace::TraceName(EventName); return NameId; }; \
 	FNetTrace::TChooseTraceEventScope<Verbosity, decltype(Stream)>::Type ScopeName(PREPROCESSOR_JOIN(NetTraceDynamicNameFunc_, __LINE__), Stream, Collector, Verbosity)
 
 #define UE_NET_TRACE_INTERNAL_SET_SCOPE_NAME(ScopeName, EventName) UE_NET_TRACE_DO_IF(ScopeName.IsValid(), ScopeName.SetEventName(FNetTrace::TraceName(EventName)))
