[Insights] Compute median for timer view frame stats.

diff --git a/Engine/Source/Developer/TraceServices/Private/Common/TimelineStatistics.h b/Engine/Source/Developer/TraceServices/Private/Common/TimelineStatistics.h
--- a/Engine/Source/Developer/TraceServices/Private/Common/TimelineStatistics.h
+++ b/Engine/Source/Developer/TraceServices/Private/Common/TimelineStatistics.h
@@ -72,7 +72,7 @@ public:
 	static void CreateFrameStatsAggregation(const TArray<const TimelineType*>& Timelines, BucketMappingFunc BucketMapper, TArray<FFrameData> Frames, TMap<BucketKeyType, FAggregatedTimingStats>& Result)
 	{
 		TMap<BucketKeyType, FInternalAggregationEntry> FrameResult;
-		TMap<BucketKeyType, FInternalAggregationEntry> GlobalResult;
+		TMap<BucketKeyType, FInternalFrameAggregationEntry> GlobalResult;
 		double StartTime = Frames[0].StartTime;
 		double EndTime = Frames[Frames.Num() - 1].EndTime;
 
@@ -83,9 +83,25 @@ public:
 			GatherKeysFromTimeline(Timeline, BucketMapper, StartTime, EndTime, FrameResult);
 		}
 
-		GlobalResult = FrameResult;
+		int32 FramesNum = Frames.Num();
 
-		for (int32 Index = 0; Index < Frames.Num(); ++Index)
+		// For very large numbers of timers or frames, an out of memory is possible so don't compute the median.
+		constexpr int64 MaxSize = 5 * 100 * 1000 * 1000;
+		bool bComputeMedian = FrameResult.Num() * FramesNum < MaxSize;
+		for (auto& KV : FrameResult)
+		{
+			FInternalFrameAggregationEntry Entry;
+			if (bComputeMedian)
+			{
+				Entry.FrameInclusiveTimes.Reserve(Frames.Num());
+				Entry.FrameExclusiveTimes.Reserve(Frames.Num());
+			}
+			Entry.Inner = KV.Value;
+			GlobalResult.Add(KV.Key, Entry);
+		}
+
+
+		for (int32 Index = 0; Index < FramesNum; ++Index)
 		{
 			// Compute instance count and total/min/max inclusive/exclusive times for each timer.
 			for (const TimelineType* Timeline : Timelines)
@@ -94,11 +110,11 @@ public:
 			}
 
 			typename TMap<BucketKeyType, FInternalAggregationEntry>::TIterator FrameResultIterator(FrameResult);
-			typename TMap<BucketKeyType, FInternalAggregationEntry>::TIterator ResultIterator(GlobalResult);
+			typename TMap<BucketKeyType, FInternalFrameAggregationEntry>::TIterator ResultIterator(GlobalResult);
 			while (ResultIterator)
 			{
 				FAggregatedTimingStats& FrameStats = FrameResultIterator->Value.Stats;
-				FAggregatedTimingStats& ResultStats = ResultIterator->Value.Stats;
+				FAggregatedTimingStats& ResultStats = ResultIterator->Value.Inner.Stats;
 
 				// TODO: Add Instance Count per frame column
 				// ResultStats.InstanceCount = (ResultStats.InstanceCount * Index + FrameStats.InstanceCount) / (Index + 1);
@@ -108,10 +124,12 @@ public:
 				ResultStats.TotalInclusiveTime += FrameStats.TotalInclusiveTime;
 				ResultStats.MinInclusiveTime = FMath::Min(ResultStats.MinInclusiveTime, FrameStats.TotalInclusiveTime);
 				ResultStats.MaxInclusiveTime = FMath::Max(ResultStats.MaxInclusiveTime, FrameStats.TotalInclusiveTime);
+				ResultIterator->Value.FrameInclusiveTimes.Add(FrameStats.TotalInclusiveTime);
 
 				ResultStats.TotalExclusiveTime += FrameStats.TotalExclusiveTime;
 				ResultStats.MinExclusiveTime = FMath::Min(ResultStats.MinExclusiveTime, FrameStats.TotalExclusiveTime);
 				ResultStats.MaxExclusiveTime = FMath::Max(ResultStats.MaxExclusiveTime, FrameStats.TotalExclusiveTime);
+				ResultIterator->Value.FrameExclusiveTimes.Add(FrameStats.TotalExclusiveTime);
 
 				// Reset the per frame stats.
 				FrameStats = FAggregatedTimingStats();
@@ -121,12 +139,33 @@ public:
 			}
 		}
 
+		if (bComputeMedian)
+		{
+			// Compute the median inclusive and exclusive time.
+			for (auto& KV : GlobalResult)
+			{
+				if (KV.Value.Inner.Stats.InstanceCount == 0)
+				{
+					// Some timers might have 0 instances because the gathering phase also searched between frames.
+					continue;
+				}
+				PreComputeHistogram(KV.Value.Inner);
+
+				for (int Index = 0; Index < FramesNum; ++Index)
+				{
+					UpdateHistogramForTimerStats(KV.Value.Inner, KV.Value.FrameInclusiveTimes[Index], KV.Value.FrameExclusiveTimes[Index]);
+				}
+
+				ComputeFrameStatsMedian(KV.Value.Inner, FramesNum);
+			}
+		}
+
 		for (auto& KV : GlobalResult)
 		{
-			KV.Value.Stats.AverageInclusiveTime = KV.Value.Stats.TotalInclusiveTime / Frames.Num();
-			KV.Value.Stats.AverageExclusiveTime = KV.Value.Stats.TotalExclusiveTime / Frames.Num();
+			KV.Value.Inner.Stats.AverageInclusiveTime = KV.Value.Inner.Stats.TotalInclusiveTime / FramesNum;
+			KV.Value.Inner.Stats.AverageExclusiveTime = KV.Value.Inner.Stats.TotalExclusiveTime / FramesNum;
 
-			Result.Add(KV.Key, KV.Value.Stats);
+			Result.Add(KV.Key, KV.Value.Inner.Stats);
 		}
 	}
 
@@ -158,6 +197,14 @@ private:
 		FAggregatedTimingStats Stats;
 	};
 
+	struct FInternalFrameAggregationEntry
+	{
+		FInternalAggregationEntry Inner;
+
+		TArray<double> FrameInclusiveTimes;
+		TArray<double> FrameExclusiveTimes;
+	};
+
 	template<typename TimelineType, typename BucketMappingFunc, typename BucketKeyType, typename CallbackType>
 	static void ProcessTimeline(const TimelineType* Timeline, BucketMappingFunc BucketMapper, CallbackType Callback, double IntervalStart, double IntervalEnd, TMap<BucketKeyType, FInternalAggregationEntry>& InternalResult)
 	{
@@ -337,7 +384,7 @@ private:
 			int32 ExclCount = 0;
 			for (int32 HistogramIndex = 0; HistogramIndex < HistogramLen; HistogramIndex++)
 			{
-				ExclCount += AggregationEntry.InclHistogram[HistogramIndex];
+				ExclCount += AggregationEntry.ExclHistogram[HistogramIndex];
 				if (ExclCount > HalfCount)
 				{
 					Stats.MedianExclusiveTime = Stats.MinExclusiveTime + HistogramIndex * AggregationEntry.ExclDT;
@@ -354,6 +401,51 @@ private:
 		}
 	}
 
+	static void ComputeFrameStatsMedian(FInternalAggregationEntry& AggregationEntry, uint32 NumFrames)
+	{
+		FAggregatedTimingStats& Stats = AggregationEntry.Stats;
+
+		const int32 HalfCount = static_cast<int32>(NumFrames / 2);
+
+		// Compute median inclusive time.
+		int32 InclCount = 0;
+		for (int32 HistogramIndex = 0; HistogramIndex < HistogramLen; HistogramIndex++)
+		{
+			InclCount += AggregationEntry.InclHistogram[HistogramIndex];
+			if (InclCount > HalfCount)
+			{
+				Stats.MedianInclusiveTime = Stats.MinInclusiveTime + HistogramIndex * AggregationEntry.InclDT;
+				if (HistogramIndex > 0 &&
+					NumFrames % 2 == 0 &&
+					InclCount - AggregationEntry.InclHistogram[HistogramIndex] == HalfCount)
+				{
+					const double PrevMedian = Stats.MinInclusiveTime + (HistogramIndex - 1) * AggregationEntry.InclDT;
+					Stats.MedianInclusiveTime = (Stats.MedianInclusiveTime + PrevMedian) / 2;
+				}
+				break;
+			}
+		}
+
+		// Compute median exclusive time.
+		int32 ExclCount = 0;
+		for (int32 HistogramIndex = 0; HistogramIndex < HistogramLen; HistogramIndex++)
+		{
+			ExclCount += AggregationEntry.ExclHistogram[HistogramIndex];
+			if (ExclCount > HalfCount)
+			{
+				Stats.MedianExclusiveTime = Stats.MinExclusiveTime + HistogramIndex * AggregationEntry.ExclDT;
+				if (HistogramIndex > 0 &&
+					NumFrames % 2 == 0 &&
+					ExclCount - AggregationEntry.ExclHistogram[HistogramIndex] == HalfCount)
+				{
+					const double PrevMedian = Stats.MinExclusiveTime + (HistogramIndex - 1) * AggregationEntry.ExclDT;
+					Stats.MedianExclusiveTime = (Stats.MedianExclusiveTime + PrevMedian) / 2;
+				}
+				break;
+			}
+		}
+	}
+
 	template<typename TimelineType, typename BucketMappingFunc, typename BucketKeyType>
 	static void GatherKeysFromTimeline(const TimelineType* Timeline, BucketMappingFunc BucketMapper, double IntervalStart, double IntervalEnd, TMap<BucketKeyType, FInternalAggregationEntry>& InternalResult)
 	{
