[Insights] Memory Insights: Added FD3D12Adapter::TraceMemoryAllocation and D3D12PoolAllocator to resolved symbols filter.

diff --git a/Engine/Source/Developer/TraceServices/Private/Model/ModuleProvider.cpp b/Engine/Source/Developer/TraceServices/Private/Model/ModuleProvider.cpp
--- a/Engine/Source/Developer/TraceServices/Private/Model/ModuleProvider.cpp
+++ b/Engine/Source/Developer/TraceServices/Private/Model/ModuleProvider.cpp
@@ -426,9 +426,11 @@ FResolvedSymbolFilter::FResolvedSymbolFilter()
 	IgnoreSymbolsByFunctionName.Add(TEXT("FWindowsPlatformMemory::"));
 	IgnoreSymbolsByFunctionName.Add(TEXT("FCachedOSPageAllocator::"));
 	IgnoreSymbolsByFunctionName.Add(TEXT("FMallocBinned"));
+	IgnoreSymbolsByFunctionName.Add(TEXT("FD3D12Adapter::TraceMemoryAllocation"));
 
 	IgnoreSymbolsByFilePath.Add(FRegexPattern(FString(TEXT(".*/Containers/.*"))));
 	IgnoreSymbolsByFilePath.Add(FRegexPattern(FString(TEXT(".*/ConcurrentLinearAllocator.*"))));
+	IgnoreSymbolsByFilePath.Add(FRegexPattern(FString(TEXT(".*/D3D12PoolAllocator.*"))));
 }
 
 /////////////////////////////////////////////////////////////////////
