[Insights] TraceServices: Added Get*ProviderName() public API for all providers. Refactored the IModule derived classes to ensure consistent implementation. Also, fixed many comments and code style issues.

diff --git a/Engine/Source/Developer/TraceInsights/Private/Insights/CookProfiler/CookProfilerManager.cpp b/Engine/Source/Developer/TraceInsights/Private/Insights/CookProfiler/CookProfilerManager.cpp
--- a/Engine/Source/Developer/TraceInsights/Private/Insights/CookProfiler/CookProfilerManager.cpp
+++ b/Engine/Source/Developer/TraceInsights/Private/Insights/CookProfiler/CookProfilerManager.cpp
@@ -2,6 +2,7 @@
 
 #include "CookProfilerManager.h"
 
+#include "Common/ProviderLock.h" // TraceServices
 #include "Features/IModularFeatures.h"
 #include "Framework/Docking/TabManager.h"
 #include "Modules/ModuleManager.h"
diff --git a/Engine/Source/Developer/TraceInsights/Private/Insights/CookProfiler/Widgets/SPackageTableTreeView.cpp b/Engine/Source/Developer/TraceInsights/Private/Insights/CookProfiler/Widgets/SPackageTableTreeView.cpp
--- a/Engine/Source/Developer/TraceInsights/Private/Insights/CookProfiler/Widgets/SPackageTableTreeView.cpp
+++ b/Engine/Source/Developer/TraceInsights/Private/Insights/CookProfiler/Widgets/SPackageTableTreeView.cpp
@@ -2,6 +2,7 @@
 
 #include "SPackageTableTreeView.h"
 
+#include "Common/ProviderLock.h" // TraceServices
 #include "Framework/MultiBox/MultiBoxBuilder.h"
 #include "SlateOptMacros.h"
 #include "Widgets/Input/STextComboBox.h"
diff --git a/Engine/Source/Developer/TraceServices/Private/AnalysisService.cpp b/Engine/Source/Developer/TraceServices/Private/AnalysisService.cpp
--- a/Engine/Source/Developer/TraceServices/Private/AnalysisService.cpp
+++ b/Engine/Source/Developer/TraceServices/Private/AnalysisService.cpp
@@ -2,9 +2,9 @@
 
 #include "TraceServices/AnalysisService.h"
 #include "AnalysisServicePrivate.h"
-#include "ModuleServicePrivate.h"
-#include "Analyzers/LogTraceAnalysis.h"
+
 #include "Analyzers/BookmarksTraceAnalysis.h"
+#include "Analyzers/LogTraceAnalysis.h"
 #include "Analyzers/MiscTraceAnalysis.h"
 #include "Analyzers/StringsAnalyzer.h"
 #include "HAL/PlatformFile.h"
@@ -19,6 +19,7 @@
 #include "Model/RegionsPrivate.h"
 #include "Model/ScreenshotProviderPrivate.h"
 #include "Model/ThreadsPrivate.h"
+#include "ModuleServicePrivate.h"
 #include "Trace/Analysis.h"
 #include "Trace/Analyzer.h"
 #include "Trace/DataStream.h"
@@ -242,25 +243,25 @@ TSharedPtr<const IAnalysisSession> FAnalysisService::StartAnalysis(uint32 TraceI
 	Session->AddProvider(GetRegionProviderName(), RegionProvider, RegionProvider);
 
 	TSharedPtr<FLogProvider> LogProvider = MakeShared<FLogProvider>(*Session);
-	Session->AddProvider(FLogProvider::ProviderName, LogProvider, LogProvider);
+	Session->AddProvider(GetLogProviderName(), LogProvider, LogProvider);
 
 	TSharedPtr<FThreadProvider> ThreadProvider = MakeShared<FThreadProvider>(*Session);
-	Session->AddProvider(FThreadProvider::ProviderName, ThreadProvider, ThreadProvider);
+	Session->AddProvider(GetThreadProviderName(), ThreadProvider, ThreadProvider);
 
 	TSharedPtr<FFrameProvider> FrameProvider = MakeShared<FFrameProvider>(*Session);
-	Session->AddProvider(FFrameProvider::ProviderName, FrameProvider);
+	Session->AddProvider(GetFrameProviderName(), FrameProvider);
 
 	TSharedPtr<FCounterProvider> CounterProvider = MakeShared<FCounterProvider>(*Session, *FrameProvider);
-	Session->AddProvider(FCounterProvider::ProviderName, CounterProvider, CounterProvider);
+	Session->AddProvider(GetCounterProviderName(), CounterProvider, CounterProvider);
 
 	TSharedPtr<FChannelProvider> ChannelProvider = MakeShared<FChannelProvider>();
-	Session->AddProvider(FChannelProvider::ProviderName, ChannelProvider);
+	Session->AddProvider(GetChannelProviderName(), ChannelProvider);
 
 	TSharedPtr<FScreenshotProvider> ScreenshotProvider = MakeShared<FScreenshotProvider>(*Session);
-	Session->AddProvider(FScreenshotProvider::ProviderName, ScreenshotProvider);
+	Session->AddProvider(GetScreenshotProviderName(), ScreenshotProvider);
 
 	TSharedPtr<FDefinitionProvider> DefProvider = MakeShared<FDefinitionProvider>(&Session.Get());
-	Session->AddProvider(FDefinitionProvider::ProviderName, DefProvider, DefProvider);
+	Session->AddProvider(GetDefinitionProviderName(), DefProvider, DefProvider);
 
 	Session->AddAnalyzer(new FMiscTraceAnalyzer(*Session, *ThreadProvider, *LogProvider, *FrameProvider, *ChannelProvider, *ScreenshotProvider, *RegionProvider));
 	Session->AddAnalyzer(new FBookmarksAnalyzer(*Session, *BookmarkProvider, LogProvider.Get()));
diff --git a/Engine/Source/Developer/TraceServices/Private/Analyzers/MetadataAnalysis.cpp b/Engine/Source/Developer/TraceServices/Private/Analyzers/MetadataAnalysis.cpp
--- a/Engine/Source/Developer/TraceServices/Private/Analyzers/MetadataAnalysis.cpp
+++ b/Engine/Source/Developer/TraceServices/Private/Analyzers/MetadataAnalysis.cpp
@@ -1,4 +1,5 @@
 ﻿// Copyright Epic Games, Inc. All Rights Reserved.
+
 #include "MetadataAnalysis.h"
 
 #include "Common/ProviderLock.h"
@@ -14,10 +15,12 @@ namespace TraceServices
 {
 
 ////////////////////////////////////////////////////////////////////////////////////////////////////
+
 template<typename ArrayType>
 static void WriteData(ArrayType& OutData, const UE::Trace::IAnalyzer::FEventData& EventData, uint16 MetadataTypeId, const FMetadataProvider* MetadataProvider, IAnalysisSession& Session);
 
 ////////////////////////////////////////////////////////////////////////////////////////////////////
+
 FMetadataAnalysis::FMetadataAnalysis(IAnalysisSession& InSession, FMetadataProvider* InProvider)
 	: Session(InSession)
 	, MetadataProvider(InProvider)
@@ -25,11 +28,13 @@ FMetadataAnalysis::FMetadataAnalysis(IAnalysisSession& InSession, FMetadataProvi
 }
 
 ////////////////////////////////////////////////////////////////////////////////////////////////////
+
 FMetadataAnalysis::~FMetadataAnalysis()
 {
 }
 
 ////////////////////////////////////////////////////////////////////////////////////////////////////
+
 void FMetadataAnalysis::OnAnalysisBegin(const FOnAnalysisContext& Context)
 {
 	FInterfaceBuilder& Builder = Context.InterfaceBuilder;
@@ -44,6 +49,7 @@ void FMetadataAnalysis::OnAnalysisBegin(const FOnAnalysisContext& Context)
 }
 
 ////////////////////////////////////////////////////////////////////////////////////////////////////
+
 void FMetadataAnalysis::OnAnalysisEnd()
 {
 	FProviderEditScopeLock _(*MetadataProvider);
@@ -51,6 +57,7 @@ void FMetadataAnalysis::OnAnalysisEnd()
 }
 
 ////////////////////////////////////////////////////////////////////////////////////////////////////
+
 bool FMetadataAnalysis::OnEvent(uint16 RouteId, EStyle Style, const FOnEventContext& Context)
 {
 	FProviderEditScopeLock _(*MetadataProvider);
@@ -127,6 +134,7 @@ bool FMetadataAnalysis::OnEvent(uint16 RouteId, EStyle Style, const FOnEventCont
 }
 
 ////////////////////////////////////////////////////////////////////////////////////////////////////
+
 uint16 FMetadataAnalysis::GetOrRegisterType(const FEventTypeInfo& EventInfo)
 {
 	if (const uint16* CachedTypeId = EncounteredMetadataTypes.Find(EventInfo.GetId()))
@@ -179,8 +187,8 @@ uint16 FMetadataAnalysis::GetOrRegisterType(const FEventTypeInfo& EventInfo)
 		case FEventFieldInfo::EType::WideString:
 		case FEventFieldInfo::EType::AnsiString:
 			{
-				//Note that we don't have a ansi string store yet, so we store any ansi
-				//string as wide for metadata.
+				// Note that we don't have an ANSI string store yet, so we store any ANSI
+				// string as WIDE for metadata.
 				constexpr auto FieldType = FMetadataSchema::EFieldType::WideStringPtr;
 				Builder.AddField<TCHAR*>(Name, FieldType);
 			}
@@ -211,6 +219,7 @@ uint16 FMetadataAnalysis::GetOrRegisterType(const FEventTypeInfo& EventInfo)
 }
 
 ////////////////////////////////////////////////////////////////////////////////////////////////////
+
 template<typename ArrayType>
 static void WriteData(ArrayType& OutData, const UE::Trace::IAnalyzer::FEventData& EventData, uint16 MetadataTypeId, const FMetadataProvider* MetadataProvider, IAnalysisSession& Session)
 {
@@ -293,4 +302,5 @@ static void WriteData(ArrayType& OutData, const UE::Trace::IAnalyzer::FEventData
 }
 
 ////////////////////////////////////////////////////////////////////////////////////////////////////
+
 } // namespace TraceServices
diff --git a/Engine/Source/Developer/TraceServices/Private/Model/AllocationsProvider.cpp b/Engine/Source/Developer/TraceServices/Private/Model/AllocationsProvider.cpp
--- a/Engine/Source/Developer/TraceServices/Private/Model/AllocationsProvider.cpp
+++ b/Engine/Source/Developer/TraceServices/Private/Model/AllocationsProvider.cpp
@@ -2651,7 +2651,7 @@ uint32 FAllocationsProvider::GetNumLiveAllocs() const
 
 FName GetAllocationsProviderName()
 {
-	static FName Name(TEXT("AllocationsProvider"));
+	static const FName Name("AllocationsProvider");
 	return Name;
 }
 
diff --git a/Engine/Source/Developer/TraceServices/Private/Model/Bookmarks.cpp b/Engine/Source/Developer/TraceServices/Private/Model/Bookmarks.cpp
--- a/Engine/Source/Developer/TraceServices/Private/Model/Bookmarks.cpp
+++ b/Engine/Source/Developer/TraceServices/Private/Model/Bookmarks.cpp
@@ -1,11 +1,12 @@
 // Copyright Epic Games, Inc. All Rights Reserved.
 
 #include "TraceServices/Model/Bookmarks.h"
-#include "Common/FormatArgs.h"
 #include "Model/BookmarksPrivate.h"
+
 #include "AnalysisServicePrivate.h"
-#include "RegionsPrivate.h"
+#include "Common/FormatArgs.h"
 #include "Common/ProviderLock.h"
+#include "RegionsPrivate.h"
 #include "TraceServices/Model/Regions.h"
 
 namespace TraceServices
@@ -139,7 +140,7 @@ void FBookmarkProvider::EnumerateBookmarks(double IntervalStart, double Interval
 
 FName GetBookmarkProviderName()
 {
-	static FName Name(TEXT("BookmarkProvider"));
+	static const FName Name("BookmarkProvider");
 	return Name;
 }
 
diff --git a/Engine/Source/Developer/TraceServices/Private/Model/CallstacksProvider.cpp b/Engine/Source/Developer/TraceServices/Private/Model/CallstacksProvider.cpp
--- a/Engine/Source/Developer/TraceServices/Private/Model/CallstacksProvider.cpp
+++ b/Engine/Source/Developer/TraceServices/Private/Model/CallstacksProvider.cpp
@@ -1,17 +1,19 @@
 // Copyright Epic Games, Inc. All Rights Reserved.
 
 #include "CallstacksProvider.h"
+
+#include "Algo/Unique.h"
+#include "Containers/ArrayView.h"
 #include "Misc/ScopeRWLock.h"
 #include "ModuleProvider.h"
 #include "TraceServices/Model/AnalysisSession.h"
-#include "Algo/Unique.h"
-#include "Containers/ArrayView.h"
 #include "UObject/NameTypes.h"
 
 namespace TraceServices
 {
 
-/////////////////////////////////////////////////////////////////////
+////////////////////////////////////////////////////////////////////////////////////////////////////
+
 static const FResolvedSymbol GNeverResolveSymbol(ESymbolQueryResult::NotLoaded, nullptr, nullptr, nullptr, 0, EResolvedSymbolFilterStatus::NotFiltered);
 static const FResolvedSymbol GNotFoundSymbol(ESymbolQueryResult::NotFound, TEXT("Unknown"), nullptr, nullptr, 0, EResolvedSymbolFilterStatus::NotFiltered);
 static const FResolvedSymbol GNoSymbol(ESymbolQueryResult::NotFound, TEXT("No callstack recorded"), nullptr, nullptr, 0, EResolvedSymbolFilterStatus::NotFiltered);
@@ -19,17 +21,19 @@ static constexpr FStackFrame GNotFoundStackFrame = { 0, &GNotFoundSymbol };
 static constexpr FStackFrame GNoStackFrame = { 0, &GNoSymbol };
 static const FCallstack GNotFoundCallstack(&GNotFoundStackFrame, 1);
 
-/////////////////////////////////////////////////////////////////////
+////////////////////////////////////////////////////////////////////////////////////////////////////
+
 #ifdef TRACE_CALLSTACK_STATS
 static struct FCallstackProviderStats
 {
-	uint64		Callstacks;
-	uint64		Frames;
-	uint64		FrameCountHistogram[256];
+	uint64 Callstacks;
+	uint64 Frames;
+	uint64 FrameCountHistogram[256];
 } GCallstackStats;
 #endif
 
-/////////////////////////////////////////////////////////////////////
+////////////////////////////////////////////////////////////////////////////////////////////////////
+
 FCallstacksProvider::FCallstacksProvider(IAnalysisSession& InSession)
 	: Session(InSession)
 	, ModuleProvider(nullptr)
@@ -41,7 +45,8 @@ FCallstacksProvider::FCallstacksProvider(IAnalysisSession& InSession)
 	FirstCallstack.Init(&GNoStackFrame, 1);
 }
 
-/////////////////////////////////////////////////////////////////////
+////////////////////////////////////////////////////////////////////////////////////////////////////
+
 void FCallstacksProvider::AddCallstack(uint32 InCallstackId, const uint64* InFrames, uint8 InFrameCount)
 {
 	if (InCallstackId == 0)
@@ -114,7 +119,8 @@ void FCallstacksProvider::AddCallstack(uint32 InCallstackId, const uint64* InFra
 	}
 }
 
-/////////////////////////////////////////////////////////////////////
+////////////////////////////////////////////////////////////////////////////////////////////////////
+
 uint32 FCallstacksProvider::AddCallstackWithHash(uint64 InCallstackHash, const uint64* InFrames, uint8 InFrameCount)
 {
 	if (InCallstackHash == 0)
@@ -133,7 +139,8 @@ uint32 FCallstacksProvider::AddCallstackWithHash(uint64 InCallstackHash, const u
 	return CallstackId;
 }
 
-/////////////////////////////////////////////////////////////////////
+////////////////////////////////////////////////////////////////////////////////////////////////////
+
 uint32 FCallstacksProvider::GetCallstackIdForHash(uint64 InCallstackHash) const
 {
 	if (InCallstackHash == 0)
@@ -152,7 +159,8 @@ uint32 FCallstacksProvider::GetCallstackIdForHash(uint64 InCallstackHash) const
 	}
 }
 
-/////////////////////////////////////////////////////////////////////
+////////////////////////////////////////////////////////////////////////////////////////////////////
+
 const FCallstack* FCallstacksProvider::GetCallstack(uint32 CallstackId) const
 {
 	FRWScopeLock ReadLock(EntriesLock, SLT_ReadOnly);
@@ -166,7 +174,8 @@ const FCallstack* FCallstacksProvider::GetCallstack(uint32 CallstackId) const
 	}
 }
 
-/////////////////////////////////////////////////////////////////////
+////////////////////////////////////////////////////////////////////////////////////////////////////
+
 void FCallstacksProvider::GetCallstacks(const TArrayView<uint32>& CallstackIds, FCallstack const** OutCallstacks) const
 {
 	uint64 OutIdx(0);
@@ -187,16 +196,21 @@ void FCallstacksProvider::GetCallstacks(const TArrayView<uint32>& CallstackIds,
 	}
 }
 
-/////////////////////////////////////////////////////////////////////
+////////////////////////////////////////////////////////////////////////////////////////////////////
+
 FName GetCallstacksProviderName()
 {
-	static FName Name(TEXT("CallstacksProvider"));
+	static const FName Name("CallstacksProvider");
 	return Name;
 }
 
+////////////////////////////////////////////////////////////////////////////////////////////////////
+
 const ICallstacksProvider* ReadCallstacksProvider(const IAnalysisSession& Session)
 {
 	return Session.ReadProvider<ICallstacksProvider>(GetCallstacksProviderName());
 }
 
+////////////////////////////////////////////////////////////////////////////////////////////////////
+
 } // namespace TraceServices
diff --git a/Engine/Source/Developer/TraceServices/Private/Model/Channel.cpp b/Engine/Source/Developer/TraceServices/Private/Model/Channel.cpp
--- a/Engine/Source/Developer/TraceServices/Private/Model/Channel.cpp
+++ b/Engine/Source/Developer/TraceServices/Private/Model/Channel.cpp
@@ -1,20 +1,18 @@
 // Copyright Epic Games, Inc. All Rights Reserved.
+
 #include "Channel.h"
-#include "UObject/NameTypes.h"
+
 #include "Misc/DateTime.h"
+#include "UObject/NameTypes.h"
 
 namespace TraceServices
 {
 
-const FName FChannelProvider::ProviderName("ChannelProvider");
-
-///////////////////////////////////////////////////////////////////////////////
 FChannelProvider::FChannelProvider()
 {
 	TimeStamp = FDateTime::Now();
 }
 
-///////////////////////////////////////////////////////////////////////////////
 void FChannelProvider::AnnounceChannel(const TCHAR* InChannelName, uint32 Id, bool bReadOnly)
 {
 	FString ChannelName(InChannelName);
@@ -24,22 +22,18 @@ void FChannelProvider::AnnounceChannel(const TCHAR* InChannelName, uint32 Id, bo
 		ChannelName.GetCharArray()[0] = TChar<TCHAR>::ToUpper(ChannelName.GetCharArray()[0]);
 	}
 
-	Channels.Add(FChannelEntry{
-		Id,
-		ChannelName,
-		false,
-		bReadOnly,
-	});
+	Channels.Add(FChannelEntry{ Id, ChannelName, false, bReadOnly });
 
 	TimeStamp = FDateTime::Now();
 }
 
-///////////////////////////////////////////////////////////////////////////////
 void FChannelProvider::UpdateChannel(uint32 Id, bool bEnabled)
 {
-	const auto FoundEntry = Channels.FindByPredicate([Id](const FChannelEntry& Entry) {
-		return Entry.Id == Id;
-	});
+	const auto FoundEntry = Channels.FindByPredicate(
+		[Id](const FChannelEntry& Entry)
+		{
+			return Entry.Id == Id;
+		});
 
 	if (FoundEntry)
 	{
@@ -49,13 +43,11 @@ void FChannelProvider::UpdateChannel(uint32 Id, bool bEnabled)
 	TimeStamp = FDateTime::Now();
 }
 
-///////////////////////////////////////////////////////////////////////////////
 uint64 FChannelProvider::GetChannelCount() const
 {
 	return Channels.Num();
 }
 
-///////////////////////////////////////////////////////////////////////////////
 const TArray<FChannelEntry>& FChannelProvider::GetChannels() const
 {
 	return Channels;
@@ -66,4 +58,15 @@ FDateTime FChannelProvider::GetTimeStamp() const
 	return TimeStamp;
 }
 
+FName GetChannelProviderName()
+{
+	static const FName Name("ChannelProvider");
+	return Name;
+}
+
+const IChannelProvider* ReadChannelProvider(const IAnalysisSession& Session)
+{
+	return Session.ReadProvider<IChannelProvider>(GetChannelProviderName());
+}
+
 } // namespace TraceServices
diff --git a/Engine/Source/Developer/TraceServices/Private/Model/Channel.h b/Engine/Source/Developer/TraceServices/Private/Model/Channel.h
--- a/Engine/Source/Developer/TraceServices/Private/Model/Channel.h
+++ b/Engine/Source/Developer/TraceServices/Private/Model/Channel.h
@@ -12,15 +12,13 @@ namespace TraceServices
 class FChannelProvider : public IChannelProvider
 {
 public:
-	static const FName ProviderName;
-
 	FChannelProvider();
 	virtual ~FChannelProvider() {}
 
 	void AnnounceChannel(const TCHAR* ChannelName, uint32 Id, bool bReadOnly);
 	void UpdateChannel(uint32 Id, bool bEnabled);
 
-	virtual uint64	GetChannelCount() const override;
+	virtual uint64 GetChannelCount() const override;
 	virtual const TArray<FChannelEntry>& GetChannels() const override;
 
 	virtual FDateTime GetTimeStamp() const override;
diff --git a/Engine/Source/Developer/TraceServices/Private/Model/ContextSwitches.cpp b/Engine/Source/Developer/TraceServices/Private/Model/ContextSwitches.cpp
--- a/Engine/Source/Developer/TraceServices/Private/Model/ContextSwitches.cpp
+++ b/Engine/Source/Developer/TraceServices/Private/Model/ContextSwitches.cpp
@@ -2,14 +2,13 @@
 
 #include "TraceServices/Model/ContextSwitches.h"
 #include "Model/ContextSwitchesPrivate.h"
-#include "AnalysisServicePrivate.h"
+
 #include "Algo/Sort.h"
+#include "AnalysisServicePrivate.h"
 
 namespace TraceServices
 {
 
-const FName FContextSwitchesProvider::ProviderName = "ContextSwitchesProvider";
-
 FContextSwitchesProvider::FContextSwitchesProvider(IAnalysisSession& InSession)
 	: Session(InSession)
 	, NumCpuCores(0)
@@ -389,12 +388,18 @@ void FContextSwitchesProvider::AddThreadInfo(uint32 ThreadId, uint32 SystemThrea
 
 void FContextSwitchesProvider::AddThreadName(uint32 SystemTreadId, uint32 SystemProcessId, FStringView Name)
 {
-	//
+	//TODO
+}
+
+FName GetContextSwitchesProviderName()
+{
+	static const FName Name("ContextSwitchesProvider");
+	return Name;
 }
 
 const IContextSwitchesProvider* ReadContextSwitchesProvider(const IAnalysisSession& Session)
 {
-	return Session.ReadProvider<IContextSwitchesProvider>(FContextSwitchesProvider::ProviderName);
+	return Session.ReadProvider<IContextSwitchesProvider>(GetContextSwitchesProviderName());
 }
 
 } // namespace TraceServices
diff --git a/Engine/Source/Developer/TraceServices/Private/Model/ContextSwitchesPrivate.h b/Engine/Source/Developer/TraceServices/Private/Model/ContextSwitchesPrivate.h
--- a/Engine/Source/Developer/TraceServices/Private/Model/ContextSwitchesPrivate.h
+++ b/Engine/Source/Developer/TraceServices/Private/Model/ContextSwitchesPrivate.h
@@ -2,9 +2,10 @@
 
 #pragma once
 
+#include "Common/PagedArray.h" // TraceServices
+#include "Containers/Map.h"
 #include "TraceServices/Model/AnalysisSession.h"
 #include "TraceServices/Model/ContextSwitches.h"
-#include "Containers/Map.h"
 
 namespace TraceServices
 {
@@ -13,8 +14,6 @@ class FContextSwitchesProvider
 	: public IContextSwitchesProvider
 {
 public:
-	static const FName ProviderName;
-
 	explicit FContextSwitchesProvider(IAnalysisSession& Session);
 	virtual ~FContextSwitchesProvider();
 
diff --git a/Engine/Source/Developer/TraceServices/Private/Model/Counters.cpp b/Engine/Source/Developer/TraceServices/Private/Model/Counters.cpp
--- a/Engine/Source/Developer/TraceServices/Private/Model/Counters.cpp
+++ b/Engine/Source/Developer/TraceServices/Private/Model/Counters.cpp
@@ -2,19 +2,17 @@
 
 #include "TraceServices/Model/Counters.h"
 #include "Model/CountersPrivate.h"
+
 #include "AnalysisServicePrivate.h"
 
 namespace TraceServices
 {
 
-const FName FCounterProvider::ProviderName("CounterProvider");
-
 FCounter::FCounter(ILinearAllocator& Allocator, const TArray64<double>& InFrameStartTimes)
 	: FrameStartTimes(InFrameStartTimes)
 	, IntCounterData(Allocator)
 	, DoubleCounterData(Allocator)
 {
-
 }
 
 void FCounter::SetIsFloatingPoint(bool bInIsFloatingPoint)
@@ -30,6 +28,7 @@ static void EnumerateCounterValuesInternal(const TCounterData<CounterType>& Coun
 	{
 		return;
 	}
+
 	TArray64<double> NoFrameStartTimes;
 	auto CounterIterator = bResetEveryFrame ? CounterData.GetIterator(FrameStartTimes) : CounterData.GetIterator(NoFrameStartTimes);
 	bool bFirstValue = true;
@@ -98,6 +97,7 @@ static void EnumerateCounterOpsInternal(const TCounterData<CounterType>& Counter
 	{
 		return;
 	}
+
 	TArray64<double> NoFrameStartTimes;
 	auto CounterIterator = bResetEveryFrame ? CounterData.GetIterator(FrameStartTimes) : CounterData.GetIterator(NoFrameStartTimes);
 	bool bFirstValue = true;
@@ -264,14 +264,20 @@ void FCounterProvider::AddCounter(const ICounter* Counter)
 	Counters.Add(Counter);
 }
 
+FName GetCounterProviderName()
+{
+	static const FName Name("CounterProvider");
+	return Name;
+}
+
 const ICounterProvider& ReadCounterProvider(const IAnalysisSession& Session)
 {
-	return *Session.ReadProvider<ICounterProvider>(FCounterProvider::ProviderName);
+	return *Session.ReadProvider<ICounterProvider>(GetCounterProviderName());
 }
 
 IEditableCounterProvider& EditCounterProvider(IAnalysisSession& Session)
 {
-	return *Session.EditProvider<IEditableCounterProvider>(FCounterProvider::ProviderName);
+	return *Session.EditProvider<IEditableCounterProvider>(GetCounterProviderName());
 }
 
 } // namespace TraceServices
diff --git a/Engine/Source/Developer/TraceServices/Private/Model/CountersPrivate.h b/Engine/Source/Developer/TraceServices/Private/Model/CountersPrivate.h
--- a/Engine/Source/Developer/TraceServices/Private/Model/CountersPrivate.h
+++ b/Engine/Source/Developer/TraceServices/Private/Model/CountersPrivate.h
@@ -2,13 +2,13 @@
 
 #pragma once
 
+#include "Templates/SharedPointer.h"
+#include "TraceServices/Model/AnalysisSession.h"
 #include "TraceServices/Model/Counters.h"
 #include "TraceServices/Model/Frames.h"
-#include "TraceServices/Model/AnalysisSession.h"
-#include "Templates/SharedPointer.h"
 
 // We need to test data pattern and measure performance for more trace sessions
-// before completly switching to variable paged array.
+// before completely switching to variable paged array.
 #define USE_VARIABLE_PAGED_ARRAY 1
 
 #if USE_VARIABLE_PAGED_ARRAY
@@ -249,19 +249,25 @@ class FCounterProvider
 	, public IEditableCounterProvider
 {
 public:
-	static const FName ProviderName;
-
 	explicit FCounterProvider(IAnalysisSession& Session, IFrameProvider& FrameProvider);
 	virtual ~FCounterProvider();
+
+	//////////////////////////////////////////////////
+	// Read operations
+
 	virtual uint64 GetCounterCount() const override { return Counters.Num(); }
 	virtual void EnumerateCounters(TFunctionRef<void(uint32, const ICounter&)> Callback) const override;
 	virtual bool ReadCounter(uint32 CounterId, TFunctionRef<void(const ICounter&)> Callback) const override;
 
-	// implement IEditableCounterProvider
+	//////////////////////////////////////////////////
+	// Edit operations
+
 	virtual const ICounter* GetCounter(IEditableCounter* EditableCounter) override;
 	virtual IEditableCounter* CreateEditableCounter() override;
 	virtual void AddCounter(const ICounter* Counter) override;
 
+	//////////////////////////////////////////////////
+
 private:
 	IAnalysisSession& Session;
 	IFrameProvider& FrameProvider;
diff --git a/Engine/Source/Developer/TraceServices/Private/Model/DefinitionProvider.cpp b/Engine/Source/Developer/TraceServices/Private/Model/DefinitionProvider.cpp
--- a/Engine/Source/Developer/TraceServices/Private/Model/DefinitionProvider.cpp
+++ b/Engine/Source/Developer/TraceServices/Private/Model/DefinitionProvider.cpp
@@ -1,18 +1,14 @@
 ﻿// Copyright Epic Games, Inc. All Rights Reserved.
 
 #include "DefinitionProvider.h"
+
 #include "HAL/UnrealMemory.h"
-#include "UObject/NameTypes.h"
 
 namespace TraceServices
 {
 
 ////////////////////////////////////////////////////////////////////////////////////////////////////
 
-const FName FDefinitionProvider::ProviderName(TEXT("DefinitionProvider"));
-
-////////////////////////////////////////////////////////////////////////////////////////////////////
-
 FDefinitionProvider::FDefinitionProvider(IAnalysisSession* InSession)
 	: PageRemain(0)
 {
@@ -56,16 +52,24 @@ void* FDefinitionProvider::Allocate(uint32 Size, uint32 Alignment)
 
 ////////////////////////////////////////////////////////////////////////////////////////////////////
 
+FName GetDefinitionProviderName()
+{
+	static const FName Name("DefinitionProvider");
+	return Name;
+}
+
+////////////////////////////////////////////////////////////////////////////////////////////////////
+
 const IDefinitionProvider* ReadDefinitionProvider(const IAnalysisSession& Session)
 {
-	return Session.ReadProvider<IDefinitionProvider>(FDefinitionProvider::ProviderName);
+	return Session.ReadProvider<IDefinitionProvider>(GetDefinitionProviderName());
 }
 
 ////////////////////////////////////////////////////////////////////////////////////////////////////
 
 IDefinitionProvider* EditDefinitionProvider(IAnalysisSession& Session)
 {
-	return Session.EditProvider<IDefinitionProvider>(FDefinitionProvider::ProviderName);
+	return Session.EditProvider<IDefinitionProvider>(GetDefinitionProviderName());
 }
 
 ////////////////////////////////////////////////////////////////////////////////////////////////////
diff --git a/Engine/Source/Developer/TraceServices/Private/Model/DefinitionProvider.h b/Engine/Source/Developer/TraceServices/Private/Model/DefinitionProvider.h
--- a/Engine/Source/Developer/TraceServices/Private/Model/DefinitionProvider.h
+++ b/Engine/Source/Developer/TraceServices/Private/Model/DefinitionProvider.h
@@ -1,7 +1,7 @@
 ﻿// Copyright Epic Games, Inc. All Rights Reserved.
+
 #pragma once
 
-#include "Common/PagedArray.h"
 #include "Containers/Map.h"
 #include "HAL/CriticalSection.h"
 #include "Misc/ScopeRWLock.h"
@@ -15,8 +15,6 @@ namespace TraceServices
 class FDefinitionProvider : public IDefinitionProvider
 {
 public:
-	static const FName ProviderName;
-
 	FDefinitionProvider(IAnalysisSession* InSession);
 
 	virtual void BeginRead() const override { Lock.ReadLock(); }
diff --git a/Engine/Source/Developer/TraceServices/Private/Model/Diagnostics.cpp b/Engine/Source/Developer/TraceServices/Private/Model/Diagnostics.cpp
--- a/Engine/Source/Developer/TraceServices/Private/Model/Diagnostics.cpp
+++ b/Engine/Source/Developer/TraceServices/Private/Model/Diagnostics.cpp
@@ -32,7 +32,7 @@ const FSessionInfo& FDiagnosticsProvider::GetSessionInfo() const
 
 FName GetDiagnosticsProviderName()
 {
-	static FName Name(TEXT("DiagnosticsProvider"));
+	static const FName Name("DiagnosticsProvider");
 	return Name;
 }
 
diff --git a/Engine/Source/Developer/TraceServices/Private/Model/Frames.cpp b/Engine/Source/Developer/TraceServices/Private/Model/Frames.cpp
--- a/Engine/Source/Developer/TraceServices/Private/Model/Frames.cpp
+++ b/Engine/Source/Developer/TraceServices/Private/Model/Frames.cpp
@@ -11,8 +11,6 @@
 namespace TraceServices
 {
 
-const FName FFrameProvider::ProviderName("FrameProvider");
-
 FFrameProvider::FFrameProvider(IAnalysisSession& InSession)
 	: Session(InSession)
 {
@@ -98,7 +96,7 @@ const FFrame* FFrameProvider::GetFrame(ETraceFrameType FrameType, uint64 Index)
 void FFrameProvider::BeginFrame(ETraceFrameType FrameType, double Time)
 {
 	Session.WriteAccessCheck();
-	
+
 	uint64 Index = Frames[FrameType].Num();
 	FFrame& Frame = Frames[FrameType].PushBack();
 	Frame.StartTime = Time;
@@ -137,9 +135,15 @@ uint32 FFrameProvider::GetFrameNumberForTimestamp(ETraceFrameType FrameType, dou
 	return uint32(LowerBound - 1);
 }
 
+FName GetFrameProviderName()
+{
+	static const FName Name("FrameProvider");
+	return Name;
+}
+
 const IFrameProvider& ReadFrameProvider(const IAnalysisSession& Session)
 {
-	return *Session.ReadProvider<IFrameProvider>(FFrameProvider::ProviderName);
+	return *Session.ReadProvider<IFrameProvider>(GetFrameProviderName());
 }
 
 } // namespace TraceServices
diff --git a/Engine/Source/Developer/TraceServices/Private/Model/FramesPrivate.h b/Engine/Source/Developer/TraceServices/Private/Model/FramesPrivate.h
--- a/Engine/Source/Developer/TraceServices/Private/Model/FramesPrivate.h
+++ b/Engine/Source/Developer/TraceServices/Private/Model/FramesPrivate.h
@@ -2,9 +2,9 @@
 
 #pragma once
 
-#include "TraceServices/Model/Frames.h"
-#include "ProfilingDebugging/MiscTrace.h"
 #include "Common/PagedArray.h"
+#include "ProfilingDebugging/MiscTrace.h"
+#include "TraceServices/Model/Frames.h"
 
 namespace TraceServices
 {
@@ -15,8 +15,6 @@ class FFrameProvider
 	: public IFrameProvider
 {
 public:
-	static const FName ProviderName;
-
 	explicit FFrameProvider(IAnalysisSession& Session);
 	virtual ~FFrameProvider() {}
 
diff --git a/Engine/Source/Developer/TraceServices/Private/Model/Log.cpp b/Engine/Source/Developer/TraceServices/Private/Model/Log.cpp
--- a/Engine/Source/Developer/TraceServices/Private/Model/Log.cpp
+++ b/Engine/Source/Developer/TraceServices/Private/Model/Log.cpp
@@ -2,14 +2,13 @@
 
 #include "TraceServices/Model/Log.h"
 #include "Model/LogPrivate.h"
+
 #include "AnalysisServicePrivate.h"
 #include "Common/FormatArgs.h"
 
 namespace TraceServices
 {
 
-const FName FLogProvider::ProviderName("LogProvider");
-
 FLogProvider::FLogProvider(IAnalysisSession& InSession)
 	: Session(InSession)
 	, Categories(InSession.GetLinearAllocator(), 128)
@@ -243,14 +242,20 @@ void FLogProvider::EnumerateCategories(TFunctionRef<void(const FLogCategoryInfo&
 	}
 }
 
+FName GetLogProviderName()
+{
+	static const FName Name("LogProvider");
+	return Name;
+}
+
 const ILogProvider& ReadLogProvider(const IAnalysisSession& Session)
 {
-	return *Session.ReadProvider<ILogProvider>(FLogProvider::ProviderName);
+	return *Session.ReadProvider<ILogProvider>(GetLogProviderName());
 }
 
 IEditableLogProvider& EditLogProvider(IAnalysisSession& Session)
 {
-	return *Session.EditProvider<IEditableLogProvider>(FLogProvider::ProviderName);
+	return *Session.EditProvider<IEditableLogProvider>(GetLogProviderName());
 }
 
 void FormatString(TCHAR* OutputString, uint32 OutputStringCount, const TCHAR* FormatString, const uint8* FormatArgs)
diff --git a/Engine/Source/Developer/TraceServices/Private/Model/LogPrivate.h b/Engine/Source/Developer/TraceServices/Private/Model/LogPrivate.h
--- a/Engine/Source/Developer/TraceServices/Private/Model/LogPrivate.h
+++ b/Engine/Source/Developer/TraceServices/Private/Model/LogPrivate.h
@@ -35,8 +35,6 @@ class FLogProvider
 	, public IEditableLogProvider
 {
 public:
-	static const FName ProviderName;
-
 	explicit FLogProvider(IAnalysisSession& Session);
 	virtual ~FLogProvider() {}
 
diff --git a/Engine/Source/Developer/TraceServices/Private/Model/Memory.cpp b/Engine/Source/Developer/TraceServices/Private/Model/Memory.cpp
--- a/Engine/Source/Developer/TraceServices/Private/Model/Memory.cpp
+++ b/Engine/Source/Developer/TraceServices/Private/Model/Memory.cpp
@@ -1,4 +1,5 @@
 // Copyright Epic Games, Inc. All Rights Reserved.
+
 #include "TraceServices/Model/Memory.h"
 #include "Model/MemoryPrivate.h"
 
@@ -232,7 +233,7 @@ void FMemoryProvider::EnumerateTagSamples(FMemoryTrackerId TrackerId, FMemoryTag
 
 FName GetMemoryProviderName()
 {
-	static FName Name(TEXT("MemoryProvider"));
+	static const FName Name("MemoryProvider");
 	return Name;
 }
 
diff --git a/Engine/Source/Developer/TraceServices/Private/Model/MemoryPrivate.h b/Engine/Source/Developer/TraceServices/Private/Model/MemoryPrivate.h
--- a/Engine/Source/Developer/TraceServices/Private/Model/MemoryPrivate.h
+++ b/Engine/Source/Developer/TraceServices/Private/Model/MemoryPrivate.h
@@ -28,16 +28,13 @@ public:
 	virtual void EnumerateTrackers(TFunctionRef<void(const FMemoryTrackerInfo&)> Callback) const override;
 	virtual uint64 GetTagSampleCount(FMemoryTrackerId Tracker, FMemoryTagId Tag) const override;
 	virtual void EnumerateTagSamples(
-		FMemoryTrackerId Tracker, 
-		FMemoryTagId Tag, 
-		double StartTime, 
-		double EnddTime, 
+		FMemoryTrackerId Tracker,
+		FMemoryTagId Tag,
+		double StartTime,
+		double EnddTime,
 		bool bIncludeRangeNeighbours,
 		TFunctionRef<void(double Time, double Duration, const FMemoryTagSample&)> Callback) const override;
 
-public:
-	static FName ProviderName;
-
 private:
 	enum {
 		DefaultPageSize = 65536,
diff --git a/Engine/Source/Developer/TraceServices/Private/Model/MetadataProvider.cpp b/Engine/Source/Developer/TraceServices/Private/Model/MetadataProvider.cpp
--- a/Engine/Source/Developer/TraceServices/Private/Model/MetadataProvider.cpp
+++ b/Engine/Source/Developer/TraceServices/Private/Model/MetadataProvider.cpp
@@ -659,7 +659,7 @@ void FMetadataProvider::EnumerateMetadata(uint32 InThreadId, uint32 InMetadataId
 
 FName GetMetadataProviderName()
 {
-	static FName Name(TEXT("MetadataProvider"));
+	static const FName Name("MetadataProvider");
 	return Name;
 }
 
diff --git a/Engine/Source/Developer/TraceServices/Private/Model/ModuleProvider.cpp b/Engine/Source/Developer/TraceServices/Private/Model/ModuleProvider.cpp
--- a/Engine/Source/Developer/TraceServices/Private/Model/ModuleProvider.cpp
+++ b/Engine/Source/Developer/TraceServices/Private/Model/ModuleProvider.cpp
@@ -33,9 +33,11 @@
 #include "DbgHelpResolver.h"
 #endif
 
-namespace TraceServices {
+namespace TraceServices
+{
+
+////////////////////////////////////////////////////////////////////////////////////////////////////
 
-/////////////////////////////////////////////////////////////////////
 class FResolvedSymbolFilter : public IResolvedSymbolFilter
 {
 public:
@@ -49,7 +51,8 @@ private:
 	TArray<FRegexPattern> IgnoreSymbolsByFilePath;
 };
 
-/////////////////////////////////////////////////////////////////////
+////////////////////////////////////////////////////////////////////////////////////////////////////
+
 template<typename SymbolResolverType>
 class TModuleProvider : public IModuleAnalysisProvider
 {
@@ -57,23 +60,23 @@ public:
 	explicit TModuleProvider(IAnalysisSession& Session);
 	virtual ~TModuleProvider();
 
-	//Query interface
-	const FResolvedSymbol*		GetSymbol(uint64 Address) override;
-	uint32						GetNumModules() const override;
-	void						EnumerateModules(uint32 Start, TFunctionRef<void(const FModule& Module)> Callback) const override;
-	FGraphEventRef				LoadSymbolsForModuleUsingPath(uint64 Base, const TCHAR* Path) override;
-	void						EnumerateSymbolSearchPaths(TFunctionRef<void(FStringView Path)> Callback) const override;
-	void						GetStats(FStats* OutStats) const override;
-
-	//Analysis interface
-	void						OnModuleLoad(const FStringView& Module, uint64 Base, uint32 Size, const uint8* Checksum, uint32 ChecksumSize) override;
-	void 						OnModuleUnload(uint64 Base) override;
-	void						OnAnalysisComplete() override;
-	void						SaveSymbolsToCache(IAnalysisCache& Cache);
-	void						LoadSymbolsFromCache(IAnalysisCache& Cache);
+	// Query interface
+	const FResolvedSymbol* GetSymbol(uint64 Address) override;
+	uint32 GetNumModules() const override;
+	void EnumerateModules(uint32 Start, TFunctionRef<void(const FModule& Module)> Callback) const override;
+	FGraphEventRef LoadSymbolsForModuleUsingPath(uint64 Base, const TCHAR* Path) override;
+	void EnumerateSymbolSearchPaths(TFunctionRef<void(FStringView Path)> Callback) const override;
+	void GetStats(FStats* OutStats) const override;
+
+	// Analysis interface
+	void OnModuleLoad(const FStringView& Module, uint64 Base, uint32 Size, const uint8* Checksum, uint32 ChecksumSize) override;
+	void OnModuleUnload(uint64 Base) override;
+	void OnAnalysisComplete() override;
+	void SaveSymbolsToCache(IAnalysisCache& Cache);
+	void LoadSymbolsFromCache(IAnalysisCache& Cache);
 
 private:
-	uint32						GetNumCachedSymbolsFromModule(uint64 Base, uint32 Size);
+	uint32 GetNumCachedSymbolsFromModule(uint64 Base, uint32 Size);
 
 	struct FSavedSymbol
 	{
@@ -84,10 +87,10 @@ private:
 		uint32 Line;
 	};
 
-	mutable FRWLock				ModulesLock;
-	TPagedArray<FModule>		Modules;
+	mutable FRWLock ModulesLock;
+	TPagedArray<FModule> Modules;
 
-	FRWLock				SymbolsLock;
+	FRWLock SymbolsLock;
 	// Persistently stored symbol strings
 	FCachedStringStore Strings;
 	// Efficient representation of symbols
@@ -99,16 +102,17 @@ private:
 	// Number of discovered symbols
 	std::atomic<uint32> SymbolsDiscovered;
 
-	IAnalysisSession&			Session;
-	FString						Platform;
-	TUniquePtr<SymbolResolverType>	Resolver;
-	FGraphEventRef				LoadSymbolsTask;
-	bool						LoadSymbolsAbort = false;
+	IAnalysisSession& Session;
+	FString Platform;
+	TUniquePtr<SymbolResolverType> Resolver;
+	FGraphEventRef LoadSymbolsTask;
+	bool LoadSymbolsAbort = false;
 
-	FResolvedSymbolFilter		SymbolFilter;
+	FResolvedSymbolFilter SymbolFilter;
 };
 
-/////////////////////////////////////////////////////////////////////
+////////////////////////////////////////////////////////////////////////////////////////////////////
+
 template<typename SymbolResolverType>
 TModuleProvider<SymbolResolverType>::TModuleProvider(IAnalysisSession& Session)
 	: Modules(Session.GetLinearAllocator(), 128)
@@ -121,7 +125,8 @@ TModuleProvider<SymbolResolverType>::TModuleProvider(IAnalysisSession& Session)
 	LoadSymbolsFromCache(Session.GetCache());
 }
 
-/////////////////////////////////////////////////////////////////////
+////////////////////////////////////////////////////////////////////////////////////////////////////
+
 template <typename SymbolResolverType>
 TModuleProvider<SymbolResolverType>::~TModuleProvider()
 {
@@ -135,7 +140,8 @@ TModuleProvider<SymbolResolverType>::~TModuleProvider()
 	SaveSymbolsToCache(Session.GetCache());
 }
 
-/////////////////////////////////////////////////////////////////////
+////////////////////////////////////////////////////////////////////////////////////////////////////
+
 template<typename SymbolResolverType>
 const FResolvedSymbol* TModuleProvider<SymbolResolverType>::GetSymbol(uint64 Address)
 {
@@ -172,7 +178,8 @@ const FResolvedSymbol* TModuleProvider<SymbolResolverType>::GetSymbol(uint64 Add
 	return ResolvedSymbol;
 }
 
-/////////////////////////////////////////////////////////////////////
+////////////////////////////////////////////////////////////////////////////////////////////////////
+
 template <typename SymbolResolverType>
 uint32 TModuleProvider<SymbolResolverType>::GetNumModules() const
 {
@@ -180,7 +187,8 @@ uint32 TModuleProvider<SymbolResolverType>::GetNumModules() const
 	return static_cast<uint32>(Modules.Num());
 }
 
-/////////////////////////////////////////////////////////////////////
+////////////////////////////////////////////////////////////////////////////////////////////////////
+
 template <typename SymbolResolverType>
 void TModuleProvider<SymbolResolverType>::EnumerateModules(uint32 Start, TFunctionRef<void(const FModule& Module)> Callback) const
 {
@@ -191,7 +199,8 @@ void TModuleProvider<SymbolResolverType>::EnumerateModules(uint32 Start, TFuncti
 	}
 }
 
-/////////////////////////////////////////////////////////////////////
+////////////////////////////////////////////////////////////////////////////////////////////////////
+
 template <typename SymbolResolverType>
 FGraphEventRef TModuleProvider<SymbolResolverType>::LoadSymbolsForModuleUsingPath(uint64 Base, const TCHAR* Path)
 {
@@ -272,14 +281,16 @@ FGraphEventRef TModuleProvider<SymbolResolverType>::LoadSymbolsForModuleUsingPat
 	return FGraphEventRef();
 }
 
-/////////////////////////////////////////////////////////////////////
+////////////////////////////////////////////////////////////////////////////////////////////////////
+
 template <typename SymbolResolverType>
 void TModuleProvider<SymbolResolverType>::EnumerateSymbolSearchPaths(TFunctionRef<void(FStringView Path)> Callback) const
 {
 	Resolver->EnumerateSymbolSearchPaths(Callback);
 }
 
-/////////////////////////////////////////////////////////////////////
+////////////////////////////////////////////////////////////////////////////////////////////////////
+
 template<typename SymbolResolverType>
 void TModuleProvider<SymbolResolverType>::GetStats(FStats* OutStats) const
 {
@@ -289,7 +300,8 @@ void TModuleProvider<SymbolResolverType>::GetStats(FStats* OutStats) const
 	OutStats->SymbolsResolved += NumCachedSymbols;
 }
 
-/////////////////////////////////////////////////////////////////////
+////////////////////////////////////////////////////////////////////////////////////////////////////
+
 template<typename SymbolResolverType>
 void TModuleProvider<SymbolResolverType>::OnModuleLoad(const FStringView& Module, uint64 Base, uint32 Size, const uint8* ImageId, uint32 ImageIdSize)
 {
@@ -312,21 +324,24 @@ void TModuleProvider<SymbolResolverType>::OnModuleLoad(const FStringView& Module
 	Resolver->QueueModuleLoad(ImageId, ImageIdSize, &NewModule);
 }
 
-/////////////////////////////////////////////////////////////////////
+////////////////////////////////////////////////////////////////////////////////////////////////////
+
 template<typename SymbolResolverType>
 void TModuleProvider<SymbolResolverType>::OnModuleUnload(uint64 Base)
 {
 	//todo: Find entry, set bLoaded to false
 }
 
-/////////////////////////////////////////////////////////////////////
+////////////////////////////////////////////////////////////////////////////////////////////////////
+
 template<typename SymbolResolverType>
 void TModuleProvider<SymbolResolverType>::OnAnalysisComplete()
 {
 	Resolver->OnAnalysisComplete();
 }
 
-/////////////////////////////////////////////////////////////////////
+////////////////////////////////////////////////////////////////////////////////////////////////////
+
 template <typename SymbolResolverType>
 void TModuleProvider<SymbolResolverType>::SaveSymbolsToCache(IAnalysisCache& Cache)
 {
@@ -359,7 +374,8 @@ void TModuleProvider<SymbolResolverType>::SaveSymbolsToCache(IAnalysisCache& Cac
 	UE_LOG(LogTraceServices, Display, TEXT("Added %d symbols to the %d previously saved symbols."), NumSavedSymbols, NumPreviouslySavedSymbols);
 }
 
-/////////////////////////////////////////////////////////////////////
+////////////////////////////////////////////////////////////////////////////////////////////////////
+
 template <typename SymbolResolverType>
 void TModuleProvider<SymbolResolverType>::LoadSymbolsFromCache(IAnalysisCache& Cache)
 {
@@ -391,7 +407,8 @@ void TModuleProvider<SymbolResolverType>::LoadSymbolsFromCache(IAnalysisCache& C
 	UE_LOG(LogTraceServices, Display, TEXT("Loaded %d symbols from cache."), SymbolCacheLookup.Num());
 }
 
-/////////////////////////////////////////////////////////////////////
+////////////////////////////////////////////////////////////////////////////////////////////////////
+
 template <typename SymbolResolverType>
 uint32 TModuleProvider<SymbolResolverType>::GetNumCachedSymbolsFromModule(uint64 Base, uint32 Size)
 {
@@ -410,7 +427,8 @@ uint32 TModuleProvider<SymbolResolverType>::GetNumCachedSymbolsFromModule(uint64
 	return Count;
 }
 
-/////////////////////////////////////////////////////////////////////
+////////////////////////////////////////////////////////////////////////////////////////////////////
+
 FResolvedSymbolFilter::FResolvedSymbolFilter()
 {
 	IgnoreSymbolsByFunctionName.Add(TEXT("FMemory::"));
@@ -433,12 +451,14 @@ FResolvedSymbolFilter::FResolvedSymbolFilter()
 	IgnoreSymbolsByFilePath.Add(FRegexPattern(FString(TEXT(".*/D3D12PoolAllocator.*"))));
 }
 
-/////////////////////////////////////////////////////////////////////
+////////////////////////////////////////////////////////////////////////////////////////////////////
+
 FResolvedSymbolFilter::~FResolvedSymbolFilter()
 {
 }
 
-/////////////////////////////////////////////////////////////////////
+////////////////////////////////////////////////////////////////////////////////////////////////////
+
 void FResolvedSymbolFilter::Update(FResolvedSymbol& InSymbol) const
 {
 	bool bIsFiltered = false;
@@ -476,7 +496,8 @@ void FResolvedSymbolFilter::Update(FResolvedSymbol& InSymbol) const
 	InSymbol.FilterStatus.store(FilterStatus, std::memory_order_release);
 }
 
-/////////////////////////////////////////////////////////////////////
+////////////////////////////////////////////////////////////////////////////////////////////////////
+
 TSharedPtr<IModuleAnalysisProvider> CreateModuleProvider(IAnalysisSession& InSession, const FAnsiStringView& InSymbolFormat)
 {
 	TSharedPtr<IModuleAnalysisProvider> Provider;
@@ -496,20 +517,22 @@ TSharedPtr<IModuleAnalysisProvider> CreateModuleProvider(IAnalysisSession& InSes
 	return Provider;
 }
 
-/////////////////////////////////////////////////////////////////////
+////////////////////////////////////////////////////////////////////////////////////////////////////
+
 FName GetModuleProviderName()
 {
-	static FName Name(TEXT("ModuleProvider"));
+	static const FName Name("ModuleProvider");
 	return Name;
 }
 
-/////////////////////////////////////////////////////////////////////
+////////////////////////////////////////////////////////////////////////////////////////////////////
+
 const IModuleProvider* ReadModuleProvider(const IAnalysisSession& Session)
 {
 	return Session.ReadProvider<IModuleProvider>(GetModuleProviderName());
 }
 
-/////////////////////////////////////////////////////////////////////
+////////////////////////////////////////////////////////////////////////////////////////////////////
 
 } // namespace TraceServices
 
diff --git a/Engine/Source/Developer/TraceServices/Private/Model/NetProfilerProvider.cpp b/Engine/Source/Developer/TraceServices/Private/Model/NetProfilerProvider.cpp
--- a/Engine/Source/Developer/TraceServices/Private/Model/NetProfilerProvider.cpp
+++ b/Engine/Source/Developer/TraceServices/Private/Model/NetProfilerProvider.cpp
@@ -829,7 +829,7 @@ ITable<FNetProfilerAggregatedStatsCounterStats>* FNetProfilerProvider::CreateSta
 
 FName GetNetProfilerProviderName()
 {
-	static FName Name(TEXT("NetProfilerProvider"));
+	static const FName Name("NetProfilerProvider");
 	return Name;
 }
 
diff --git a/Engine/Source/Developer/TraceServices/Private/Model/Regions.cpp b/Engine/Source/Developer/TraceServices/Private/Model/Regions.cpp
--- a/Engine/Source/Developer/TraceServices/Private/Model/Regions.cpp
+++ b/Engine/Source/Developer/TraceServices/Private/Model/Regions.cpp
@@ -1,13 +1,13 @@
 // Copyright Epic Games, Inc. All Rights Reserved.
 
 #include "TraceServices/Model/Regions.h"
+#include "Model/RegionsPrivate.h"
 
 #include "AnalysisServicePrivate.h"
 #include "Common/FormatArgs.h"
 #include "Common/Utils.h"
 #include "Internationalization/Internationalization.h"
 #include "Logging/MessageLog.h"
-#include "Model/RegionsPrivate.h"
 
 #define LOCTEXT_NAMESPACE "RegionProvider"
 
@@ -235,7 +235,7 @@ bool FRegionLane::EnumerateRegions(double IntervalStart, double IntervalEnd, TFu
 
 FName GetRegionProviderName()
 {
-	static FName Name(TEXT("RegionProvider"));
+	static const FName Name("RegionProvider");
 	return Name;
 }
 
diff --git a/Engine/Source/Developer/TraceServices/Private/Model/ScreenshotProvider.cpp b/Engine/Source/Developer/TraceServices/Private/Model/ScreenshotProvider.cpp
--- a/Engine/Source/Developer/TraceServices/Private/Model/ScreenshotProvider.cpp
+++ b/Engine/Source/Developer/TraceServices/Private/Model/ScreenshotProvider.cpp
@@ -1,19 +1,18 @@
 // Copyright Epic Games, Inc. All Rights Reserved.
 
+#include "TraceServices/Model/Screenshot.h"
 #include "Model/ScreenshotProviderPrivate.h"
+
 #include "AnalysisServicePrivate.h"
 #include "Common/FormatArgs.h"
 
 namespace TraceServices
 {
 
-const FName FScreenshotProvider::ProviderName("ScreenshotProvider");
-
 FScreenshotProvider::FScreenshotProvider(IAnalysisSession& InSession)
 	: Session(InSession)
 {
-
-} 
+}
 
 TSharedPtr<FScreenshot> FScreenshotProvider::AddScreenshot(uint32 Id)
 {
@@ -24,7 +23,7 @@ TSharedPtr<FScreenshot> FScreenshotProvider::AddScreenshot(uint32 Id)
 void FScreenshotProvider::AddScreenshotChunk(uint32 Id, uint32 ChunkNum, uint16 Size, const TArrayView<const uint8>& ChunkData)
 {
 	TSharedPtr<FScreenshot> *ScreenshotPtr = Screenshots.Find(Id);
-	
+
 	if (ScreenshotPtr == nullptr)
 	{
 		return;
@@ -50,9 +49,15 @@ const TSharedPtr<const FScreenshot> FScreenshotProvider::GetScreenshot(uint32 Id
 	return *ScreenshotPtr;
 }
 
+FName GetScreenshotProviderName()
+{
+	static const FName Name("ScreenshotProvider");
+	return Name;
+}
+
 const IScreenshotProvider& ReadScreenshotProvider(const IAnalysisSession& Session)
 {
-	return *Session.ReadProvider<IScreenshotProvider>(FScreenshotProvider::ProviderName);
+	return *Session.ReadProvider<IScreenshotProvider>(GetScreenshotProviderName());
 }
 
 } // namespace TraceServices
diff --git a/Engine/Source/Developer/TraceServices/Private/Model/ScreenshotProviderPrivate.h b/Engine/Source/Developer/TraceServices/Private/Model/ScreenshotProviderPrivate.h
--- a/Engine/Source/Developer/TraceServices/Private/Model/ScreenshotProviderPrivate.h
+++ b/Engine/Source/Developer/TraceServices/Private/Model/ScreenshotProviderPrivate.h
@@ -12,8 +12,6 @@ class FScreenshotProvider
 	: public IScreenshotProvider
 {
 public:
-	static const FName ProviderName;
-
 	explicit FScreenshotProvider(IAnalysisSession& Session);
 	virtual ~FScreenshotProvider() {}
 
diff --git a/Engine/Source/Developer/TraceServices/Private/Model/StackSamples.cpp b/Engine/Source/Developer/TraceServices/Private/Model/StackSamples.cpp
--- a/Engine/Source/Developer/TraceServices/Private/Model/StackSamples.cpp
+++ b/Engine/Source/Developer/TraceServices/Private/Model/StackSamples.cpp
@@ -8,8 +8,6 @@
 namespace TraceServices
 {
 
-const FName FStackSamplesProvider::ProviderName = "StackSamplesProvider";
-
 FStackSamplesProvider::FStackSamplesProvider(IAnalysisSession& InSession)
 	: Session(InSession)
 	, AddressValues(Session.GetLinearAllocator(), 4096)
@@ -41,7 +39,7 @@ void FStackSamplesProvider::Add(uint32 ThreadId, double Time, uint32 Count, cons
 	{
 		StackSamples = &Threads.Add(ThreadId, new TPagedArray<FStackSample>(Session.GetLinearAllocator(), 4096));
 	}
-	
+
 	FStackSample& StackSample = (*StackSamples)->PushBack();
 	StackSample.Time = Time;
 	StackSample.Count = Count;
@@ -72,9 +70,15 @@ void FStackSamplesProvider::Add(uint32 ThreadId, double Time, uint32 Count, cons
 	StackSample.Addresses = &AddressValues.Last() - Count;
 }
 
+FName GetStackSamplesProviderName()
+{
+	static const FName Name("StackSamplesProvider");
+	return Name;
+}
+
 const IStackSamplesProvider& ReadStackSamplesProvider(const IAnalysisSession& Session)
 {
-	return *Session.ReadProvider<IStackSamplesProvider>(FStackSamplesProvider::ProviderName);
+	return *Session.ReadProvider<IStackSamplesProvider>(GetStackSamplesProviderName());
 }
 
 } // namespace TraceServices
diff --git a/Engine/Source/Developer/TraceServices/Private/Model/StackSamplesPrivate.h b/Engine/Source/Developer/TraceServices/Private/Model/StackSamplesPrivate.h
--- a/Engine/Source/Developer/TraceServices/Private/Model/StackSamplesPrivate.h
+++ b/Engine/Source/Developer/TraceServices/Private/Model/StackSamplesPrivate.h
@@ -2,9 +2,10 @@
 
 #pragma once
 
+#include "Common/PagedArray.h"
+#include "Containers/Map.h"
 #include "TraceServices/Model/AnalysisSession.h"
 #include "TraceServices/Model/StackSamples.h"
-#include "Containers/Map.h"
 
 namespace TraceServices
 {
@@ -13,8 +14,6 @@ class FStackSamplesProvider
 	: public IStackSamplesProvider
 {
 public:
-	static const FName ProviderName;
-
 	explicit FStackSamplesProvider(IAnalysisSession& Session);
 	virtual ~FStackSamplesProvider();
 
diff --git a/Engine/Source/Developer/TraceServices/Private/Model/Threads.cpp b/Engine/Source/Developer/TraceServices/Private/Model/Threads.cpp
--- a/Engine/Source/Developer/TraceServices/Private/Model/Threads.cpp
+++ b/Engine/Source/Developer/TraceServices/Private/Model/Threads.cpp
@@ -2,19 +2,17 @@
 
 #include "TraceServices/Model/Threads.h"
 #include "Model/ThreadsPrivate.h"
-#include "Misc/ScopeLock.h"
+
 #include "AnalysisServicePrivate.h"
 #include "Common/StringStore.h"
+#include "Misc/ScopeLock.h"
 
 namespace TraceServices
 {
 
-const FName FThreadProvider::ProviderName = "ThreadProvider";
-
 FThreadProvider::FThreadProvider(IAnalysisSession& InSession)
 	: Session(InSession)
 {
-
 }
 
 FThreadProvider::~FThreadProvider()
@@ -214,14 +212,20 @@ bool FThreadProvider::FThreadInfoInternal::operator<(const FThreadInfoInternal&
 	return GroupSortOrder < Other.GroupSortOrder;
 }
 
+FName GetThreadProviderName()
+{
+	static const FName Name("ThreadProvider");
+	return Name;
+}
+
 const IThreadProvider& ReadThreadProvider(const IAnalysisSession& Session)
 {
-	return *Session.ReadProvider<IThreadProvider>(FThreadProvider::ProviderName);
+	return *Session.ReadProvider<IThreadProvider>(GetThreadProviderName());
 }
 
 IEditableThreadProvider& EditThreadProvider(IAnalysisSession& Session)
 {
-	return *Session.EditProvider<IEditableThreadProvider>(FThreadProvider::ProviderName);
+	return *Session.EditProvider<IEditableThreadProvider>(GetThreadProviderName());
 }
 
 } // namespace TraceServices
diff --git a/Engine/Source/Developer/TraceServices/Private/Model/ThreadsPrivate.h b/Engine/Source/Developer/TraceServices/Private/Model/ThreadsPrivate.h
--- a/Engine/Source/Developer/TraceServices/Private/Model/ThreadsPrivate.h
+++ b/Engine/Source/Developer/TraceServices/Private/Model/ThreadsPrivate.h
@@ -2,11 +2,11 @@
 
 #pragma once
 
+#include "Containers/Map.h"
+#include "HAL/PlatformAffinity.h"
+#include "ProfilingDebugging/MiscTrace.h"
 #include "TraceServices/Model/AnalysisSession.h"
 #include "TraceServices/Model/Threads.h"
-#include "ProfilingDebugging/MiscTrace.h"
-#include "HAL/PlatformAffinity.h"
-#include "Containers/Map.h"
 
 namespace TraceServices
 {
@@ -16,8 +16,6 @@ class FThreadProvider
 	, public IEditableThreadProvider
 {
 public:
-	static const FName ProviderName;
-
 	explicit FThreadProvider(IAnalysisSession& Session);
 	virtual ~FThreadProvider();
 
diff --git a/Engine/Source/Developer/TraceServices/Private/ModuleServicePrivate.h b/Engine/Source/Developer/TraceServices/Private/ModuleServicePrivate.h
--- a/Engine/Source/Developer/TraceServices/Private/ModuleServicePrivate.h
+++ b/Engine/Source/Developer/TraceServices/Private/ModuleServicePrivate.h
@@ -6,6 +6,7 @@
 #include "Containers/Map.h"
 #include "Containers/Set.h"
 #include "Misc/ScopeLock.h"
+#include "UObject/NameTypes.h"
 
 namespace TraceServices
 {
@@ -17,17 +18,20 @@ class FModuleService
 {
 public:
 	FModuleService();
+
 	virtual void GetAvailableModules(TArray<FModuleInfo>& OutModules) override;
 	virtual void GetAvailableModulesEx(TArray<FModuleInfoEx>& OutModules) override;
 	virtual void GetEnabledModules(TArray<FModuleInfo>& OutModules) override;
 	virtual void SetModuleEnabled(const FName& ModuleName, bool bEnabled) override;
 	virtual void GenerateReports(const IAnalysisSession& Session, const TCHAR* CmdLine, const TCHAR* OutputDirectory) override;
+
 	void OnAnalysisBegin(IAnalysisSession& Session);
+
 	TArray<const TCHAR*> GetModuleLoggers(const FName& ModuleName);
-	TSet<FName> GetEnabledModulesFromCommandLine(const TCHAR* CommandLine);
 
 private:
 	void Initialize();
+	TSet<FName> GetEnabledModulesFromCommandLine(const TCHAR* CommandLine);
 
 	mutable FCriticalSection CriticalSection;
 	bool bIsInitialized = false;
diff --git a/Engine/Source/Developer/TraceServices/Private/Modules/CookProfilerModule.cpp b/Engine/Source/Developer/TraceServices/Private/Modules/CookProfilerModule.cpp
--- a/Engine/Source/Developer/TraceServices/Private/Modules/CookProfilerModule.cpp
+++ b/Engine/Source/Developer/TraceServices/Private/Modules/CookProfilerModule.cpp
@@ -1,42 +1,48 @@
 // Copyright Epic Games, Inc. All Rights Reserved.
 
 #include "CookProfilerModule.h"
-#include "Analyzers/CookAnalysis.h"
+
 #include "AnalysisServicePrivate.h"
+#include "Analyzers/CookAnalysis.h"
 #include "Model/CookProfilerProviderPrivate.h"
 
 namespace TraceServices
 {
 
-static const FName CookProfilerModuleName("TraceModule_CookProfiler");
-static const FName CookProfilerProviderName("CookProfilerProvider");
-
 void FCookProfilerModule::GetModuleInfo(FModuleInfo& OutModuleInfo)
 {
-	OutModuleInfo.Name = CookProfilerProviderName;
+	static const FName CookProfilerModuleName("TraceModule_CookProfiler");
+
+	OutModuleInfo.Name = CookProfilerModuleName;
 	OutModuleInfo.DisplayName = TEXT("Cook Profiling");
 }
 
 void FCookProfilerModule::OnAnalysisBegin(IAnalysisSession& Session)
 {
 	TSharedPtr<FCookProfilerProvider> CookProfilerProvider = MakeShared<FCookProfilerProvider>(Session);
-	Session.AddProvider(CookProfilerProviderName, CookProfilerProvider, CookProfilerProvider);
+	Session.AddProvider(GetCookProfilerProviderName(), CookProfilerProvider, CookProfilerProvider);
 	Session.AddAnalyzer(new FCookAnalyzer(Session, *CookProfilerProvider));
 }
 
-void FCookProfilerModule::GetLoggers(TArray<const TCHAR *>& OutLoggers)
+void FCookProfilerModule::GetLoggers(TArray<const TCHAR*>& OutLoggers)
 {
 	OutLoggers.Add(TEXT("CookProfiler"));
 }
 
+FName GetCookProfilerProviderName()
+{
+	static const FName Name("CookProfilerProvider");
+	return Name;
+}
+
 const ICookProfilerProvider* ReadCookProfilerProvider(const IAnalysisSession& Session)
 {
-	return Session.ReadProvider<ICookProfilerProvider>(CookProfilerProviderName);
+	return Session.ReadProvider<ICookProfilerProvider>(GetCookProfilerProviderName());
 }
 
 IEditableCookProfilerProvider* EditCookProfilerProvider(IAnalysisSession& Session)
 {
-	return Session.EditProvider<IEditableCookProfilerProvider>(CookProfilerProviderName);
+	return Session.EditProvider<IEditableCookProfilerProvider>(GetCookProfilerProviderName());
 }
 
 } // namespace TraceServices
diff --git a/Engine/Source/Developer/TraceServices/Private/Modules/CookProfilerModule.h b/Engine/Source/Developer/TraceServices/Private/Modules/CookProfilerModule.h
--- a/Engine/Source/Developer/TraceServices/Private/Modules/CookProfilerModule.h
+++ b/Engine/Source/Developer/TraceServices/Private/Modules/CookProfilerModule.h
@@ -11,17 +11,12 @@ class FCookProfilerModule
 	: public IModule
 {
 public:
-	#if WITH_EDITOR
-	virtual bool ShouldBeEnabledByDefault() const override { return false; }
-	#endif
 	virtual void GetModuleInfo(FModuleInfo& OutModuleInfo) override;
-	virtual void OnAnalysisBegin(IAnalysisSession& Session) override;
-	virtual void GenerateReports(const IAnalysisSession& Session, const TCHAR* CmdLine, const TCHAR* OutputDirectory) override {}
+#if WITH_EDITOR
+	virtual bool ShouldBeEnabledByDefault() const override { return false; }
+#endif
 	virtual void GetLoggers(TArray<const TCHAR*>& OutLoggers) override;
-	virtual const TCHAR* GetCommandLineArgument() override
-	{
-		return TEXT("cooktrace");
-	}
+	virtual void OnAnalysisBegin(IAnalysisSession& Session) override;
 };
 
 } // namespace TraceServices
diff --git a/Engine/Source/Developer/TraceServices/Private/Modules/CountersModule.h b/Engine/Source/Developer/TraceServices/Private/Modules/CountersModule.h
--- a/Engine/Source/Developer/TraceServices/Private/Modules/CountersModule.h
+++ b/Engine/Source/Developer/TraceServices/Private/Modules/CountersModule.h
@@ -12,12 +12,8 @@ class FCountersModule
 {
 public:
 	virtual void GetModuleInfo(FModuleInfo& OutModuleInfo) override;
+	virtual void GetLoggers(TArray<const TCHAR*>& OutLoggers) override;
 	virtual void OnAnalysisBegin(IAnalysisSession& Session) override;
-	virtual void GetLoggers(TArray<const TCHAR *>& OutLoggers) override;
-	virtual const TCHAR* GetCommandLineArgument() override
-	{
-		return TEXT("counterstrace");
-	}
 	virtual void GenerateReports(const IAnalysisSession& Session, const TCHAR* CmdLine, const TCHAR* OutputDirectory) override {}
 };
 
diff --git a/Engine/Source/Developer/TraceServices/Private/Modules/CsvProfilerModule.cpp b/Engine/Source/Developer/TraceServices/Private/Modules/CsvProfilerModule.cpp
--- a/Engine/Source/Developer/TraceServices/Private/Modules/CsvProfilerModule.cpp
+++ b/Engine/Source/Developer/TraceServices/Private/Modules/CsvProfilerModule.cpp
@@ -1,6 +1,7 @@
 // Copyright Epic Games, Inc. All Rights Reserved.
 
 #include "CsvProfilerModule.h"
+
 #include "Analyzers/CsvProfilerTraceAnalysis.h"
 #include "TraceServices/Model/Counters.h"
 #include "TraceServices/Model/Frames.h"
@@ -9,15 +10,13 @@
 namespace TraceServices
 {
 
-static const FName CsvProfilerModuleName("TraceModule_CsvProfiler");
-static const FName CsvProfilerProviderName("CsvProfilerProvider");
-
 void FCsvProfilerModule::GetModuleInfo(FModuleInfo& OutModuleInfo)
 {
+	static const FName CsvProfilerModuleName("TraceModule_CsvProfiler");
+
 	OutModuleInfo.Name = CsvProfilerModuleName;
 	OutModuleInfo.DisplayName = TEXT("CsvProfiler");
 }
-	
 
 void FCsvProfilerModule::OnAnalysisBegin(IAnalysisSession& Session)
 {
@@ -25,7 +24,7 @@ void FCsvProfilerModule::OnAnalysisBegin(IAnalysisSession& Session)
 	const IThreadProvider& ThreadProvider = ReadThreadProvider(Session);
 	IEditableCounterProvider& EditableCounterProvider = EditCounterProvider(Session);
 	TSharedPtr<FCsvProfilerProvider> CsvProfilerProvider = MakeShared<FCsvProfilerProvider>(Session);
-	Session.AddProvider(CsvProfilerProviderName, CsvProfilerProvider);
+	Session.AddProvider(GetCsvProfilerProviderName(), CsvProfilerProvider);
 	Session.AddAnalyzer(new FCsvProfilerAnalyzer(Session, *CsvProfilerProvider, EditableCounterProvider, FrameProvider, ThreadProvider));
 }
 
@@ -41,9 +40,15 @@ void FCsvProfilerModule::GenerateReports(const IAnalysisSession& Session, const
 	}
 }
 
+FName GetCsvProfilerProviderName()
+{
+	static const FName Name("CsvProfilerProvider");
+	return Name;
+}
+
 const ICsvProfilerProvider* ReadCsvProfilerProvider(const IAnalysisSession& Session)
 {
-	return Session.ReadProvider<ICsvProfilerProvider>(CsvProfilerProviderName);
+	return Session.ReadProvider<ICsvProfilerProvider>(GetCsvProfilerProviderName());
 }
 
 } // namespace TraceServices
diff --git a/Engine/Source/Developer/TraceServices/Private/Modules/CsvProfilerModule.h b/Engine/Source/Developer/TraceServices/Private/Modules/CsvProfilerModule.h
--- a/Engine/Source/Developer/TraceServices/Private/Modules/CsvProfilerModule.h
+++ b/Engine/Source/Developer/TraceServices/Private/Modules/CsvProfilerModule.h
@@ -3,7 +3,6 @@
 #pragma once
 
 #include "TraceServices/ModuleService.h"
-#include "UObject/NameTypes.h"
 
 namespace TraceServices
 {
@@ -14,13 +13,6 @@ class FCsvProfilerModule
 public:
 	virtual void GetModuleInfo(FModuleInfo& OutModuleInfo) override;
 	virtual void OnAnalysisBegin(IAnalysisSession& Session) override;
-	virtual void GetLoggers(TArray<const TCHAR*>& OutLoggers) override
-	{
-	};
-	virtual const TCHAR* GetCommandLineArgument() override
-	{
-		return nullptr;
-	}
 	virtual void GenerateReports(const IAnalysisSession& Session, const TCHAR* CmdLine, const TCHAR* OutputDirectory) override;
 };
 
diff --git a/Engine/Source/Developer/TraceServices/Private/Modules/LoadTimeProfilerModule.cpp b/Engine/Source/Developer/TraceServices/Private/Modules/LoadTimeProfilerModule.cpp
--- a/Engine/Source/Developer/TraceServices/Private/Modules/LoadTimeProfilerModule.cpp
+++ b/Engine/Source/Developer/TraceServices/Private/Modules/LoadTimeProfilerModule.cpp
@@ -1,20 +1,17 @@
 // Copyright Epic Games, Inc. All Rights Reserved.
 
 #include "LoadTimeProfilerModule.h"
-#include "Analyzers/PlatformFileTraceAnalysis.h"
-#include "Analyzers/LoadTimeTraceAnalysis.h"
+
 #include "AnalysisServicePrivate.h"
+#include "Analyzers/LoadTimeTraceAnalysis.h"
+#include "Analyzers/PlatformFileTraceAnalysis.h"
+#include "HAL/FileManager.h"
 #include "Model/FileActivity.h"
 #include "TraceServices/Model/Bookmarks.h"
-#include "HAL/FileManager.h"
 
 namespace TraceServices
 {
 
-static const FName LoadTimeProfilerModuleName("TraceModule_LoadTimeProfiler");
-static const FName LoadTimeProfilerProviderName("LoadTimeProfiler");
-static const FName FileActivityProviderName("FileActivity");
-
 void BookmarksToCsv(const IBookmarkProvider& BookmarkProvider, const TCHAR* Filename, double CaptureStartTime, double CaptureEndTime)
 {
 	TSharedPtr<FArchive> BookmarksOutputFile = MakeShareable(IFileManager::Get().CreateFileWriter(Filename));
@@ -30,6 +27,8 @@ void BookmarksToCsv(const IBookmarkProvider& BookmarkProvider, const TCHAR* File
 
 void FLoadTimeProfilerModule::GetModuleInfo(FModuleInfo& OutModuleInfo)
 {
+	static const FName LoadTimeProfilerModuleName("TraceModule_LoadTimeProfiler");
+
 	OutModuleInfo.Name = LoadTimeProfilerModuleName;
 	OutModuleInfo.DisplayName = TEXT("Asset Loading");
 }
@@ -37,14 +36,14 @@ void FLoadTimeProfilerModule::GetModuleInfo(FModuleInfo& OutModuleInfo)
 void FLoadTimeProfilerModule::OnAnalysisBegin(IAnalysisSession& Session)
 {
 	TSharedPtr<FLoadTimeProfilerProvider> LoadTimeProfilerProvider = MakeShared<FLoadTimeProfilerProvider>(Session, EditCounterProvider(Session));
-	Session.AddProvider(LoadTimeProfilerProviderName, LoadTimeProfilerProvider);
+	Session.AddProvider(GetLoadTimeProfilerProviderName(), LoadTimeProfilerProvider);
 	Session.AddAnalyzer(new FAsyncLoadingTraceAnalyzer(Session, *LoadTimeProfilerProvider));
 	TSharedPtr<FFileActivityProvider> FileActivityProvider = MakeShared<FFileActivityProvider>(Session);
-	Session.AddProvider(FileActivityProviderName, FileActivityProvider);
+	Session.AddProvider(GetFileActivityProviderName(), FileActivityProvider);
 	Session.AddAnalyzer(new FPlatformFileTraceAnalyzer(Session, *FileActivityProvider));
 }
 
-void FLoadTimeProfilerModule::GetLoggers(TArray<const TCHAR *>& OutLoggers)
+void FLoadTimeProfilerModule::GetLoggers(TArray<const TCHAR*>& OutLoggers)
 {
 	OutLoggers.Add(TEXT("LoadTime"));
 	OutLoggers.Add(TEXT("PlatformFile"));
@@ -59,7 +58,7 @@ void FLoadTimeProfilerModule::GenerateReports(const IAnalysisSession& Session, c
 	FParse::Value(CmdLine, TEXT("-BeginCaptureBookmark="), BeginCaptureBookmarkName);
 	FString EndCaptureBookmarkName;
 	FParse::Value(CmdLine, TEXT("-EndCaptureBookmark="), EndCaptureBookmarkName);
-	
+
 
 	BookmarkProvider.EnumerateBookmarks(0.0, DBL_MAX, [BeginCaptureBookmarkName, EndCaptureBookmarkName, &CaptureStartTime, &CaptureEndTime](const FBookmark& Bookmark)
 	{
@@ -106,14 +105,26 @@ void FLoadTimeProfilerModule::GenerateReports(const IAnalysisSession& Session, c
 	BookmarksToCsv(BookmarkProvider, *(FString(ReportDirectory) / TEXT("Bookmarks.csv")), CaptureStartTime, CaptureEndTime);
 }
 
+FName GetLoadTimeProfilerProviderName()
+{
+	static const FName Name("LoadTimeProfilerProvider");
+	return Name;
+}
+
 const ILoadTimeProfilerProvider* ReadLoadTimeProfilerProvider(const IAnalysisSession& Session)
 {
-	return Session.ReadProvider<ILoadTimeProfilerProvider>(LoadTimeProfilerProviderName);
+	return Session.ReadProvider<ILoadTimeProfilerProvider>(GetLoadTimeProfilerProviderName());
+}
+
+FName GetFileActivityProviderName()
+{
+	static const FName Name("FileActivityProvider");
+	return Name;
 }
 
 const IFileActivityProvider* ReadFileActivityProvider(const IAnalysisSession& Session)
 {
-	return Session.ReadProvider<IFileActivityProvider>(FileActivityProviderName);
+	return Session.ReadProvider<IFileActivityProvider>(GetFileActivityProviderName());
 }
 
 } // namespace TraceServices
diff --git a/Engine/Source/Developer/TraceServices/Private/Modules/LoadTimeProfilerModule.h b/Engine/Source/Developer/TraceServices/Private/Modules/LoadTimeProfilerModule.h
--- a/Engine/Source/Developer/TraceServices/Private/Modules/LoadTimeProfilerModule.h
+++ b/Engine/Source/Developer/TraceServices/Private/Modules/LoadTimeProfilerModule.h
@@ -11,17 +11,13 @@ class FLoadTimeProfilerModule
 	: public IModule
 {
 public:
-	#if WITH_EDITOR
-	virtual bool ShouldBeEnabledByDefault() const override { return false; }
-	#endif
 	virtual void GetModuleInfo(FModuleInfo& OutModuleInfo) override;
+#if WITH_EDITOR
+	virtual bool ShouldBeEnabledByDefault() const override { return false; }
+#endif
+	virtual void GetLoggers(TArray<const TCHAR*>& OutLoggers) override;
 	virtual void OnAnalysisBegin(IAnalysisSession& Session) override;
 	virtual void GenerateReports(const IAnalysisSession& Session, const TCHAR* CmdLine, const TCHAR* OutputDirectory) override;
-	virtual void GetLoggers(TArray<const TCHAR*>& OutLoggers) override;
-	virtual const TCHAR* GetCommandLineArgument() override
-	{
-		return TEXT("loadtimetrace");
-	}
 };
 
 } // namespace TraceServices
diff --git a/Engine/Source/Developer/TraceServices/Private/Modules/NetProfilerModule.cpp b/Engine/Source/Developer/TraceServices/Private/Modules/NetProfilerModule.cpp
--- a/Engine/Source/Developer/TraceServices/Private/Modules/NetProfilerModule.cpp
+++ b/Engine/Source/Developer/TraceServices/Private/Modules/NetProfilerModule.cpp
@@ -1,17 +1,18 @@
 // Copyright Epic Games, Inc. All Rights Reserved.
 
 #include "NetProfilerModule.h"
-#include "Analyzers/NetTraceAnalyzer.h"
+
 #include "AnalysisServicePrivate.h"
+#include "Analyzers/NetTraceAnalyzer.h"
 #include "Model/NetProfilerProvider.h"
 
 namespace TraceServices
 {
 
-FName FNetProfilerModule::ModuleName("TraceModule_NetProfiler");
-
 void FNetProfilerModule::GetModuleInfo(FModuleInfo& OutModuleInfo)
 {
+	static const FName ModuleName("TraceModule_NetProfiler");
+
 	OutModuleInfo.Name = ModuleName;
 	OutModuleInfo.DisplayName = TEXT("NetProfiler");
 }
@@ -24,8 +25,4 @@ void FNetProfilerModule::OnAnalysisBegin(IAnalysisSession& Session)
 	Session.AddAnalyzer(new FNetTraceAnalyzer(Session, *NetProfilerProvider));
 }
 
-void FNetProfilerModule::GetLoggers(TArray<const TCHAR*>& OutLoggers)
-{
-}
-
 } // namespace TraceServices
diff --git a/Engine/Source/Developer/TraceServices/Private/Modules/NetProfilerModule.h b/Engine/Source/Developer/TraceServices/Private/Modules/NetProfilerModule.h
--- a/Engine/Source/Developer/TraceServices/Private/Modules/NetProfilerModule.h
+++ b/Engine/Source/Developer/TraceServices/Private/Modules/NetProfilerModule.h
@@ -3,7 +3,6 @@
 #pragma once
 
 #include "TraceServices/ModuleService.h"
-#include "UObject/NameTypes.h"
 
 namespace TraceServices
 {
@@ -14,15 +13,6 @@ class FNetProfilerModule
 public:
 	virtual void GetModuleInfo(FModuleInfo& OutModuleInfo) override;
 	virtual void OnAnalysisBegin(IAnalysisSession& Session) override;
-	virtual void GetLoggers(TArray<const TCHAR*>& OutLoggers) override;
-	virtual const TCHAR* GetCommandLineArgument() override
-	{
-		return nullptr;
-	}
-	virtual void GenerateReports(const IAnalysisSession& Session, const TCHAR* CmdLine, const TCHAR* OutputDirectory) override {}
-
-private:
-	static FName ModuleName;
 };
 
 } // namespace TraceServices
diff --git a/Engine/Source/Developer/TraceServices/Private/Modules/PlatformEventsModule.cpp b/Engine/Source/Developer/TraceServices/Private/Modules/PlatformEventsModule.cpp
--- a/Engine/Source/Developer/TraceServices/Private/Modules/PlatformEventsModule.cpp
+++ b/Engine/Source/Developer/TraceServices/Private/Modules/PlatformEventsModule.cpp
@@ -1,31 +1,32 @@
 // Copyright Epic Games, Inc. All Rights Reserved.
 
 #include "PlatformEventsModule.h"
-#include "Analyzers/PlatformEventTraceAnalysis.h"
+
 #include "AnalysisServicePrivate.h"
+#include "Analyzers/PlatformEventTraceAnalysis.h"
 #include "Model/ContextSwitchesPrivate.h"
 #include "Model/StackSamplesPrivate.h"
 
 namespace TraceServices
 {
 
-static const FName PlatformEventsModuleName("TraceModule_PlatformEvents");
-
 void FPlatformEventsModule::GetModuleInfo(FModuleInfo& OutModuleInfo)
 {
+	static const FName PlatformEventsModuleName("TraceModule_PlatformEvents");
+
 	OutModuleInfo.Name = PlatformEventsModuleName;
 	OutModuleInfo.DisplayName = TEXT("PlatformEvents");
 }
-	
+
 void FPlatformEventsModule::OnAnalysisBegin(IAnalysisSession& InSession)
 {
 	FAnalysisSession& Session = static_cast<FAnalysisSession&>(InSession);
 
 	TSharedPtr<FContextSwitchesProvider> ContextSwitchesProvider = MakeShared<FContextSwitchesProvider>(Session);
 	TSharedPtr<FStackSamplesProvider> StackSamplesProvider = MakeShared<FStackSamplesProvider>(Session);
-	
-	Session.AddProvider(FContextSwitchesProvider::ProviderName, ContextSwitchesProvider);
-	Session.AddProvider(FStackSamplesProvider::ProviderName, StackSamplesProvider);
+
+	Session.AddProvider(GetContextSwitchesProviderName(), ContextSwitchesProvider);
+	Session.AddProvider(GetStackSamplesProviderName(), StackSamplesProvider);
 	Session.AddAnalyzer(new FPlatformEventTraceAnalyzer(Session, *ContextSwitchesProvider, *StackSamplesProvider));
 }
 
diff --git a/Engine/Source/Developer/TraceServices/Private/Modules/PlatformEventsModule.h b/Engine/Source/Developer/TraceServices/Private/Modules/PlatformEventsModule.h
--- a/Engine/Source/Developer/TraceServices/Private/Modules/PlatformEventsModule.h
+++ b/Engine/Source/Developer/TraceServices/Private/Modules/PlatformEventsModule.h
@@ -3,7 +3,6 @@
 #pragma once
 
 #include "TraceServices/ModuleService.h"
-#include "UObject/NameTypes.h"
 
 namespace TraceServices
 {
@@ -14,16 +13,6 @@ class FPlatformEventsModule
 public:
 	virtual void GetModuleInfo(FModuleInfo& OutModuleInfo) override;
 	virtual void OnAnalysisBegin(IAnalysisSession& Session) override;
-	virtual void GetLoggers(TArray<const TCHAR *>& OutLoggers) override
-	{
-	}
-	virtual const TCHAR* GetCommandLineArgument() override
-	{
-		return nullptr;
-	}
-	virtual void GenerateReports(const IAnalysisSession& Session, const TCHAR* CmdLine, const TCHAR* OutputDirectory) override
-	{
-	}
 };
 
 } // namespace TraceServices
diff --git a/Engine/Source/Developer/TraceServices/Private/Modules/StatsModule.cpp b/Engine/Source/Developer/TraceServices/Private/Modules/StatsModule.cpp
--- a/Engine/Source/Developer/TraceServices/Private/Modules/StatsModule.cpp
+++ b/Engine/Source/Developer/TraceServices/Private/Modules/StatsModule.cpp
@@ -1,29 +1,25 @@
 // Copyright Epic Games, Inc. All Rights Reserved.
 
 #include "StatsModule.h"
+
 #include "Analyzers/StatsTraceAnalysis.h"
 #include "TraceServices/Model/Counters.h"
 
 namespace TraceServices
 {
 
-static const FName StatsModuleName("TraceModule_Stats");
-
 void FStatsModule::GetModuleInfo(FModuleInfo& OutModuleInfo)
 {
+	static const FName StatsModuleName("TraceModule_Stats");
+
 	OutModuleInfo.Name = StatsModuleName;
 	OutModuleInfo.DisplayName = TEXT("Stats");
 }
-	
+
 void FStatsModule::OnAnalysisBegin(IAnalysisSession& Session)
 {
 	IEditableCounterProvider& EditableCounterProvider = EditCounterProvider(Session);
 	Session.AddAnalyzer(new FStatsAnalyzer(Session, EditableCounterProvider));
 }
 
-void FStatsModule::GetLoggers(TArray<const TCHAR *>& OutLoggers)
-{
-	//OutLoggers.Add(TEXT("Stats"));
-}
-
 } // namespace TraceServices
diff --git a/Engine/Source/Developer/TraceServices/Private/Modules/StatsModule.h b/Engine/Source/Developer/TraceServices/Private/Modules/StatsModule.h
--- a/Engine/Source/Developer/TraceServices/Private/Modules/StatsModule.h
+++ b/Engine/Source/Developer/TraceServices/Private/Modules/StatsModule.h
@@ -3,7 +3,6 @@
 #pragma once
 
 #include "TraceServices/ModuleService.h"
-#include "UObject/NameTypes.h"
 
 namespace TraceServices
 {
@@ -14,12 +13,6 @@ class FStatsModule
 public:
 	virtual void GetModuleInfo(FModuleInfo& OutModuleInfo) override;
 	virtual void OnAnalysisBegin(IAnalysisSession& Session) override;
-	virtual void GetLoggers(TArray<const TCHAR *>& OutLoggers) override;
-	virtual const TCHAR* GetCommandLineArgument() override
-	{
-		return nullptr;
-	}
-	virtual void GenerateReports(const IAnalysisSession& Session, const TCHAR* CmdLine, const TCHAR* OutputDirectory) override {}
 };
 
 } // namespace TraceServices
diff --git a/Engine/Source/Developer/TraceServices/Private/Modules/TasksModule.cpp b/Engine/Source/Developer/TraceServices/Private/Modules/TasksModule.cpp
--- a/Engine/Source/Developer/TraceServices/Private/Modules/TasksModule.cpp
+++ b/Engine/Source/Developer/TraceServices/Private/Modules/TasksModule.cpp
@@ -1,31 +1,39 @@
 // Copyright Epic Games, Inc. All Rights Reserved.
 
 #include "TasksModule.h"
+
+#include "AnalysisServicePrivate.h"
 #include "Analyzers/TasksAnalysis.h"
 #include "Model/TasksProfilerPrivate.h"
-#include "AnalysisServicePrivate.h"
 #include "TraceServices/ModuleService.h"
 
 namespace TraceServices
+{
+
+void FTasksModule::GetModuleInfo(FModuleInfo& OutModuleInfo)
 {
 	static const FName TasksModuleName("TraceModule_TasksProfiler");
-	static const FName TasksProviderName("TasksProvider");
-
-	void FTasksModule::GetModuleInfo(FModuleInfo& OutModuleInfo)
-	{
-		OutModuleInfo.Name = TasksModuleName;
-		OutModuleInfo.DisplayName = TEXT("TasksProfiler");
-	}
-
-	void FTasksModule::OnAnalysisBegin(IAnalysisSession& Session)
-	{
-		TSharedPtr<FTasksProvider> TasksProvider = MakeShared<FTasksProvider>(Session);
-		Session.AddProvider(TasksProviderName, TasksProvider);
-		Session.AddAnalyzer(new FTasksAnalyzer(Session, *TasksProvider));
-	}
-
-	const ITasksProvider* ReadTasksProvider(const IAnalysisSession& Session)
-	{
-		return Session.ReadProvider<ITasksProvider>(TasksProviderName);
-	}
+
+	OutModuleInfo.Name = TasksModuleName;
+	OutModuleInfo.DisplayName = TEXT("TasksProfiler");
+}
+
+void FTasksModule::OnAnalysisBegin(IAnalysisSession& Session)
+{
+	TSharedPtr<FTasksProvider> TasksProvider = MakeShared<FTasksProvider>(Session);
+	Session.AddProvider(GetTaskProviderName(), TasksProvider);
+	Session.AddAnalyzer(new FTasksAnalyzer(Session, *TasksProvider));
+}
+
+FName GetTaskProviderName()
+{
+	static const FName Name("TasksProvider");
+	return Name;
+}
+
+const ITasksProvider* ReadTasksProvider(const IAnalysisSession& Session)
+{
+	return Session.ReadProvider<ITasksProvider>(GetTaskProviderName());
+}
+
 } // namespace TraceServices
diff --git a/Engine/Source/Developer/TraceServices/Private/Modules/TasksModule.h b/Engine/Source/Developer/TraceServices/Private/Modules/TasksModule.h
--- a/Engine/Source/Developer/TraceServices/Private/Modules/TasksModule.h
+++ b/Engine/Source/Developer/TraceServices/Private/Modules/TasksModule.h
@@ -3,30 +3,19 @@
 #pragma once
 
 #include "TraceServices/ModuleService.h"
-#include "UObject/NameTypes.h"
 
 namespace TraceServices
 {
-	class FTasksModule : public IModule
-	{
-	public:
-		virtual void GetModuleInfo(FModuleInfo& OutModuleInfo) override;
-		virtual void OnAnalysisBegin(IAnalysisSession& Session) override;
-		virtual void GetLoggers(TArray<const TCHAR*>& OutLoggers) override
-		{
-		}
-		virtual const TCHAR* GetCommandLineArgument() override
-		{
-			return nullptr;
-		}
-		virtual void GenerateReports(const IAnalysisSession& Session, const TCHAR* CmdLine, const TCHAR* OutputDirectory) override
-		{
-		}
+
+class FTasksModule
+	: public IModule
+{
+public:
+	virtual void GetModuleInfo(FModuleInfo& OutModuleInfo) override;
 #if WITH_EDITOR
-		virtual bool ShouldBeEnabledByDefault() const override
-		{
-			return false;
-		}
+	virtual bool ShouldBeEnabledByDefault() const override { return false; }
 #endif
-	};
+	virtual void OnAnalysisBegin(IAnalysisSession& Session) override;
+};
+
 } // namespace TraceServices
diff --git a/Engine/Source/Developer/TraceServices/Private/Modules/TimingProfilerModule.cpp b/Engine/Source/Developer/TraceServices/Private/Modules/TimingProfilerModule.cpp
--- a/Engine/Source/Developer/TraceServices/Private/Modules/TimingProfilerModule.cpp
+++ b/Engine/Source/Developer/TraceServices/Private/Modules/TimingProfilerModule.cpp
@@ -1,52 +1,56 @@
 // Copyright Epic Games, Inc. All Rights Reserved.
 
 #include "TimingProfilerModule.h"
+
+#include "AnalysisServicePrivate.h"
 #include "Analyzers/CpuProfilerTraceAnalysis.h"
 #include "Analyzers/GpuProfilerTraceAnalysis.h"
-#include "AnalysisServicePrivate.h"
 #include "Model/ThreadsPrivate.h"
 #include "Model/TimingProfilerPrivate.h"
 
 namespace TraceServices
 {
 
-static const FName TimingProfilerModuleName("TraceModule_TimingProfiler");
-static const FName TimingProfilerProviderName("TimingProfilerProvider");
-
 void FTimingProfilerModule::GetModuleInfo(FModuleInfo& OutModuleInfo)
 {
+	static const FName TimingProfilerModuleName("TraceModule_TimingProfiler");
+
 	OutModuleInfo.Name = TimingProfilerModuleName;
 	OutModuleInfo.DisplayName = TEXT("Timing");
 }
-	
+
 void FTimingProfilerModule::OnAnalysisBegin(IAnalysisSession& InSession)
 {
 	FAnalysisSession& Session = static_cast<FAnalysisSession&>(InSession);
-	
-	IEditableThreadProvider* EditableThreadProvider = Session.EditProvider<IEditableThreadProvider>(FThreadProvider::ProviderName);
-	check(EditableThreadProvider != nullptr);
 
-	// see comment in FAnalysisService::StartAnalysis for details
+	IEditableThreadProvider& EditableThreadProvider = EditThreadProvider(Session);
+
 	TSharedPtr<FTimingProfilerProvider> TimingProfilerProvider = MakeShared<FTimingProfilerProvider>(Session);
-	Session.AddProvider(TimingProfilerProviderName, TimingProfilerProvider, TimingProfilerProvider);
-	Session.AddAnalyzer(new FCpuProfilerAnalyzer(Session, *TimingProfilerProvider, *EditableThreadProvider));
+	Session.AddProvider(GetTimingProfilerProviderName(), TimingProfilerProvider, TimingProfilerProvider);
+	Session.AddAnalyzer(new FCpuProfilerAnalyzer(Session, *TimingProfilerProvider, EditableThreadProvider));
 	Session.AddAnalyzer(new FGpuProfilerAnalyzer(Session, *TimingProfilerProvider));
 }
 
-void FTimingProfilerModule::GetLoggers(TArray<const TCHAR *>& OutLoggers)
+void FTimingProfilerModule::GetLoggers(TArray<const TCHAR*>& OutLoggers)
 {
 	OutLoggers.Add(TEXT("CpuProfiler"));
 	OutLoggers.Add(TEXT("GpuProfiler"));
 }
 
+FName GetTimingProfilerProviderName()
+{
+	static const FName Name("TimingProfilerProvider");
+	return Name;
+}
+
 const ITimingProfilerProvider* ReadTimingProfilerProvider(const IAnalysisSession& Session)
 {
-	return Session.ReadProvider<ITimingProfilerProvider>(TimingProfilerProviderName);
+	return Session.ReadProvider<ITimingProfilerProvider>(GetTimingProfilerProviderName());
 }
 
 IEditableTimingProfilerProvider* EditTimingProfilerProvider(IAnalysisSession& Session)
 {
-	return Session.EditProvider<IEditableTimingProfilerProvider>(TimingProfilerProviderName);
+	return Session.EditProvider<IEditableTimingProfilerProvider>(GetTimingProfilerProviderName());
 }
 
 } // namespace TraceServices
diff --git a/Engine/Source/Developer/TraceServices/Private/Modules/TimingProfilerModule.h b/Engine/Source/Developer/TraceServices/Private/Modules/TimingProfilerModule.h
--- a/Engine/Source/Developer/TraceServices/Private/Modules/TimingProfilerModule.h
+++ b/Engine/Source/Developer/TraceServices/Private/Modules/TimingProfilerModule.h
@@ -3,7 +3,6 @@
 #pragma once
 
 #include "TraceServices/ModuleService.h"
-#include "UObject/NameTypes.h"
 
 namespace TraceServices
 {
@@ -13,13 +12,8 @@ class FTimingProfilerModule
 {
 public:
 	virtual void GetModuleInfo(FModuleInfo& OutModuleInfo) override;
+	virtual void GetLoggers(TArray<const TCHAR*>& OutLoggers) override;
 	virtual void OnAnalysisBegin(IAnalysisSession& Session) override;
-	virtual void GetLoggers(TArray<const TCHAR *>& OutLoggers) override;
-	virtual const TCHAR* GetCommandLineArgument() override
-	{
-		return TEXT("cpuprofilertrace");
-	}
-	virtual void GenerateReports(const IAnalysisSession& Session, const TCHAR* CmdLine, const TCHAR* OutputDirectory) override {}
 };
 
 } // namespace TraceServices
diff --git a/Engine/Source/Developer/TraceServices/Public/TraceServices/Model/AllocationsProvider.h b/Engine/Source/Developer/TraceServices/Public/TraceServices/Model/AllocationsProvider.h
--- a/Engine/Source/Developer/TraceServices/Public/TraceServices/Model/AllocationsProvider.h
+++ b/Engine/Source/Developer/TraceServices/Public/TraceServices/Model/AllocationsProvider.h
@@ -2,10 +2,12 @@
 
 #pragma once
 
-#include "AnalysisSession.h"
-#include "CoreTypes.h"
-#include "Templates/UniquePtr.h"
+#include "HAL/Platform.h"
 #include "ProfilingDebugging/MemoryTrace.h"
+#include "Templates/Function.h"
+#include "Templates/UniquePtr.h"
+#include "TraceServices/Model/AnalysisSession.h"
+#include "UObject/NameTypes.h"
 
 #include <new>
 
@@ -14,7 +16,7 @@ namespace TraceServices
 
 // Id type for tags
 typedef uint32 TagIdType;
-	
+
 class IAllocationsProvider : public IProvider
 {
 public:
diff --git a/Engine/Source/Developer/TraceServices/Public/TraceServices/Model/Bookmarks.h b/Engine/Source/Developer/TraceServices/Public/TraceServices/Model/Bookmarks.h
--- a/Engine/Source/Developer/TraceServices/Public/TraceServices/Model/Bookmarks.h
+++ b/Engine/Source/Developer/TraceServices/Public/TraceServices/Model/Bookmarks.h
@@ -1,12 +1,11 @@
 // Copyright Epic Games, Inc. All Rights Reserved.
 
-#pragma  once
+#pragma once
 
 #include "HAL/Platform.h"
 #include "Templates/Function.h"
 #include "TraceServices/Model/AnalysisSession.h"
-
-template <typename FuncType> class TFunctionRef;
+#include "UObject/NameTypes.h"
 
 namespace TraceServices
 {
@@ -37,7 +36,7 @@ public:
 
 	/*
 	* Update the information in the Bookmark.
-	* 
+	*
 	* @param BookmarkPoint	The unique identity (memory address) of the bookmark.
 	* @param FormatString	The format string of the bookmark's message.
 	* @param File			The source file the bookmark is in.
@@ -56,7 +55,7 @@ public:
 
 	/*
 	* Append a new instance of a bookmark from the trace session.
-	* 
+	*
 	* @param BookmarkPoint	The unique identity (memory address) of the bookmark.
 	* @param Time			The time in seconds of the event.
 	* @param Text			The fully formatted bookmark string for this instance. This may vary between instances.
diff --git a/Engine/Source/Developer/TraceServices/Public/TraceServices/Model/Callstack.h b/Engine/Source/Developer/TraceServices/Public/TraceServices/Model/Callstack.h
--- a/Engine/Source/Developer/TraceServices/Public/TraceServices/Model/Callstack.h
+++ b/Engine/Source/Developer/TraceServices/Public/TraceServices/Model/Callstack.h
@@ -2,9 +2,11 @@
 
 #pragma once
 
+#include "HAL/Platform.h"
 #include "TraceServices/Model/AnalysisSession.h"
 #include "TraceServices/Model/Memory.h"
 #include "TraceServices/Model/Modules.h"
+#include "UObject/NameTypes.h"
 
 namespace TraceServices
 {
@@ -21,9 +23,9 @@ struct TRACESERVICES_API FCallstack
 {
 	/** Creates an empty callstack. */
 	FCallstack();
-	/** Creates an callstack initialised with a certain number of stack frames. */
+	/** Creates an callstack initialized with a certain number of stack frames. */
 	FCallstack(const FStackFrame* FirstEntry, uint8 FrameCount);
-	/** Initialises the callstack with a certain number of stack frames. */
+	/** Initializes the callstack with a certain number of stack frames. */
 	void Init(const FStackFrame* FirstEntry, uint8 FrameCount);
 	/** Gets the number of stack frames in callstack. */
 	uint32 Num() const;
@@ -68,7 +70,7 @@ public:
 	virtual void GetCallstacks(const TArrayView<uint32>& CallstackIds, FCallstack const** OutCallstacks) const = 0;
 
 	/**
-	  * Gets a callstack id for a registerd callstack hash.
+	  * Gets a callstack id for a registered callstack hash.
 	  * @param CallstackHash The callstack hash.
 	  * @returns The callstack id.
 	  */
diff --git a/Engine/Source/Developer/TraceServices/Public/TraceServices/Model/Channel.h b/Engine/Source/Developer/TraceServices/Public/TraceServices/Model/Channel.h
--- a/Engine/Source/Developer/TraceServices/Public/TraceServices/Model/Channel.h
+++ b/Engine/Source/Developer/TraceServices/Public/TraceServices/Model/Channel.h
@@ -2,8 +2,11 @@
 
 #pragma once
 
-#include "TraceServices/Model/AnalysisSession.h"
+#include "HAL/Platform.h"
+#include "Containers/Array.h"
 #include "Containers/UnrealString.h"
+#include "TraceServices/Model/AnalysisSession.h"
+#include "UObject/NameTypes.h"
 
 namespace TraceServices
 {
@@ -21,9 +24,12 @@ class IChannelProvider
 {
 public:
 	virtual ~IChannelProvider() = default;
-	virtual uint64	GetChannelCount() const = 0;
+	virtual uint64 GetChannelCount() const = 0;
 	virtual const TArray<FChannelEntry>& GetChannels() const = 0;
 	virtual FDateTime GetTimeStamp() const = 0;
 };
 
+TRACESERVICES_API FName GetChannelProviderName();
+TRACESERVICES_API const IChannelProvider* ReadChannelProvider(const IAnalysisSession& Session);
+
 } // namespace TraceServices
diff --git a/Engine/Source/Developer/TraceServices/Public/TraceServices/Model/ContextSwitches.h b/Engine/Source/Developer/TraceServices/Public/TraceServices/Model/ContextSwitches.h
--- a/Engine/Source/Developer/TraceServices/Public/TraceServices/Model/ContextSwitches.h
+++ b/Engine/Source/Developer/TraceServices/Public/TraceServices/Model/ContextSwitches.h
@@ -2,10 +2,10 @@
 
 #pragma once
 
-#include "Common/PagedArray.h"
 #include "HAL/Platform.h"
 #include "Templates/Function.h"
 #include "TraceServices/Model/AnalysisSession.h"
+#include "UObject/NameTypes.h"
 
 namespace TraceServices
 {
@@ -58,6 +58,7 @@ public:
 	virtual void EnumerateContextSwitches(uint32 ThreadId, double StartTime, double EndTime, ContextSwitchCallback Callback) const = 0;
 };
 
+TRACESERVICES_API FName GetContextSwitchesProviderName();
 TRACESERVICES_API const IContextSwitchesProvider* ReadContextSwitchesProvider(const IAnalysisSession& Session);
 
 } // namespace TraceServices
diff --git a/Engine/Source/Developer/TraceServices/Public/TraceServices/Model/CookProfilerProvider.h b/Engine/Source/Developer/TraceServices/Public/TraceServices/Model/CookProfilerProvider.h
--- a/Engine/Source/Developer/TraceServices/Public/TraceServices/Model/CookProfilerProvider.h
+++ b/Engine/Source/Developer/TraceServices/Public/TraceServices/Model/CookProfilerProvider.h
@@ -2,10 +2,10 @@
 
 #pragma once
 
-#include "Common/ProviderLock.h"
 #include "HAL/Platform.h"
 #include "Templates/Function.h"
 #include "TraceServices/Model/AnalysisSession.h"
+#include "UObject/NameTypes.h"
 
 namespace TraceServices
 {
@@ -49,6 +49,7 @@ public:
 	virtual FPackageData* EditPackage(uint64 Id) = 0;
 };
 
+TRACESERVICES_API FName GetCookProfilerProviderName();
 TRACESERVICES_API const ICookProfilerProvider* ReadCookProfilerProvider(const IAnalysisSession& Session);
 TRACESERVICES_API IEditableCookProfilerProvider* EditCookProfilerProvider(IAnalysisSession& Session);
 
diff --git a/Engine/Source/Developer/TraceServices/Public/TraceServices/Model/Counters.h b/Engine/Source/Developer/TraceServices/Public/TraceServices/Model/Counters.h
--- a/Engine/Source/Developer/TraceServices/Public/TraceServices/Model/Counters.h
+++ b/Engine/Source/Developer/TraceServices/Public/TraceServices/Model/Counters.h
@@ -5,8 +5,7 @@
 #include "HAL/Platform.h"
 #include "Templates/Function.h"
 #include "TraceServices/Model/AnalysisSession.h"
-
-template <typename FuncType> class TFunctionRef;
+#include "UObject/NameTypes.h"
 
 namespace TraceServices
 {
@@ -52,49 +51,49 @@ public:
 
 	/*
 	* Sets the name of the counter.
-	* 
+	*
 	* @param Name	The name of the counter.
 	*/
 	virtual void SetName(const TCHAR* Name) = 0;
 
 	/*
 	* Sets the group of the counter.
-	* 
+	*
 	* @param Group	The group name of the counter.
 	*/
 	virtual void SetGroup(const TCHAR* Group) = 0;
 
 	/*
 	* Sets the description of the counter.
-	* 
+	*
 	* @param Description	The description of the counter.
 	*/
 	virtual void SetDescription(const TCHAR* Description) = 0;
 
 	/*
 	* Sets whether the counter is an integral or floating point number.
-	* 
+	*
 	* @param bIsFloatingPoint	True if the counter is floating point, false if it is integral.
 	*/
 	virtual void SetIsFloatingPoint(bool bIsFloatingPoint) = 0;
 
 	/*
 	* Sets whether the counter value is reset every frame. This can be used for counters polled from stats.
-	* 
+	*
 	* @param bInIsResetEveryFrame	True if the counter is reset every frame.
 	*/
 	virtual void SetIsResetEveryFrame(bool bInIsResetEveryFrame) = 0;
 
 	/*
 	* Sets the display hint for this counter.
-	* 
+	*
 	* @param DisplayHint	The display hint for this counter.
 	*/
 	virtual void SetDisplayHint(ECounterDisplayHint DisplayHint) = 0;
 
 	/*
 	* Add a value to the value counter.
-	* 
+	*
 	* @param Time	The time at which the value was added in seconds.
 	* @param Value	The value to add to the current value of the counter.
 	*/
@@ -110,7 +109,7 @@ public:
 
 	/*
 	* Set the value of the counter.
-	* 
+	*
 	* @param Time	The time at which the value was set in seconds.
 	* @param Value	The new value for the counter.
 	*/
@@ -145,7 +144,7 @@ public:
 	virtual ~IEditableCounterProvider() = default;
 
 	/*
-	* Retreive the counter interface to an editable counter.
+	* Retrieve the counter interface to an editable counter.
 	*
 	* @return The interface the counter.
 	*         This can be null if the provider doesn't care about deriving new counters from existing editable counters.
@@ -154,7 +153,7 @@ public:
 
 	/*
 	* Create a new counter mutation interface.
-	* 
+	*
 	* @return The interface to mutate the counter.
 	*/
 	virtual IEditableCounter* CreateEditableCounter() = 0;
@@ -165,6 +164,7 @@ public:
 	virtual void AddCounter(const ICounter* Counter) = 0;
 };
 
+TRACESERVICES_API FName GetCounterProviderName();
 TRACESERVICES_API const ICounterProvider& ReadCounterProvider(const IAnalysisSession& Session);
 TRACESERVICES_API IEditableCounterProvider& EditCounterProvider(IAnalysisSession& Session);
 
diff --git a/Engine/Source/Developer/TraceServices/Public/TraceServices/Model/CsvProfilerProvider.h b/Engine/Source/Developer/TraceServices/Public/TraceServices/Model/CsvProfilerProvider.h
--- a/Engine/Source/Developer/TraceServices/Public/TraceServices/Model/CsvProfilerProvider.h
+++ b/Engine/Source/Developer/TraceServices/Public/TraceServices/Model/CsvProfilerProvider.h
@@ -6,8 +6,7 @@
 #include "HAL/Platform.h"
 #include "Templates/Function.h"
 #include "TraceServices/Model/AnalysisSession.h"
-
-template <typename FuncType> class TFunctionRef;
+#include "UObject/NameTypes.h"
 
 namespace TraceServices
 {
@@ -31,6 +30,7 @@ public:
 	virtual const IUntypedTable& GetTable(uint32 CaptureId) const = 0;
 };
 
-const ICsvProfilerProvider* ReadCsvProfilerProvider(const IAnalysisSession& Session);
+TRACESERVICES_API FName GetCsvProfilerProviderName();
+TRACESERVICES_API const ICsvProfilerProvider* ReadCsvProfilerProvider(const IAnalysisSession& Session);
 
 } // namespace TraceServices
diff --git a/Engine/Source/Developer/TraceServices/Public/TraceServices/Model/Definitions.h b/Engine/Source/Developer/TraceServices/Public/TraceServices/Model/Definitions.h
--- a/Engine/Source/Developer/TraceServices/Public/TraceServices/Model/Definitions.h
+++ b/Engine/Source/Developer/TraceServices/Public/TraceServices/Model/Definitions.h
@@ -1,9 +1,13 @@
 ﻿// Copyright Epic Games, Inc. All Rights Reserved.
+
 #pragma once
 
 #include "Containers/Map.h"
+#include "HAL/Platform.h"
+#include "Templates/Function.h"
 #include "TraceServices/Model/AnalysisSession.h"
-#include "Trace/Analyzer.h"
+#include "Trace/Analyzer.h" // TraceAnalysis
+#include "UObject/NameTypes.h"
 
 namespace TraceServices
 {
@@ -81,7 +85,7 @@ public:
 
 	/**
 	 * Gets the definition as a owned string.
-	 * @param Reference Id used to uniqely identify the definition.
+	 * @param Reference Id used to uniquely identify the definition.
 	 * @return A string describing the referenced value
 	 */
 	template<typename DefinitionType>
@@ -106,6 +110,7 @@ protected:
 	TMap<uint32, StringifierFn> Stringifiers;
 };
 
+TRACESERVICES_API FName GetDefinitionProviderName();
 TRACESERVICES_API const IDefinitionProvider* ReadDefinitionProvider(const IAnalysisSession& Session);
 TRACESERVICES_API IDefinitionProvider* EditDefinitionProvider(IAnalysisSession& Session);
 
diff --git a/Engine/Source/Developer/TraceServices/Public/TraceServices/Model/Frames.h b/Engine/Source/Developer/TraceServices/Public/TraceServices/Model/Frames.h
--- a/Engine/Source/Developer/TraceServices/Public/TraceServices/Model/Frames.h
+++ b/Engine/Source/Developer/TraceServices/Public/TraceServices/Model/Frames.h
@@ -1,6 +1,6 @@
 // Copyright Epic Games, Inc. All Rights Reserved.
 
-#pragma  once
+#pragma once
 
 #include "Containers/ContainersFwd.h"
 #include "Delegates/Delegate.h"
@@ -8,8 +8,7 @@
 #include "ProfilingDebugging/MiscTrace.h"
 #include "Templates/Function.h"
 #include "TraceServices/Model/AnalysisSession.h"
-
-template <typename FuncType> class TFunctionRef;
+#include "UObject/NameTypes.h"
 
 namespace TraceServices
 {
@@ -45,14 +44,14 @@ public:
 	virtual void EnumerateFrames(ETraceFrameType FrameType, uint64 Start, uint64 End, TFunctionRef<void(const FFrame&)> Callback) const = 0;
 	/**
 	* Enumerate frames whose duration intersects [StartTime, EndTime].
-	* 
+	*
 	* @param FrameType	The frame type to enumerate.
 	* @param StartTime	The start timestamp in seconds.
 	* @param EndTime	The end timestamp in seconds.
 	* @param Callback	The callback to be called for each frame.
 	*/
 	virtual void EnumerateFrames(ETraceFrameType FrameType, double StartTime, double EndTime, TFunctionRef<void(const FFrame&)> Callback) const = 0;
-	
+
 	/**
 	 * Gets the array of frame start times, for the specified frame type.
 	 * @param FrameType - type of frames
@@ -88,6 +87,7 @@ public:
 	virtual uint32 GetFrameNumberForTimestamp(ETraceFrameType FrameType, double Time) const = 0;
 };
 
+TRACESERVICES_API FName GetFrameProviderName();
 TRACESERVICES_API const IFrameProvider& ReadFrameProvider(const IAnalysisSession& Session);
 
 } // namespace TraceServices
diff --git a/Engine/Source/Developer/TraceServices/Public/TraceServices/Model/LoadTimeProfiler.h b/Engine/Source/Developer/TraceServices/Public/TraceServices/Model/LoadTimeProfiler.h
--- a/Engine/Source/Developer/TraceServices/Public/TraceServices/Model/LoadTimeProfiler.h
+++ b/Engine/Source/Developer/TraceServices/Public/TraceServices/Model/LoadTimeProfiler.h
@@ -5,14 +5,15 @@
 #include "Containers/Array.h"
 #include "HAL/Platform.h"
 #include "Serialization/LoadTimeTrace.h"
+#include "Templates/Function.h"
 #include "TraceServices/Containers/Tables.h"
 #include "TraceServices/Containers/Timelines.h"
 #include "TraceServices/Model/AnalysisSession.h"
-
-template <typename FuncType> class TFunctionRef;
+#include "UObject/NameTypes.h"
 
 namespace TraceServices
 {
+
 template <typename RowType> class ITable;
 
 struct FFileInfo
@@ -179,7 +180,10 @@ public:
 	virtual const ITable<FLoadRequest>& GetRequestsTable() const = 0;
 };
 
+TRACESERVICES_API FName GetLoadTimeProfilerProviderName();
 TRACESERVICES_API const ILoadTimeProfilerProvider* ReadLoadTimeProfilerProvider(const IAnalysisSession& Session);
+
+TRACESERVICES_API FName GetFileActivityProviderName();
 TRACESERVICES_API const IFileActivityProvider* ReadFileActivityProvider(const IAnalysisSession& Session);
 
 } // namespace TraceServices
diff --git a/Engine/Source/Developer/TraceServices/Public/TraceServices/Model/Log.h b/Engine/Source/Developer/TraceServices/Public/TraceServices/Model/Log.h
--- a/Engine/Source/Developer/TraceServices/Public/TraceServices/Model/Log.h
+++ b/Engine/Source/Developer/TraceServices/Public/TraceServices/Model/Log.h
@@ -2,13 +2,12 @@
 
 #pragma once
 
-#include "CoreTypes.h"
+#include "HAL/Platform.h"
 #include "Logging/LogVerbosity.h"
 #include "Templates/Function.h"
 #include "TraceServices/Containers/Tables.h"
 #include "TraceServices/Model/AnalysisSession.h"
-
-template <typename FuncType> class TFunctionRef;
+#include "UObject/NameTypes.h"
 
 namespace TraceServices
 {
@@ -131,6 +130,7 @@ public:
 	virtual void AppendMessage(uint64 LogPoint, double Time, const TCHAR* Text) = 0;
 };
 
+TRACESERVICES_API FName GetLogProviderName();
 TRACESERVICES_API const ILogProvider& ReadLogProvider(const IAnalysisSession& Session);
 TRACESERVICES_API IEditableLogProvider& EditLogProvider(IAnalysisSession& Session);
 TRACESERVICES_API void FormatString(TCHAR* OutputString, uint32 OutputStringCount, const TCHAR* FormatString, const uint8* FormatArgs);
diff --git a/Engine/Source/Developer/TraceServices/Public/TraceServices/Model/Memory.h b/Engine/Source/Developer/TraceServices/Public/TraceServices/Model/Memory.h
--- a/Engine/Source/Developer/TraceServices/Public/TraceServices/Model/Memory.h
+++ b/Engine/Source/Developer/TraceServices/Public/TraceServices/Model/Memory.h
@@ -9,8 +9,6 @@
 #include "TraceServices/Model/AnalysisSession.h"
 #include "UObject/NameTypes.h"
 
-template <typename FuncType> class TFunctionRef;
-
 namespace TraceServices
 {
 
diff --git a/Engine/Source/Developer/TraceServices/Public/TraceServices/Model/MetadataProvider.h b/Engine/Source/Developer/TraceServices/Public/TraceServices/Model/MetadataProvider.h
--- a/Engine/Source/Developer/TraceServices/Public/TraceServices/Model/MetadataProvider.h
+++ b/Engine/Source/Developer/TraceServices/Public/TraceServices/Model/MetadataProvider.h
@@ -2,11 +2,12 @@
 
 #pragma once
 
-#include "AnalysisSession.h"
-#include "CoreTypes.h"
-#include "UObject/NameTypes.h"
+#include "HAL/Platform.h"
 #include "Internationalization/Text.h"
 #include "Memory/SharedBuffer.h"
+#include "Templates/Function.h"
+#include "TraceServices/Model/AnalysisSession.h"
+#include "UObject/NameTypes.h"
 
 namespace TraceServices
 {
diff --git a/Engine/Source/Developer/TraceServices/Public/TraceServices/Model/Modules.h b/Engine/Source/Developer/TraceServices/Public/TraceServices/Model/Modules.h
--- a/Engine/Source/Developer/TraceServices/Public/TraceServices/Model/Modules.h
+++ b/Engine/Source/Developer/TraceServices/Public/TraceServices/Model/Modules.h
@@ -3,8 +3,12 @@
 #pragma once
 
 #include "Async/TaskGraphInterfaces.h"
-#include "TraceServices/Model/AnalysisSession.h"
+#include "HAL/Platform.h"
+#include "Templates/Function.h"
 #include "Templates/UnrealTemplate.h"
+#include "TraceServices/Model/AnalysisSession.h"
+#include "UObject/NameTypes.h"
+
 #include <atomic>
 
 namespace TraceServices
@@ -17,7 +21,7 @@ namespace TraceServices
   */
 enum class ESymbolQueryResult : uint8
 {
-	Pending,		// Symbol is pending resolution
+	Pending,	// Symbol is pending resolution
 	OK,			// Symbol has been correctly resolved
 	NotLoaded,	// Module debug data could not be loaded or found
 	Mismatch,	// Module debug data could not be loaded because debug data did not match traced binary
diff --git a/Engine/Source/Developer/TraceServices/Public/TraceServices/Model/NetProfiler.h b/Engine/Source/Developer/TraceServices/Public/TraceServices/Model/NetProfiler.h
--- a/Engine/Source/Developer/TraceServices/Public/TraceServices/Model/NetProfiler.h
+++ b/Engine/Source/Developer/TraceServices/Public/TraceServices/Model/NetProfiler.h
@@ -2,10 +2,12 @@
 
 #pragma once
 
-#include "CoreTypes.h"
+#include "HAL/Platform.h"
 #include "Templates/Function.h"
 #include "TraceServices/Containers/Tables.h"
 #include "TraceServices/Model/AnalysisSession.h"
+#include "UObject/NameTypes.h"
+
 #include <limits>
 
 namespace TraceServices
@@ -133,7 +135,7 @@ struct FNetProfilerContentEvent
 	uint64 Level : 8;			// Level
 	uint64 Padding : 8;		// Padding
 
-	FNetProfilerBunchInfo BunchInfo;	
+	FNetProfilerBunchInfo BunchInfo;
 
 	uint32 EventTypeIndex;		// Will replace name index
 	uint32 NameIndex;			// Identify the name / type, should we store the actual Name as well
@@ -301,12 +303,11 @@ public:
 	virtual void ReadNetStatsCounterTypes(TFunctionRef<void(const FNetProfilerStatsCounterType*, uint64)> Callback) const = 0;
 	virtual void ReadNetStatsCounterType(uint32 TypeIndex, TFunctionRef<void(const FNetProfilerStatsCounterType&)> Callback) const = 0;
 
-
 	// Computes aggregated stats for a packet interval or for a range of content events in a single packet.
 	// [PacketIndexIntervalStart, PacketIndexIntervalEnd] is the inclusive packet interval.
 	// [StartPosition, EndPosition) is the exclusive bit range interval; only used when PacketIndexIntervalStart == PacketIndexIntervalEnd.
 	virtual ITable<FNetProfilerAggregatedStats>* CreateAggregation(uint32 ConnectionIndex, ENetProfilerConnectionMode Mode, uint32 PacketIndexIntervalStart, uint32 PacketIndexIntervalEnd, uint32 StartPosition, uint32 EndPosition) const = 0;
-	
+
 	// Computes aggregated statscounters for a packet interval
 	// [PacketIndexIntervalStart, PacketIndexIntervalEnd] is the inclusive packet interval.
 	virtual ITable<FNetProfilerAggregatedStatsCounterStats>* CreateStatsCountersAggregation(uint32 ConnectionIndex, ENetProfilerConnectionMode Mode, uint32 PacketIndexIntervalStart, uint32 PacketIndexIntervalEnd) const = 0;
diff --git a/Engine/Source/Developer/TraceServices/Public/TraceServices/Model/Regions.h b/Engine/Source/Developer/TraceServices/Public/TraceServices/Model/Regions.h
--- a/Engine/Source/Developer/TraceServices/Public/TraceServices/Model/Regions.h
+++ b/Engine/Source/Developer/TraceServices/Public/TraceServices/Model/Regions.h
@@ -6,8 +6,7 @@
 #include "HAL/Platform.h"
 #include "Templates/Function.h"
 #include "TraceServices/Model/AnalysisSession.h"
-
-template <typename FuncType> class TFunctionRef;
+#include "UObject/NameTypes.h"
 
 namespace TraceServices
 {
diff --git a/Engine/Source/Developer/TraceServices/Public/TraceServices/Model/Screenshot.h b/Engine/Source/Developer/TraceServices/Public/TraceServices/Model/Screenshot.h
--- a/Engine/Source/Developer/TraceServices/Public/TraceServices/Model/Screenshot.h
+++ b/Engine/Source/Developer/TraceServices/Public/TraceServices/Model/Screenshot.h
@@ -1,9 +1,11 @@
 // Copyright Epic Games, Inc. All Rights Reserved.
 
-#pragma  once
+#pragma once
 
+#include "Containers/Array.h"
+#include "Templates/SharedPointer.h"
 #include "TraceServices/Model/AnalysisSession.h"
-#include "Templates/Function.h"
+#include "UObject/NameTypes.h"
 
 namespace TraceServices
 {
@@ -28,9 +30,10 @@ class IScreenshotProvider
 public:
 	virtual ~IScreenshotProvider() = default;
 
-	virtual const TSharedPtr<const FScreenshot> GetScreenshot(uint32 Id) const = 0 ;
+	virtual const TSharedPtr<const FScreenshot> GetScreenshot(uint32 Id) const = 0;
 };
 
+TRACESERVICES_API FName GetScreenshotProviderName();
 TRACESERVICES_API const IScreenshotProvider& ReadScreenshotProvider(const IAnalysisSession& Session);
 
 } // namespace TraceServices
diff --git a/Engine/Source/Developer/TraceServices/Public/TraceServices/Model/StackSamples.h b/Engine/Source/Developer/TraceServices/Public/TraceServices/Model/StackSamples.h
--- a/Engine/Source/Developer/TraceServices/Public/TraceServices/Model/StackSamples.h
+++ b/Engine/Source/Developer/TraceServices/Public/TraceServices/Model/StackSamples.h
@@ -2,9 +2,10 @@
 
 #pragma once
 
-#include "Common/PagedArray.h"
+#include "Common/PagedArray.h" // TraceServices
 #include "HAL/Platform.h"
 #include "TraceServices/Model/AnalysisSession.h"
+#include "UObject/NameTypes.h"
 
 namespace TraceServices
 {
@@ -24,6 +25,7 @@ public:
 	virtual const TPagedArray<FStackSample>* GetStackSamples(uint32 ThreadId) const = 0;
 };
 
+TRACESERVICES_API FName GetStackSamplesProviderName();
 TRACESERVICES_API const IStackSamplesProvider& ReadStackSamplesProvider(const IAnalysisSession& Session);
 
 } // namespace TraceServices
diff --git a/Engine/Source/Developer/TraceServices/Public/TraceServices/Model/TasksProfiler.h b/Engine/Source/Developer/TraceServices/Public/TraceServices/Model/TasksProfiler.h
--- a/Engine/Source/Developer/TraceServices/Public/TraceServices/Model/TasksProfiler.h
+++ b/Engine/Source/Developer/TraceServices/Public/TraceServices/Model/TasksProfiler.h
@@ -8,123 +8,127 @@
 #include "HAL/Platform.h"
 #include "Templates/Function.h"
 #include "TraceServices/Model/AnalysisSession.h"
+#include "UObject/NameTypes.h"
 
 namespace TraceServices
 {
-	// task biography
-	struct FTaskInfo
+
+// task biography
+struct FTaskInfo
+{
+	FTaskInfo() = default;
+
+	TaskTrace::FId Id = TaskTrace::InvalidId;
+
+	// this members are filled during task launch
+	const TCHAR* DebugName = nullptr; // user provided or automatically generated (in format `full_path(lineno)`)
+	bool bTracked = false; // if the task is awaitable, not awaitable tasks will never be "completed", they are done when their execution is finished
+	int32 ThreadToExecuteOn = ENamedThreads::AnyThread; // ENamedThreads::Type - combines the info about thread to execute on (one of named threads or a worker
+	// thread - ENamedThreads::AnyThread), task priority, thread priority (in case of a worker thread), and queue index (in case of a named
+	// thread - main queue or local queue). See ENamedThreads to split the info into separate fields
+
+	static constexpr double InvalidTimestamp = 0.0;
+
+	// some tasks can be created with delayed launch, otherwise "created" and "launched" timestamps and thread IDs are equal
+	double CreatedTimestamp = InvalidTimestamp;
+	uint32 CreatedThreadId;
+	// the moment and the place when the task was launched
+	double LaunchedTimestamp = InvalidTimestamp;
+	uint32 LaunchedThreadId;
+	// the moment and the place when all task's dependencies are resolved and it's added to the execution queue
+	double ScheduledTimestamp = InvalidTimestamp;
+	uint32 ScheduledThreadId;
+	// execution started
+	double StartedTimestamp = InvalidTimestamp;
+	uint32 StartedThreadId;
+	// execution finished, but the task can be still not completed waiting for all nested tasks to complete
+	double FinishedTimestamp = InvalidTimestamp;
+	// execution done and all nested tasks are completed
+	double CompletedTimestamp = InvalidTimestamp;
+	uint32 CompletedThreadId;
+	// the last reference is released and the task is destroyed
+	double DestroyedTimestamp = InvalidTimestamp;
+	uint32 DestroyedThreadId;
+	// task size including user-provided task body
+	uint64 TaskSize = 0;
+
+	// relation with another task, when and where it was established
+	struct FRelationInfo
 	{
-		FTaskInfo() = default;
-
-		TaskTrace::FId Id = TaskTrace::InvalidId;
-
-		// this members are filled during task launch
-		const TCHAR* DebugName = nullptr; // user provided or automatically generated (in format `full_path(lineno)`)
-		bool bTracked = false; // if the task is awaitable, not awaitable tasks will never be "completed", they are done when their execution is finished
-		int32 ThreadToExecuteOn = ENamedThreads::AnyThread; // ENamedThreads::Type - combines the info about thread to execute on (one of named threads or a worker
-		// thread - ENamedThreads::AnyThread), task priority, thread priority (in case of a worker thread), and queue index (in case of a named
-		// thread - main queue or local queue). See ENamedThreads to split the info into separate fields
-
-		static constexpr double InvalidTimestamp = 0.0;
-
-		// some tasks can be created with delayed launch, otherwise "created" and "launched" timestamps and thread IDs are equal
-		double CreatedTimestamp = InvalidTimestamp;
-		uint32 CreatedThreadId;
-		// the moment and the place when the task was launched
-		double LaunchedTimestamp = InvalidTimestamp;
-		uint32 LaunchedThreadId;
-		// the moment and the place when all task's dependencies are resolved and it's added to the execution queue
-		double ScheduledTimestamp = InvalidTimestamp;
-		uint32 ScheduledThreadId;
-		// execution started
-		double StartedTimestamp = InvalidTimestamp;
-		uint32 StartedThreadId;
-		// execution finished, but the task can be still not completed waiting for all nested tasks to complete
-		double FinishedTimestamp = InvalidTimestamp;
-		// execution done and all nested tasks are completed
-		double CompletedTimestamp = InvalidTimestamp;
-		uint32 CompletedThreadId;
-		// the last reference is released and the task is destroyed
-		double DestroyedTimestamp = InvalidTimestamp;
-		uint32 DestroyedThreadId;
-		// task size including user-provided task body
-		uint64 TaskSize = 0;
-
-		//  relation with another task, when and where it was established
-		struct FRelationInfo
+		FRelationInfo(TaskTrace::FId InRelativeId, double InTimestamp, uint32 InThreadId)
+			: RelativeId(InRelativeId)
+			, Timestamp(InTimestamp)
+			, ThreadId(InThreadId)
 		{
-			FRelationInfo(TaskTrace::FId InRelativeId, double InTimestamp, uint32 InThreadId)
-				: RelativeId(InRelativeId)
-				, Timestamp(InTimestamp)
-				, ThreadId(InThreadId)
-			{
-			}
-
-			TaskTrace::FId RelativeId;
-			double Timestamp;
-			uint32 ThreadId;
-		};
-
-		// other tasks that must complete before this taks can be scheduled
-		TArray<FRelationInfo> Prerequisites;
-		// other tasks that wait for this task completion before they'll start execution
-		TArray<FRelationInfo> Subsequents;
-		// other tasks that have this task as a nested
-		TArray<FRelationInfo> ParentTasks;
-		// the task is completed only after all nested tasks are completed
-		TArray<FRelationInfo> NestedTasks;
-	};
+		}
 
-	struct FWaitingForTasks
-	{
-		TArray<TaskTrace::FId> Tasks;
-		double StartedTimestamp;
-		double FinishedTimestamp = FTaskInfo::InvalidTimestamp;
+		TaskTrace::FId RelativeId;
+		double Timestamp;
+		uint32 ThreadId;
 	};
 
-	enum class ETaskEnumerationOption
-	{
-		Alive,
-		Active,
-		WaitingForPrerequisites,
-		Queued,
-		Executing,
-		WaitingForNested,
-		Completed,
-		Count
-	};
+	// other tasks that must complete before this task can be scheduled
+	TArray<FRelationInfo> Prerequisites;
+	// other tasks that wait for this task completion before they'll start execution
+	TArray<FRelationInfo> Subsequents;
+	// other tasks that have this task as a nested
+	TArray<FRelationInfo> ParentTasks;
+	// the task is completed only after all nested tasks are completed
+	TArray<FRelationInfo> NestedTasks;
+};
+
+struct FWaitingForTasks
+{
+	TArray<TaskTrace::FId> Tasks;
+	double StartedTimestamp;
+	double FinishedTimestamp = FTaskInfo::InvalidTimestamp;
+};
 
-	enum class ETaskEnumerationResult
-	{
-		Continue,
-		Stop,
-	};
+enum class ETaskEnumerationOption
+{
+	Alive,
+	Active,
+	WaitingForPrerequisites,
+	Queued,
+	Executing,
+	WaitingForNested,
+	Completed,
+	Count
+};
+
+enum class ETaskEnumerationResult
+{
+	Continue,
+	Stop,
+};
 
-	// query interface to tasks info
-	class ITasksProvider : public IProvider
-	{
-	public:
-		typedef TFunctionRef<ETaskEnumerationResult(const FTaskInfo& Task)> TaskCallback;
+// query interface to tasks info
+class ITasksProvider : public IProvider
+{
+public:
+	typedef TFunctionRef<ETaskEnumerationResult(const FTaskInfo& Task)> TaskCallback;
 
-		virtual ~ITasksProvider() = default;
+	virtual ~ITasksProvider() = default;
 
-		// returns a task for given thread and timestamp, if any, otherwise `nullptr`
-		virtual const FTaskInfo* TryGetTask(uint32 ThreadId, double Timestamp) const = 0;
-		// returns task info for given task ID
-		virtual const FTaskInfo* TryGetTask(TaskTrace::FId TaskId) const = 0;
+	// returns a task for given thread and timestamp, if any, otherwise `nullptr`
+	virtual const FTaskInfo* TryGetTask(uint32 ThreadId, double Timestamp) const = 0;
+	// returns task info for given task ID
+	virtual const FTaskInfo* TryGetTask(TaskTrace::FId TaskId) const = 0;
 
-		// returns an info about waiting for tasks completion for given thread and timestamp, if any, otherwise `nullptr`
-		virtual const FWaitingForTasks* TryGetWaiting(const TCHAR* TimerName, uint32 ThreadId, double Timestamp) const = 0;
+	// returns an info about waiting for tasks completion for given thread and timestamp, if any, otherwise `nullptr`
+	virtual const FWaitingForTasks* TryGetWaiting(const TCHAR* TimerName, uint32 ThreadId, double Timestamp) const = 0;
 
-		// returns an info about ParallelFor tasks for given timer
-		virtual TArray<TaskTrace::FId> TryGetParallelForTasks(const TCHAR* TimerName, uint32 ThreadId, double StartTime, double EndTime) const = 0;
+	// returns an info about ParallelFor tasks for given timer
+	virtual TArray<TaskTrace::FId> TryGetParallelForTasks(const TCHAR* TimerName, uint32 ThreadId, double StartTime, double EndTime) const = 0;
 
-		// returns the number of tasks stored in the provider
-		virtual int64 GetNumTasks() const = 0;
+	// returns the number of tasks stored in the provider
+	virtual int64 GetNumTasks() const = 0;
 
-		// Calls the callback for each task stored in the provider with CreatedTimestamp <= EndTime and FinishedTimestamp >= StartTime
-		virtual void EnumerateTasks(double StartTime, double EndTime, ETaskEnumerationOption EnumerationOption, TaskCallback Callback) const = 0;
-	};
+	// Calls the callback for each task stored in the provider with CreatedTimestamp <= EndTime and FinishedTimestamp >= StartTime
+	virtual void EnumerateTasks(double StartTime, double EndTime, ETaskEnumerationOption EnumerationOption, TaskCallback Callback) const = 0;
+};
+
+TRACESERVICES_API FName GetTaskProviderName();
+TRACESERVICES_API const ITasksProvider* ReadTasksProvider(const IAnalysisSession& Session);
 
-	TRACESERVICES_API const ITasksProvider* ReadTasksProvider(const IAnalysisSession& Session);
-}
\ No newline at end of file
+} // namespace TraceServices
diff --git a/Engine/Source/Developer/TraceServices/Public/TraceServices/Model/Threads.h b/Engine/Source/Developer/TraceServices/Public/TraceServices/Model/Threads.h
--- a/Engine/Source/Developer/TraceServices/Public/TraceServices/Model/Threads.h
+++ b/Engine/Source/Developer/TraceServices/Public/TraceServices/Model/Threads.h
@@ -6,8 +6,7 @@
 #include "HAL/PlatformAffinity.h"
 #include "Templates/Function.h"
 #include "TraceServices/Model/AnalysisSession.h"
-
-template <typename FuncType> class TFunctionRef;
+#include "UObject/NameTypes.h"
 
 namespace TraceServices
 {
@@ -40,7 +39,7 @@ public:
 
 	/*
 	* Note the existence of a new thread.
-	* 
+	*
 	* @param Id			The thread identity.
 	* @param Name		The name the user may know the thread by, if available.
 	* @param Priority	The system priority level of the thread, if available.
@@ -48,6 +47,7 @@ public:
 	virtual void AddThread(uint32 Id, const TCHAR* Name, EThreadPriority Priority) = 0;
 };
 
+TRACESERVICES_API FName GetThreadProviderName();
 TRACESERVICES_API const IThreadProvider& ReadThreadProvider(const IAnalysisSession& Session);
 TRACESERVICES_API IEditableThreadProvider& EditThreadProvider(IAnalysisSession& Session);
 
diff --git a/Engine/Source/Developer/TraceServices/Public/TraceServices/Model/TimingProfiler.h b/Engine/Source/Developer/TraceServices/Public/TraceServices/Model/TimingProfiler.h
--- a/Engine/Source/Developer/TraceServices/Public/TraceServices/Model/TimingProfiler.h
+++ b/Engine/Source/Developer/TraceServices/Public/TraceServices/Model/TimingProfiler.h
@@ -5,13 +5,13 @@
 #include "Containers/Array.h"
 #include "Containers/ArrayView.h"
 #include "Containers/StringFwd.h"
+#include "HAL/Platform.h"
 #include "HAL/PlatformMath.h"
 #include "Model/MonotonicTimeline.h"
 #include "TraceServices/Containers/Tables.h"
 #include "TraceServices/Containers/Timelines.h"
 #include "TraceServices/Model/AnalysisSession.h"
-
-template <typename FuncType> class TFunctionRef;
+#include "UObject/NameTypes.h"
 
 namespace TraceServices
 {
@@ -120,18 +120,18 @@ public:
 
 	/*
 	* A new CPU timer object has been found.
-	* 
+	*
 	* @param Name	The name attached to the timer.
 	* @param File	The source file in which the timer is defined.
 	* @param Line	The line number of the source file in which the timer is defined.
-	* 
+	*
 	* @return The identity of the CPU timer object.
 	*/
 	virtual uint32 AddCpuTimer(FStringView Name, const TCHAR* File, uint32 Line) = 0;
 
 	/*
 	* Update an existing timer with information. Some information is unavailable when it's created.
-	* 
+	*
 	* @param TimerId	The identity of the timer to update.
 	* @param Name		The name attached to the timer.
 	* @param File		The source file in which the timer is defined.
@@ -141,33 +141,34 @@ public:
 
 	/*
 	* Add metadata to a timer.
-	* 
+	*
 	* @param OriginalTimerId	The identity of the timer to add metadata to.
 	* @param Metadata			The metadata.
-	* 
+	*
 	* @return The identity of the metadata.
 	*/
 	virtual uint32 AddMetadata(uint32 OriginalTimerId, TArray<uint8>&& Metadata) = 0;
 
 	/*
 	* Get metadata for a timer.
-	* 
+	*
 	* @param TimerId	The identity of the metadata to retrieve.
-	* 
+	*
 	* @return The metadata.
 	*/
 	virtual TArrayView<const uint8> GetMetadata(uint32 TimerId) const = 0;
 
 	/*
 	* Get an object to receive ordered timeline events for a thread.
-	* 
+	*
 	* @param ThreadId	The thread for which the events are for.
-	* 
+	*
 	* @return The object to receive the serial events for the specified thread.
 	*/
 	virtual IEditableTimeline<FTimingProfilerEvent>& GetCpuThreadEditableTimeline(uint32 ThreadId) = 0;
 };
 
+TRACESERVICES_API FName GetTimingProfilerProviderName();
 TRACESERVICES_API const ITimingProfilerProvider* ReadTimingProfilerProvider(const IAnalysisSession& Session);
 TRACESERVICES_API IEditableTimingProfilerProvider* EditTimingProfilerProvider(IAnalysisSession& Session);
 
diff --git a/Engine/Source/Developer/TraceServices/Public/TraceServices/ModuleService.h b/Engine/Source/Developer/TraceServices/Public/TraceServices/ModuleService.h
--- a/Engine/Source/Developer/TraceServices/Public/TraceServices/ModuleService.h
+++ b/Engine/Source/Developer/TraceServices/Public/TraceServices/ModuleService.h
@@ -8,8 +8,6 @@
 #include "HAL/Platform.h"
 #include "UObject/NameTypes.h"
 
-class FName;
-
 namespace TraceServices
 {
 
@@ -33,13 +31,12 @@ class IModule
 	: public IModularFeature
 {
 public:
-	virtual bool ShouldBeEnabledByDefault() const { return true; }
 	virtual void GetModuleInfo(FModuleInfo& OutModuleInfo) = 0;
-	virtual void OnAnalysisBegin(IAnalysisSession& Session) = 0;
-	virtual void GenerateReports(const IAnalysisSession& Session, const TCHAR* CmdLine, const TCHAR* OutputDirectory) {}
-
+	virtual bool ShouldBeEnabledByDefault() const { return true; }
 	virtual void GetLoggers(TArray<const TCHAR*>& OutLoggers) {}
 	virtual const TCHAR* GetCommandLineArgument() { return nullptr; }
+	virtual void OnAnalysisBegin(IAnalysisSession& Session) = 0;
+	virtual void GenerateReports(const IAnalysisSession& Session, const TCHAR* CmdLine, const TCHAR* OutputDirectory) {}
 };
 
 class IModuleService
