Use AppName to rename eboot.bin

* Use raw AppName to attempt to find "eboot.bin" source binary. This is the canonical way to find it.
* Rename variable to AppName to better reflect what it is.

diff --git a/Engine/Source/Developer/TraceServices/Private/Model/SymslibResolver.cpp b/Engine/Source/Developer/TraceServices/Private/Model/SymslibResolver.cpp
--- a/Engine/Source/Developer/TraceServices/Private/Model/SymslibResolver.cpp
+++ b/Engine/Source/Developer/TraceServices/Private/Model/SymslibResolver.cpp
@@ -219,22 +219,22 @@ namespace
 		return FString();
 	}
 
-	static bool LoadBinary(const TCHAR* Path, SYMS_Arena* Arena, SYMS_ParseBundle& Bundle, TArray<FAutoMappedFile>& Files, const FString& SearchPath, const FString& Platform, const FString& Project)
+	static bool LoadBinary(const TCHAR* Path, SYMS_Arena* Arena, SYMS_ParseBundle& Bundle, TArray<FAutoMappedFile>& Files, const FString& SearchPath, const FString& Platform, const FString& AppName)
 	{
 		FString FilePath = Path;
 		bool bFileFound = false;
 
 		// Remap known renamed binaries
-		if (FilePath.EndsWith(TEXT("eboot.bin")) && !Platform.IsEmpty() && !Project.IsEmpty())
+		if (FilePath.EndsWith(TEXT("eboot.bin")) && !Platform.IsEmpty() && !AppName.IsEmpty())
 		{
-			FilePath = FPaths::Combine(FPaths::GetPath(FilePath), FString::Printf(TEXT("%sGame.self"), *Project));
+			FilePath = FPaths::Combine(FPaths::GetPath(FilePath), FString::Printf(TEXT("%s.self"), *AppName));
 		}
 
 		// First lookup file in symbol path
 		FString BinaryPath = FindBinaryFileInPath(FilePath, SearchPath);
 		if (BinaryPath.IsEmpty() && !Platform.IsEmpty())
 		{
-			BinaryPath = FindFileInEngineFolder(FilePath, SearchPath, Platform, Project);
+			BinaryPath = FindFileInEngineFolder(FilePath, SearchPath, Platform, AppName);
 		}
 		bFileFound = !BinaryPath.IsEmpty();
 
@@ -533,7 +533,7 @@ FSymslibResolver::FSymslibResolver(IAnalysisSession& InSession, IResolvedSymbolF
 	{
 		const FSessionInfo Info = DiagnosticsProvider->GetSessionInfo();
 		Platform = Info.Platform;
-		Project = Info.AppName;
+		AppName = Info.AppName;
 	}
 }
 
@@ -882,7 +882,7 @@ EModuleStatus FSymslibResolver::LoadModule(FModuleEntry* Entry, FStringView Sear
 	// contents of binary & debug file
 	SYMS_ParseBundle Bundle;
 
-	if (!LoadBinary(Entry->Module->FullName, Group->arena, Bundle, Files, SearchPath, Platform, Project))
+	if (!LoadBinary(Entry->Module->FullName, Group->arena, Bundle, Files, SearchPath, Platform, AppName))
 	{
 		syms_group_release(Group);
 		OutStatusMessage.Appendf(TEXT("Failed to load image for '%s' in '%s'."), Entry->Module->Name, *SearchPath);
diff --git a/Engine/Source/Developer/TraceServices/Private/Model/SymslibResolver.h b/Engine/Source/Developer/TraceServices/Private/Model/SymslibResolver.h
--- a/Engine/Source/Developer/TraceServices/Private/Model/SymslibResolver.h
+++ b/Engine/Source/Developer/TraceServices/Private/Model/SymslibResolver.h
@@ -123,7 +123,7 @@ private:
 	IAnalysisSession& Session;
 	IResolvedSymbolFilter& SymbolFilter;
 	FString Platform;
-	FString Project;
+	FString AppName;
 };
 
 /////////////////////////////////////////////////////////////////////
