Network insights
- Traced events associated with a NetId (NetRefHandle or NetGUID) now displays either Pending (NetId:NN) until NetId has a known name.

diff --git a/Engine/Source/Developer/TraceInsights/Private/Insights/NetworkingProfiler/Widgets/SPacketContentView.cpp b/Engine/Source/Developer/TraceInsights/Private/Insights/NetworkingProfiler/Widgets/SPacketContentView.cpp
--- a/Engine/Source/Developer/TraceInsights/Private/Insights/NetworkingProfiler/Widgets/SPacketContentView.cpp
+++ b/Engine/Source/Developer/TraceInsights/Private/Insights/NetworkingProfiler/Widgets/SPacketContentView.cpp
@@ -1050,20 +1050,23 @@ void SPacketContentView::UpdateState(float FontScale)
 				NetProfilerProvider->EnumeratePacketContentEventsByPosition(ConnectionIndex, ConnectionMode, PacketIndex, StartPos, EndPos, [this, &Builder, &FilteredDrawStateBuilder, NetProfilerProvider](const TraceServices::FNetProfilerContentEvent& Event)
 				{
 					const TCHAR* Name = nullptr;
-					NetProfilerProvider->ReadName(Event.NameIndex, [&Name](const TraceServices::FNetProfilerName& NetProfilerName)
-					{
-						Name = NetProfilerName.Name;
-					});
 
+					uint32 NameIndex = Event.NameIndex;
 					uint64 NetId = 0;
 					if (Event.ObjectInstanceIndex != 0)
 					{
-						NetProfilerProvider->ReadObject(GameInstanceIndex, Event.ObjectInstanceIndex, [&NetId](const TraceServices::FNetProfilerObjectInstance& ObjectInstance)
+						NetProfilerProvider->ReadObject(GameInstanceIndex, Event.ObjectInstanceIndex, [&NetId, &NameIndex](const TraceServices::FNetProfilerObjectInstance& ObjectInstance)
 						{
+							NameIndex = ObjectInstance.NameIndex;
 							NetId = ObjectInstance.NetObjectId;
 						});
 					}
 
+					NetProfilerProvider->ReadName(NameIndex, [&Name](const TraceServices::FNetProfilerName& NetProfilerName)
+					{
+						Name = NetProfilerName.Name;
+					});
+
 					Builder.AddEvent(Event, Name, NetId);
 
 					if ((!bFilterByEventType || FilterEventTypeIndex == Event.EventTypeIndex) &&
diff --git a/Engine/Source/Developer/TraceServices/Private/Analyzers/NetTraceAnalyzer.cpp b/Engine/Source/Developer/TraceServices/Private/Analyzers/NetTraceAnalyzer.cpp
--- a/Engine/Source/Developer/TraceServices/Private/Analyzers/NetTraceAnalyzer.cpp
+++ b/Engine/Source/Developer/TraceServices/Private/Analyzers/NetTraceAnalyzer.cpp
@@ -59,6 +59,7 @@ void FNetTraceAnalyzer::OnAnalysisBegin(const FOnAnalysisContext& Context)
 	{
 		FAnalysisSessionEditScope _(Session);
 		BunchHeaderNameIndex = NetProfilerProvider.AddNetProfilerName(TEXT("BunchHeader"));
+		PendingNameIndex = NetProfilerProvider.AddNetProfilerName(TEXT("Pending"));
 	}
 }
 
@@ -277,6 +278,23 @@ void FNetTraceAnalyzer::HandlePacketContentEvent(const FOnEventContext& Context,
 						Event.NameIndex = ActiveObjectState->NameIndex;
 						Event.ObjectInstanceIndex = ActiveObjectState->ObjectIndex;
 					}
+					else if (DecodedNameOrObjectId != 0)
+					{
+						// Sometime we report data for objects that are still pending creation, which we will update as soon as we have more data.
+						FNetProfilerObjectInstance& ObjectInstance = NetProfilerProvider.CreateObject(GameInstanceState->GameInstanceIndex);
+
+						// Fill in the object data we currently have
+						ObjectInstance.LifeTime.Begin = GetLastTimestamp();
+						ObjectInstance.NameIndex = PendingNameIndex;
+						ObjectInstance.NetId = DecodedNameOrObjectId;
+						ObjectInstance.TypeId = 0;
+
+						// Add to active objects
+						GameInstanceState->ActiveObjects.Add(DecodedNameOrObjectId, { ObjectInstance.ObjectIndex, ObjectInstance.NameIndex });
+						
+						Event.NameIndex = ObjectInstance.NameIndex;
+						Event.ObjectInstanceIndex = ObjectInstance.ObjectIndex;
+					}
 				}
 				else if (DecodedEventType == EContentEventType::NameId)
 				{
@@ -778,17 +796,21 @@ void FNetTraceAnalyzer::HandleObjectCreatedEvent(const FOnEventContext& Context,
 	const uint32* NetProfilerNameIndex = TracedNameIdToNetProfilerNameIdMap.Find(NameId);
 	const uint32 NameIndex = NetProfilerNameIndex ? *NetProfilerNameIndex : 0u;
 
-	if (GameInstanceState->ActiveObjects.Contains(ObjectId))
+	if (FNetTraceActiveObjectState* ActiveObjectInstance = GameInstanceState->ActiveObjects.Find(ObjectId))
 	{
-		if (FNetProfilerObjectInstance* ExistingInstance = NetProfilerProvider.EditObject(GameInstanceState->GameInstanceIndex, GameInstanceState->ActiveObjects[ObjectId].ObjectIndex))
+		if (FNetProfilerObjectInstance* ExistingInstance = NetProfilerProvider.EditObject(GameInstanceState->GameInstanceIndex, ActiveObjectInstance->ObjectIndex))
 		{
-			if (ExistingInstance->NameIndex == NameIndex)
+			if (ExistingInstance->NameIndex == NameIndex || ExistingInstance->NameIndex == PendingNameIndex)
 			{
 				// Update existing object instance
 				ExistingInstance->LifeTime.Begin = GetLastTimestamp();
 				ExistingInstance->NetObjectId = ObjectId;
 				ExistingInstance->TypeId = TypeId;
 
+				// Update name in both the persistent instance and the active one
+				ExistingInstance->NameIndex = NameIndex;
+				ActiveObjectInstance->NameIndex = NameIndex;
+
 				return;
 			}
 
diff --git a/Engine/Source/Developer/TraceServices/Private/Analyzers/NetTraceAnalyzer.h b/Engine/Source/Developer/TraceServices/Private/Analyzers/NetTraceAnalyzer.h
--- a/Engine/Source/Developer/TraceServices/Private/Analyzers/NetTraceAnalyzer.h
+++ b/Engine/Source/Developer/TraceServices/Private/Analyzers/NetTraceAnalyzer.h
@@ -62,6 +62,7 @@ private:
 
 	// Default names
 	uint32 BunchHeaderNameIndex;
+	uint32 PendingNameIndex;
 
 	// Shared for trace
 	TMap<uint16, uint32> TracedNameIdToNetProfilerNameIdMap;
