Fixed merge issue from

Network insights
- Traced events associated with a NetId (NetRefHandle or NetGUID) now displays either Pending (NetId:NN) until NetId has a known name.

diff --git a/Engine/Source/Developer/TraceServices/Private/Analyzers/NetTraceAnalyzer.cpp b/Engine/Source/Developer/TraceServices/Private/Analyzers/NetTraceAnalyzer.cpp
--- a/Engine/Source/Developer/TraceServices/Private/Analyzers/NetTraceAnalyzer.cpp
+++ b/Engine/Source/Developer/TraceServices/Private/Analyzers/NetTraceAnalyzer.cpp
@@ -286,7 +286,7 @@ void FNetTraceAnalyzer::HandlePacketContentEvent(const FOnEventContext& Context,
 						// Fill in the object data we currently have
 						ObjectInstance.LifeTime.Begin = GetLastTimestamp();
 						ObjectInstance.NameIndex = PendingNameIndex;
-						ObjectInstance.NetId = DecodedNameOrObjectId;
+						ObjectInstance.NetObjectId = DecodedNameOrObjectId;
 						ObjectInstance.TypeId = 0;
 
 						// Add to active objects
