[Insights] Improve frame stats aggregation performance

diff --git a/Engine/Source/Developer/TraceServices/Private/Common/TimelineStatistics.h b/Engine/Source/Developer/TraceServices/Private/Common/TimelineStatistics.h
--- a/Engine/Source/Developer/TraceServices/Private/Common/TimelineStatistics.h
+++ b/Engine/Source/Developer/TraceServices/Private/Common/TimelineStatistics.h
@@ -3,8 +3,11 @@
 #pragma once
 
 #include "Containers/Map.h"
+#include "Containers/Queue.h"
 #include "CoreTypes.h"
 #include <cfloat>
+#include "Misc/QueuedThreadPool.h"
+#include "Tasks/Task.h"
 #include "Templates/SharedPointer.h"
 #include "TraceServices/Common/CancellationToken.h"
 #include "TraceServices/Containers/Timelines.h"
@@ -52,7 +55,7 @@ public:
 			{
 				return;
 			}
-			ProcessTimeline(Timeline, BucketMapper, UpdateTotalMinMaxTimerStats, IntervalStart, IntervalEnd, InternalResult);
+			ProcessTimeline(Timeline, BucketMapper, UpdateTotalMinMaxTimerAggregationStats, IntervalStart, IntervalEnd, InternalResult);
 		}
 		// Now, as we know min/max inclusive/exclusive times for each timer, we can compute histogram and median values.
 		const bool ComputeMedian = true;
@@ -91,8 +94,8 @@ public:
 											TSharedPtr<FCancellationToken> CancellationToken,
 											TMap<BucketKeyType, FAggregatedTimingStats>& Result)
 	{
-		TMap<BucketKeyType, FInternalAggregationEntry> FrameResult;
 		TMap<BucketKeyType, FInternalFrameAggregationEntry> GlobalResult;
+		TMap<BucketKeyType, FAggregatedTimingStats> InitialTimerMap;
 		double StartTime = Frames[0].StartTime;
 		double EndTime = Frames[Frames.Num() - 1].EndTime;
 
@@ -100,65 +103,153 @@ public:
 		// This is to guarantee that GlobalResult and FrameResult have the same iteration order.
 		for (const TimelineType* Timeline : Timelines)
 		{
-			GatherKeysFromTimeline(Timeline, BucketMapper, StartTime, EndTime, FrameResult);
+			GatherKeysFromTimeline(Timeline, BucketMapper, StartTime, EndTime, InitialTimerMap);
 		}
 
 		int32 FramesNum = Frames.Num();
 
 		// For very large numbers of timers or frames, an out of memory is possible so don't compute the median.
 		constexpr int64 MaxSize = 5 * 100 * 1000 * 1000;
-		bool bComputeMedian = FrameResult.Num() * FramesNum < MaxSize;
-		for (auto& KV : FrameResult)
+		bool bComputeMedian = InitialTimerMap.Num() * FramesNum < MaxSize;
+		for (auto& KV : InitialTimerMap)
 		{
 			FInternalFrameAggregationEntry Entry;
 			if (bComputeMedian)
 			{
-				Entry.FrameInclusiveTimes.Reserve(Frames.Num());
-				Entry.FrameExclusiveTimes.Reserve(Frames.Num());
+				Entry.FrameInclusiveTimes.AddUninitialized(Frames.Num());
+				Entry.FrameExclusiveTimes.AddUninitialized(Frames.Num());
 			}
-			Entry.Inner = KV.Value;
+			Entry.Inner.Stats = KV.Value;
 			GlobalResult.Add(KV.Key, Entry);
 		}
 
-		for (int32 Index = 0; Index < FramesNum; ++Index)
+		std::atomic<int> NextFrameToProcess(0);
+		TQueue<TSharedPtr<TMap<BucketKeyType, FAggregatedTimingStats>>, EQueueMode::Mpsc> FrameResultsQueue;
+
+		constexpr float RatioOfThreadsToUse = 0.75;
+		int32 NumTasks = FMath::Max(1, GThreadPool->GetNumThreads() * RatioOfThreadsToUse);
+		NumTasks = FMath::Min((int32)NumTasks, Frames.Num());
+		int32 FramesPerTask = Frames.Num() / NumTasks;
+		int32 ExtraFrameTasks = Frames.Num() - NumTasks * FramesPerTask;
+
+		int32 StartIndex = 0;
+		int32 EndIndex = FramesPerTask + (ExtraFrameTasks > 0 ? 1 : 0);
+		--ExtraFrameTasks;
+
+		TArray<UE::Tasks::TTask<void>> AsyncTasks;
+		for (int32 Index = 0; Index < NumTasks; ++Index)
+		{
+			AsyncTasks.Add(UE::Tasks::Launch(UE_SOURCE_LOCATION, [StartIndex, EndIndex, &Timelines, bComputeMedian, BucketMapper, &Frames, &FrameResultsQueue, &NextFrameToProcess, &InitialTimerMap, &GlobalResult, CancellationToken]()
+			{
+				if (CancellationToken.IsValid() && CancellationToken->ShouldCancel())
+				{
+					return;
+				}
+
+				TSharedPtr<TMap<BucketKeyType, FAggregatedTimingStats>> TaskResult = MakeShared<TMap<BucketKeyType, FAggregatedTimingStats>>(InitialTimerMap);
+
+				for(int FrameIndex = StartIndex; FrameIndex < EndIndex; ++FrameIndex)
+				{
+					if (CancellationToken.IsValid() && CancellationToken->ShouldCancel())
+					{
+						return;
+					}
+
+					TSharedPtr<TMap<BucketKeyType, FAggregatedTimingStats>> FrameResult = MakeShared<TMap<BucketKeyType, FAggregatedTimingStats>>(InitialTimerMap);
+
+					// Compute instance count and total/min/max inclusive/exclusive times for each timer.
+					for (const TimelineType* Timeline : Timelines)
+					{
+						ProcessTimelineForFrameStats(Timeline, BucketMapper, UpdateTotalMinMaxTimerStats, Frames[FrameIndex].StartTime, Frames[FrameIndex].EndTime, *FrameResult);
+					}
+
+					typename TMap<BucketKeyType, FAggregatedTimingStats>::TIterator FrameResultIterator(*FrameResult);
+					typename TMap<BucketKeyType, FAggregatedTimingStats>::TIterator TaskResultIterator(*TaskResult);
+					typename TMap<BucketKeyType, FInternalFrameAggregationEntry>::TIterator GlobalResultIterator(GlobalResult);
+					while (FrameResultIterator)
+					{
+						FAggregatedTimingStats& FrameStats = FrameResultIterator->Value;
+						FAggregatedTimingStats& TaskResultStats = TaskResultIterator->Value;
+
+						// TODO: Add Instance Count per frame column
+						// ResultStats.InstanceCount = (ResultStats.InstanceCount * Index + FrameStats.InstanceCount) / (Index + 1);
+
+						TaskResultStats.InstanceCount += FrameStats.InstanceCount;
+
+						TaskResultStats.TotalInclusiveTime += FrameStats.TotalInclusiveTime;
+						TaskResultStats.MinInclusiveTime = FMath::Min(TaskResultStats.MinInclusiveTime, FrameStats.TotalInclusiveTime);
+						TaskResultStats.MaxInclusiveTime = FMath::Max(TaskResultStats.MaxInclusiveTime, FrameStats.TotalInclusiveTime);
+
+						TaskResultStats.TotalExclusiveTime += FrameStats.TotalExclusiveTime;
+						TaskResultStats.MinExclusiveTime = FMath::Min(TaskResultStats.MinExclusiveTime, FrameStats.TotalExclusiveTime);
+						TaskResultStats.MaxExclusiveTime = FMath::Max(TaskResultStats.MaxExclusiveTime, FrameStats.TotalExclusiveTime);
+
+						if (bComputeMedian)
+						{
+							GlobalResultIterator->Value.FrameInclusiveTimes[FrameIndex] = FrameStats.TotalInclusiveTime;
+							GlobalResultIterator->Value.FrameExclusiveTimes[FrameIndex] = FrameStats.TotalExclusiveTime;
+						}
+
+						// Reset the per frame stats.
+						FrameStats = FAggregatedTimingStats();
+
+						++FrameResultIterator;
+						++TaskResultIterator;
+						++GlobalResultIterator;
+					}
+
+				}
+
+				FrameResultsQueue.Enqueue(TaskResult);
+			}));
+
+			StartIndex = EndIndex;
+			EndIndex = StartIndex + FramesPerTask + (ExtraFrameTasks > 0 ? 1 : 0);
+			--ExtraFrameTasks;
+		}
+
+		check(StartIndex == Frames.Num());
+
+		int ProcessedResults = 0;
+		while (ProcessedResults < NumTasks)
 		{
 			if (CancellationToken.IsValid() && CancellationToken->ShouldCancel())
 			{
+				UE::Tasks::Wait(AsyncTasks);
 				return;
 			}
 
-			// Compute instance count and total/min/max inclusive/exclusive times for each timer.
-			for (const TimelineType* Timeline : Timelines)
+			if (FrameResultsQueue.IsEmpty())
 			{
-				ProcessTimeline(Timeline, BucketMapper, UpdateTotalMinMaxTimerStats, Frames[Index].StartTime, Frames[Index].EndTime, FrameResult);
+				FPlatformProcess::SleepNoStats(0.1f);
+				continue;
 			}
 
-			typename TMap<BucketKeyType, FInternalAggregationEntry>::TIterator FrameResultIterator(FrameResult);
+			TSharedPtr<TMap<BucketKeyType, FAggregatedTimingStats>> TaskResult;
+			ensure(FrameResultsQueue.Dequeue(TaskResult));
+			++ProcessedResults;
+
+			typename TMap<BucketKeyType, FAggregatedTimingStats>::TIterator TaskResultIterator(*TaskResult);
 			typename TMap<BucketKeyType, FInternalFrameAggregationEntry>::TIterator ResultIterator(GlobalResult);
 			while (ResultIterator)
 			{
-				FAggregatedTimingStats& FrameStats = FrameResultIterator->Value.Stats;
+				FAggregatedTimingStats& TaskStats = TaskResultIterator->Value;
 				FAggregatedTimingStats& ResultStats = ResultIterator->Value.Inner.Stats;
 
 				// TODO: Add Instance Count per frame column
 				// ResultStats.InstanceCount = (ResultStats.InstanceCount * Index + FrameStats.InstanceCount) / (Index + 1);
 
-				ResultStats.InstanceCount += FrameStats.InstanceCount;
-
-				ResultStats.TotalInclusiveTime += FrameStats.TotalInclusiveTime;
-				ResultStats.MinInclusiveTime = FMath::Min(ResultStats.MinInclusiveTime, FrameStats.TotalInclusiveTime);
-				ResultStats.MaxInclusiveTime = FMath::Max(ResultStats.MaxInclusiveTime, FrameStats.TotalInclusiveTime);
-				ResultIterator->Value.FrameInclusiveTimes.Add(FrameStats.TotalInclusiveTime);
+				ResultStats.InstanceCount += TaskStats.InstanceCount;
 
-				ResultStats.TotalExclusiveTime += FrameStats.TotalExclusiveTime;
-				ResultStats.MinExclusiveTime = FMath::Min(ResultStats.MinExclusiveTime, FrameStats.TotalExclusiveTime);
-				ResultStats.MaxExclusiveTime = FMath::Max(ResultStats.MaxExclusiveTime, FrameStats.TotalExclusiveTime);
-				ResultIterator->Value.FrameExclusiveTimes.Add(FrameStats.TotalExclusiveTime);
+				ResultStats.TotalInclusiveTime += TaskStats.TotalInclusiveTime;
+				ResultStats.MinInclusiveTime = FMath::Min(ResultStats.MinInclusiveTime, TaskStats.MinInclusiveTime);
+				ResultStats.MaxInclusiveTime = FMath::Max(ResultStats.MaxInclusiveTime, TaskStats.MaxInclusiveTime);
 
-				// Reset the per frame stats.
-				FrameStats = FAggregatedTimingStats();
+				ResultStats.TotalExclusiveTime += TaskStats.TotalExclusiveTime;
+				ResultStats.MinExclusiveTime = FMath::Min(ResultStats.MinExclusiveTime, TaskStats.MinExclusiveTime);
+				ResultStats.MaxExclusiveTime = FMath::Max(ResultStats.MaxExclusiveTime, TaskStats.MaxExclusiveTime);
 
-				++FrameResultIterator;
+				++TaskResultIterator;
 				++ResultIterator;
 			}
 		}
@@ -285,9 +376,65 @@ private:
 		});
 	}
 
-	static void UpdateTotalMinMaxTimerStats(FInternalAggregationEntry& AggregationEntry, double InclTime, double ExclTime)
+	template<typename TimelineType, typename BucketMappingFunc, typename BucketKeyType, typename CallbackType>
+	static void ProcessTimelineForFrameStats(const TimelineType* Timeline, BucketMappingFunc BucketMapper, CallbackType Callback, double IntervalStart, double IntervalEnd, TMap<BucketKeyType, FAggregatedTimingStats>& InternalResult)
+	{
+		struct FStackEntry
+		{
+			double StartTime;
+			double ExclusiveTime;
+			BucketKeyType BucketKey;
+		};
+
+		TArray<FStackEntry> Stack;
+		Stack.Reserve(1024);
+		double LastTime = 0.0;
+		Timeline->EnumerateEvents(IntervalStart, IntervalEnd, [BucketMapper, Callback, IntervalStart, IntervalEnd, &Stack, &LastTime, &InternalResult](bool IsEnter, double Time, const typename TimelineType::EventType& Event)
+		{
+			Time = FMath::Clamp(Time, IntervalStart, IntervalEnd);
+			BucketKeyType BucketKey = BucketMapper(Event);
+			if (Stack.Num())
+			{
+				FStackEntry& StackEntry = Stack.Top();
+				StackEntry.ExclusiveTime += Time - LastTime;
+			}
+			LastTime = Time;
+			if (IsEnter)
+			{
+				FStackEntry& StackEntry = Stack.AddDefaulted_GetRef();
+				StackEntry.StartTime = Time;
+				StackEntry.ExclusiveTime = 0.0;
+				StackEntry.BucketKey = BucketKey;
+			}
+			else
+			{
+				FStackEntry& StackEntry = Stack.Last();
+				double EventInclusiveTime = Time - StackEntry.StartTime;
+				check(EventInclusiveTime >= 0.0);
+				double EventExclusiveTime = StackEntry.ExclusiveTime;
+				check(EventExclusiveTime >= 0.0 && EventExclusiveTime <= EventInclusiveTime);
+				Stack.Pop(false);
+				double EventNonRecursiveInclusiveTime = EventInclusiveTime;
+				for (const FStackEntry& AncestorStackEntry : Stack)
+				{
+					if (AncestorStackEntry.BucketKey == BucketKey)
+					{
+						EventNonRecursiveInclusiveTime = 0.0;
+					}
+				}
+				Callback(InternalResult[BucketKey], EventNonRecursiveInclusiveTime, EventExclusiveTime);
+			}
+			return EEventEnumerate::Continue;
+		});
+	}
+
+	static void UpdateTotalMinMaxTimerAggregationStats(FInternalAggregationEntry& AggregationEntry, double InclTime, double ExclTime)
+	{
+		UpdateTotalMinMaxTimerStats(AggregationEntry.Stats, InclTime, ExclTime);
+	}
+
+	static void UpdateTotalMinMaxTimerStats(FAggregatedTimingStats& Stats, double InclTime, double ExclTime)
 	{
-		FAggregatedTimingStats& Stats = AggregationEntry.Stats;
 		Stats.TotalInclusiveTime += InclTime;
 		if (InclTime < Stats.MinInclusiveTime)
 		{
@@ -471,7 +618,7 @@ private:
 	}
 
 	template<typename TimelineType, typename BucketMappingFunc, typename BucketKeyType>
-	static void GatherKeysFromTimeline(const TimelineType* Timeline, BucketMappingFunc BucketMapper, double IntervalStart, double IntervalEnd, TMap<BucketKeyType, FInternalAggregationEntry>& InternalResult)
+	static void GatherKeysFromTimeline(const TimelineType* Timeline, BucketMappingFunc BucketMapper, double IntervalStart, double IntervalEnd, TMap<BucketKeyType, FAggregatedTimingStats>& InternalResult)
 	{
 		Timeline->EnumerateEvents(IntervalStart, IntervalEnd, [BucketMapper, IntervalStart, IntervalEnd, &InternalResult](bool IsEnter, double Time, const typename TimelineType::EventType& Event)
 			{
