[Insights] Fix Frame Stats aggregation issues

diff --git a/Engine/Source/Developer/TraceServices/Private/Common/TimelineStatistics.h b/Engine/Source/Developer/TraceServices/Private/Common/TimelineStatistics.h
--- a/Engine/Source/Developer/TraceServices/Private/Common/TimelineStatistics.h
+++ b/Engine/Source/Developer/TraceServices/Private/Common/TimelineStatistics.h
@@ -90,7 +90,7 @@ public:
 	template<typename TimelineType, typename BucketMappingFunc, typename BucketKeyType>
 	static void CreateFrameStatsAggregation(const TArray<const TimelineType*>& Timelines,
 											BucketMappingFunc BucketMapper,
-											TArray<FFrameData> Frames,
+											const TArray<FFrameData>& Frames,
 											TSharedPtr<FCancellationToken> CancellationToken,
 											TMap<BucketKeyType, FAggregatedTimingStats>& Result)
 	{
@@ -110,7 +110,7 @@ public:
 
 		// For very large numbers of timers or frames, an out of memory is possible so don't compute the median.
 		constexpr int64 MaxSize = 5 * 100 * 1000 * 1000;
-		bool bComputeMedian = InitialTimerMap.Num() * FramesNum < MaxSize;
+		bool bComputeMedian = InitialTimerMap.Num() * (int64) FramesNum < MaxSize;
 		for (auto& KV : InitialTimerMap)
 		{
 			FInternalFrameAggregationEntry Entry;
@@ -120,10 +120,9 @@ public:
 				Entry.FrameExclusiveTimes.AddUninitialized(Frames.Num());
 			}
 			Entry.Inner.Stats = KV.Value;
-			GlobalResult.Add(KV.Key, Entry);
+			GlobalResult.Add(KV.Key, MoveTemp(Entry));
 		}
 
-		std::atomic<int> NextFrameToProcess(0);
 		TQueue<TSharedPtr<TMap<BucketKeyType, FAggregatedTimingStats>>, EQueueMode::Mpsc> FrameResultsQueue;
 
 		constexpr float RatioOfThreadsToUse = 0.75;
@@ -139,7 +138,7 @@ public:
 		TArray<UE::Tasks::TTask<void>> AsyncTasks;
 		for (int32 Index = 0; Index < NumTasks; ++Index)
 		{
-			AsyncTasks.Add(UE::Tasks::Launch(UE_SOURCE_LOCATION, [StartIndex, EndIndex, &Timelines, bComputeMedian, BucketMapper, &Frames, &FrameResultsQueue, &NextFrameToProcess, &InitialTimerMap, &GlobalResult, CancellationToken]()
+			AsyncTasks.Add(UE::Tasks::Launch(UE_SOURCE_LOCATION, [StartIndex, EndIndex, &Timelines, bComputeMedian, BucketMapper, &Frames, &FrameResultsQueue, &InitialTimerMap, &GlobalResult, CancellationToken]()
 			{
 				if (CancellationToken.IsValid() && CancellationToken->ShouldCancel())
 				{
@@ -147,6 +146,7 @@ public:
 				}
 
 				TSharedPtr<TMap<BucketKeyType, FAggregatedTimingStats>> TaskResult = MakeShared<TMap<BucketKeyType, FAggregatedTimingStats>>(InitialTimerMap);
+				TMap<BucketKeyType, FAggregatedTimingStats> FrameResult(InitialTimerMap);
 
 				for(int FrameIndex = StartIndex; FrameIndex < EndIndex; ++FrameIndex)
 				{
@@ -155,15 +155,13 @@ public:
 						return;
 					}
 
-					TSharedPtr<TMap<BucketKeyType, FAggregatedTimingStats>> FrameResult = MakeShared<TMap<BucketKeyType, FAggregatedTimingStats>>(InitialTimerMap);
-
 					// Compute instance count and total/min/max inclusive/exclusive times for each timer.
 					for (const TimelineType* Timeline : Timelines)
 					{
-						ProcessTimelineForFrameStats(Timeline, BucketMapper, UpdateTotalMinMaxTimerStats, Frames[FrameIndex].StartTime, Frames[FrameIndex].EndTime, *FrameResult);
+						ProcessTimelineForFrameStats(Timeline, BucketMapper, UpdateTotalMinMaxTimerStats, Frames[FrameIndex].StartTime, Frames[FrameIndex].EndTime, FrameResult);
 					}
 
-					typename TMap<BucketKeyType, FAggregatedTimingStats>::TIterator FrameResultIterator(*FrameResult);
+					typename TMap<BucketKeyType, FAggregatedTimingStats>::TIterator FrameResultIterator(FrameResult);
 					typename TMap<BucketKeyType, FAggregatedTimingStats>::TIterator TaskResultIterator(*TaskResult);
 					typename TMap<BucketKeyType, FInternalFrameAggregationEntry>::TIterator GlobalResultIterator(GlobalResult);
 					while (FrameResultIterator)
@@ -200,7 +198,7 @@ public:
 
 				}
 
-				FrameResultsQueue.Enqueue(TaskResult);
+				FrameResultsQueue.Enqueue(MoveTemp(TaskResult));
 			}));
 
 			StartIndex = EndIndex;
