[Insights] Fixed conversion warning.

diff --git a/Engine/Source/Developer/TraceServices/Private/Common/TimelineStatistics.h b/Engine/Source/Developer/TraceServices/Private/Common/TimelineStatistics.h
--- a/Engine/Source/Developer/TraceServices/Private/Common/TimelineStatistics.h
+++ b/Engine/Source/Developer/TraceServices/Private/Common/TimelineStatistics.h
@@ -126,7 +126,7 @@ public:
 		TQueue<TSharedPtr<TMap<BucketKeyType, FAggregatedTimingStats>>, EQueueMode::Mpsc> FrameResultsQueue;
 
 		constexpr float RatioOfThreadsToUse = 0.75;
-		int32 NumTasks = FMath::Max(1, GThreadPool->GetNumThreads() * RatioOfThreadsToUse);
+		int32 NumTasks = FMath::Max(1, int32(GThreadPool->GetNumThreads() * RatioOfThreadsToUse));
 		NumTasks = FMath::Min((int32)NumTasks, Frames.Num());
 		int32 FramesPerTask = Frames.Num() / NumTasks;
 		int32 ExtraFrameTasks = Frames.Num() - NumTasks * FramesPerTask;
