[Insights] MemoryInsights: Fixed log messages in mem analysis re TagTracker errors when memalloc channel is not enabled.

diff --git a/Engine/Source/Developer/TraceServices/Private/Model/AllocationsProvider.cpp b/Engine/Source/Developer/TraceServices/Private/Model/AllocationsProvider.cpp
--- a/Engine/Source/Developer/TraceServices/Private/Model/AllocationsProvider.cpp
+++ b/Engine/Source/Developer/TraceServices/Private/Model/AllocationsProvider.cpp
@@ -2109,6 +2109,13 @@ void FAllocationsProvider::EditPushTagFromPtr(uint32 ThreadId, uint8 Tracker, ui
 {
 	EditAccessCheck();
 
+	if (!bInitialized)
+	{
+		// No errors if the "memallocs" channel was not enabled.
+		TagTracker.PushTagFromPtr(ThreadId, Tracker, 0);
+		return;
+	}
+
 	const FAllocationItem* Alloc = nullptr;
 
 	if (IsValidRootHeap(RootHeapId))
@@ -2157,6 +2164,10 @@ void FAllocationsProvider::EditOnAnalysisCompleted(double Time)
 
 	if (!bInitialized)
 	{
+		if (TagTracker.GetNumErrors() > 0)
+		{
+			UE_LOG(LogTraceServices, Error, TEXT("[MemAlloc] TagTracker errors: %u"), TagTracker.GetNumErrors());
+		}
 		return;
 	}
 
