At end of analysis, add still-open cpu scopes as infinite-length markers for aid in debugging hangs.

diff --git a/Engine/Source/Developer/TraceServices/Private/Analyzers/CpuProfilerTraceAnalysis.cpp b/Engine/Source/Developer/TraceServices/Private/Analyzers/CpuProfilerTraceAnalysis.cpp
--- a/Engine/Source/Developer/TraceServices/Private/Analyzers/CpuProfilerTraceAnalysis.cpp
+++ b/Engine/Source/Developer/TraceServices/Private/Analyzers/CpuProfilerTraceAnalysis.cpp
@@ -56,12 +56,11 @@ void FCpuProfilerAnalyzer::OnAnalysisBegin(const FOnAnalysisContext& Context)
 
 ////////////////////////////////////////////////////////////////////////////////////////////////////
 
-void FCpuProfilerAnalyzer::OnAnalysisEnd(/*const FOnAnalysisEndContext& Context*/)
+void FCpuProfilerAnalyzer::OnAnalysisEnd()
 {
-#if 0
 	LLM_SCOPE_BYNAME(TEXT("Insights/FCpuProfilerAnalyzer"));
 
-	//TODO: Context.EventTime
+	double Timestamp = std::numeric_limits<double>::infinity();
 	for (auto& KV : ThreadStatesMap)
 	{
 		FThreadState& ThreadState = *KV.Value;
@@ -72,13 +71,6 @@ void FCpuProfilerAnalyzer::OnAnalysisEnd(/*const FOnAnalysisEndContext& Context*
 
 		if (ThreadState.LastCycle != 0 && ThreadState.LastCycle != ~0)
 		{
-			//double Timestamp = Context.EventTime.AsSeconds(ThreadState.LastCycle);
-			//Session.UpdateDurationSeconds(Timestamp);
-			double Timestamp;
-			{
-				FAnalysisSessionEditScope _(Session);
-				Timestamp = Session.GetDurationSeconds();
-			}
 			while (ThreadState.ScopeStack.Num())
 			{
 				ThreadState.ScopeStack.Pop();
@@ -90,7 +82,6 @@ void FCpuProfilerAnalyzer::OnAnalysisEnd(/*const FOnAnalysisEndContext& Context*
 
 		ensure(ThreadState.ScopeStack.Num() == 0);
 	}
-#endif
 }
 
 ////////////////////////////////////////////////////////////////////////////////////////////////////
