Fix crash when when searching for previous match in Quick Find.

diff --git a/Engine/Source/Developer/TraceServices/Public/Model/MonotonicTimeline.h b/Engine/Source/Developer/TraceServices/Public/Model/MonotonicTimeline.h
--- a/Engine/Source/Developer/TraceServices/Public/Model/MonotonicTimeline.h
+++ b/Engine/Source/Developer/TraceServices/Public/Model/MonotonicTimeline.h
@@ -571,23 +571,20 @@ public:
 
 		if (bSearchStartTimeUsingPages)
 		{
+			ScopeEntryIterator.NextPage();
 			const FEventScopeEntryPage* CurrentScopePage = ScopeEntryIterator.GetCurrentPage();
-			do
+			check(CurrentStackDepth <= CurrentScopePage->InitialStackCount);
+
+			while (CurrentStackDepth > 0)
 			{
-				check(CurrentStackDepth <= CurrentScopePage->InitialStackCount);
-				while (CurrentStackDepth > 0)
+				--CurrentStackDepth;
+				EventType CurrentEvent = DetailLevel.GetEvent(CurrentScopePage->InitialStack[CurrentStackDepth].EventIndex);
+				double StartTime = DetailLevel.ScopeEntries[CurrentScopePage->InitialStack[CurrentStackDepth].EnterScopeIndex].Time;
+				if (Callback(false, StartTime, CurrentEvent) == EEventEnumerate::Stop)
 				{
-					--CurrentStackDepth;
-					EventType CurrentEvent = DetailLevel.GetEvent(CurrentScopePage->InitialStack[CurrentStackDepth].EventIndex);
-					double StartTime = DetailLevel.ScopeEntries[CurrentScopePage->InitialStack[CurrentStackDepth].EnterScopeIndex].Time;
-					if (Callback(false, StartTime, CurrentEvent) == EEventEnumerate::Stop)
-					{
-						return;
-					}
+					return;
 				}
-				CurrentScopePage = ScopeEntryIterator.PrevPage();
 			}
-			while (CurrentScopePage != nullptr);
 		}
 	}
 
