[Insights] Fix crash when computing aggregated frame stats.

diff --git a/Engine/Source/Developer/TraceServices/Private/Model/TimingProfiler.cpp b/Engine/Source/Developer/TraceServices/Private/Model/TimingProfiler.cpp
--- a/Engine/Source/Developer/TraceServices/Private/Model/TimingProfiler.cpp
+++ b/Engine/Source/Developer/TraceServices/Private/Model/TimingProfiler.cpp
@@ -308,6 +308,8 @@ ITable<FTimingProfilerAggregatedStats>* FTimingProfilerProvider::CreateAggregati
 
 		if (Frames.Num() > 0)
 		{
+			// Do not allow inf for the last frame end time.
+			Frames[Frames.Num() - 1].EndTime = FMath::Min(Session.GetDurationSeconds(), Frames[Frames.Num() - 1].EndTime);
 			FTimelineStatistics::CreateFrameStatsAggregation(IncludedTimelines, BucketMappingFunc, Frames, Params.CancellationToken, Aggregation);
 		}
 	}
