[Insights] CpuProfiler: Fixed case where there are still pending events (cpu timing events with metadata) at the end of the analysis session.

diff --git a/Engine/Source/Developer/TraceServices/Private/Analyzers/CpuProfilerTraceAnalysis.cpp b/Engine/Source/Developer/TraceServices/Private/Analyzers/CpuProfilerTraceAnalysis.cpp
--- a/Engine/Source/Developer/TraceServices/Private/Analyzers/CpuProfilerTraceAnalysis.cpp
+++ b/Engine/Source/Developer/TraceServices/Private/Analyzers/CpuProfilerTraceAnalysis.cpp
@@ -60,17 +60,24 @@ void FCpuProfilerAnalyzer::OnAnalysisEnd()
 {
 	LLM_SCOPE_BYNAME(TEXT("Insights/FCpuProfilerAnalyzer"));
 
-	double Timestamp = std::numeric_limits<double>::infinity();
+	const double Timestamp = std::numeric_limits<double>::infinity();
+
 	for (auto& KV : ThreadStatesMap)
 	{
 		FThreadState& ThreadState = *KV.Value;
 
-		//TODO: EndThread(Context.EventTime, ThreadState);
-
-		ensure(ThreadState.PendingEvents.Num() == 0);
-
 		if (ThreadState.LastCycle != 0 && ThreadState.LastCycle != ~0)
 		{
+			int32 RemainingPending = ThreadState.PendingEvents.Num();
+			if (RemainingPending > 0)
+			{
+				uint64 LastCycle = ThreadState.LastCycle;
+				const FPendingEvent* PendingCursor = ThreadState.PendingEvents.GetData();
+				DispatchPendingEvents(LastCycle, ~0ull, ThreadState, PendingCursor, RemainingPending);
+				check(RemainingPending == 0);
+				ThreadState.PendingEvents.Reset();
+			}
+
 			while (ThreadState.ScopeStack.Num())
 			{
 				ThreadState.ScopeStack.Pop();
@@ -80,6 +87,7 @@ void FCpuProfilerAnalyzer::OnAnalysisEnd()
 			}
 		}
 
+		ensure(ThreadState.PendingEvents.Num() == 0);
 		ensure(ThreadState.ScopeStack.Num() == 0);
 	}
 }
@@ -161,7 +169,7 @@ bool FCpuProfilerAnalyzer::OnEvent(uint16 RouteId, EStyle Style, const FOnEventC
 		{
 			uint64 LastCycle = ThreadState.LastCycle;
 			const FPendingEvent* PendingCursor = ThreadState.PendingEvents.GetData();
-			DispatchPendingEvents(LastCycle, ~0ull, Context.EventTime, ThreadState, PendingCursor, RemainingPending);
+			DispatchPendingEvents(LastCycle, ~0ull, ThreadState, PendingCursor, RemainingPending);
 			check(RemainingPending == 0);
 			ThreadState.PendingEvents.Reset();
 		}
@@ -283,7 +291,7 @@ uint64 FCpuProfilerAnalyzer::ProcessBuffer(const FEventTime& EventTime, FThreadS
 		}
 
 		// Dispatch pending events that are younger than the one we've just decoded.
-		DispatchPendingEvents(LastCycle, ActualCycle, EventTime, ThreadState, PendingCursor, RemainingPending);
+		DispatchPendingEvents(LastCycle, ActualCycle, ThreadState, PendingCursor, RemainingPending);
 
 		double ActualTime = EventTime.AsSeconds(ActualCycle);
 
@@ -371,7 +379,7 @@ uint64 FCpuProfilerAnalyzer::ProcessBufferV2(const FEventTime& EventTime, FThrea
 		}
 
 		// Dispatch pending events that are younger than the one we've just decoded.
-		DispatchPendingEvents(LastCycle, ActualCycle, EventTime, ThreadState, PendingCursor, RemainingPending);
+		DispatchPendingEvents(LastCycle, ActualCycle, ThreadState, PendingCursor, RemainingPending);
 
 		double ActualTime = EventTime.AsSeconds(ActualCycle);
 
@@ -541,7 +549,6 @@ uint64 FCpuProfilerAnalyzer::ProcessBufferV2(const FEventTime& EventTime, FThrea
 void FCpuProfilerAnalyzer::DispatchPendingEvents(
 	uint64& LastCycle,
 	uint64 CurrentCycle,
-	const FEventTime& EventTime,
 	FThreadState& ThreadState,
 	const FPendingEvent*& PendingCursor,
 	int32& RemainingPending)
@@ -584,7 +591,7 @@ void FCpuProfilerAnalyzer::DispatchPendingEvents(
 		// Update LastCycle in order to verify time (of following pending events) increases monotonically.
 		LastCycle = PendingCycle;
 
-		double PendingTime = EventTime.AsSeconds(PendingCycle);
+		double PendingTime = PendingCursor->Time;
 
 		if (bEnter)
 		{
@@ -644,7 +651,8 @@ void FCpuProfilerAnalyzer::OnCpuScopeEnter(const FOnEventContext& Context)
 	TimerId = EditableTimingProfilerProvider.AddMetadata(TimerId, MoveTemp(CborData));
 
 	uint64 Cycle = Context.EventTime.AsCycle64();
-	ThreadState.PendingEvents.Add({Cycle, TimerId});
+	double Time = Context.EventTime.AsSeconds();
+	ThreadState.PendingEvents.Add({ Cycle, Time, TimerId });
 }
 
 ////////////////////////////////////////////////////////////////////////////////////////////////////
@@ -665,7 +673,8 @@ void FCpuProfilerAnalyzer::OnCpuScopeLeave(const FOnEventContext& Context)
 	}
 
 	uint64 Cycle = Context.EventTime.AsCycle64();
-	ThreadState.PendingEvents.Add({~Cycle});
+	double Time = Context.EventTime.AsSeconds();
+	ThreadState.PendingEvents.Add({ ~Cycle, Time, 0 });
 }
 
 ////////////////////////////////////////////////////////////////////////////////////////////////////
diff --git a/Engine/Source/Developer/TraceServices/Private/Analyzers/CpuProfilerTraceAnalysis.h b/Engine/Source/Developer/TraceServices/Private/Analyzers/CpuProfilerTraceAnalysis.h
--- a/Engine/Source/Developer/TraceServices/Private/Analyzers/CpuProfilerTraceAnalysis.h
+++ b/Engine/Source/Developer/TraceServices/Private/Analyzers/CpuProfilerTraceAnalysis.h
@@ -35,6 +35,7 @@ private:
 	struct FPendingEvent
 	{
 		uint64 Cycle;
+		double Time;
 		uint32 TimerId;
 	};
 
@@ -56,7 +57,7 @@ private:
 	FThreadState& GetThreadState(uint32 ThreadId);
 	uint64 ProcessBuffer(const FEventTime& EventTime, FThreadState& ThreadState, const uint8* BufferPtr, uint32 BufferSize);
 	uint64 ProcessBufferV2(const FEventTime& EventTime, FThreadState& ThreadState, const uint8* BufferPtr, uint32 BufferSize);
-	void DispatchPendingEvents(uint64& LastCycle, uint64 CurrentCycle, const FEventTime& EventTime, FThreadState& ThreadState, const FPendingEvent*& PendingCursor, int32& RemainingPending);
+	void DispatchPendingEvents(uint64& LastCycle, uint64 CurrentCycle, FThreadState& ThreadState, const FPendingEvent*& PendingCursor, int32& RemainingPending);
 
 	enum : uint16
 	{
