[Insights] AllocationsProvider: Improved memory analysis:
  - Added INSIGHTS_DOUBLE_ALLOC_FREE_PREVIOUS toggle (Automatically free the previous alloc when detecting addresses allocated multiple times.). Works with INSIGHTS_VALIDATE_ALLOC_EVENTS.
  - Added INSIGHTS_VALIDATE_FREE_EVENTS toggle (Automatically add a fake alloc when Free event detects a missing alloc.)
  - Added INSIGHTS_VALIDATE_HEAP_MARK_ALLOC_EVENTS toggle (Automatically add a fake alloc when MarkAllocAsHeap event detects a missing alloc.)
  - Added INSIGHTS_VALIDATE_HEAP_UNMARK_ALLOC_EVENTS toggle (Automatically add a fake heap (or mark an existing alloc as heap) when UnmarkAllocAsHeap detects a missing alloc.)
  - Added INSIGHTS_WARNINGS_FOR_NULLPTR_ALLOCS toggle (Generates warnings for alloc/free events with nullptr, for the System root heap.)
  - Removed INSIGHTS_REMAP_INVALID_ALLOC_EVENTS and INSIGHTS_REMAP_INVALID_FREE_EVENTS.

diff --git a/Engine/Source/Developer/TraceServices/Private/Model/AllocationsProvider.cpp b/Engine/Source/Developer/TraceServices/Private/Model/AllocationsProvider.cpp
--- a/Engine/Source/Developer/TraceServices/Private/Model/AllocationsProvider.cpp
+++ b/Engine/Source/Developer/TraceServices/Private/Model/AllocationsProvider.cpp
@@ -37,12 +37,28 @@
 // ~10% to ~30% of all allocs are expected to have an "event distance" == 1 event ("free" event follows the "alloc" event immediately)
 #define INSIGHTS_USE_LAST_ALLOC 1
 
+// Detects cases where an address is allocated multiple times (i.e. missing Free or MarkAllocAsHeap events).
+// Note: This slows down analysis significantly, so it is disabled by default.
 #define INSIGHTS_VALIDATE_ALLOC_EVENTS 0
-#define INSIGHTS_REMAP_INVALID_ALLOC_EVENTS 0
-#define INSIGHTS_REMAP_INVALID_FREE_EVENTS 0
+
+// Automatically free the previous alloc when detecting addresses allocated multiple times (INSIGHTS_VALIDATE_ALLOC_EVENTS).
+#define INSIGHTS_DOUBLE_ALLOC_FREE_PREVIOUS 1
+
+// Automatically add a fake alloc when Free event detects a missing alloc.
+#define INSIGHTS_VALIDATE_FREE_EVENTS 1
+
+// Automatically add a fake alloc when MarkAllocAsHeap event detects a missing alloc.
+#define INSIGHTS_VALIDATE_HEAP_MARK_ALLOC_EVENTS 1
+
+// Automatically add a fake heap (or mark an existing alloc as heap) when UnmarkAllocAsHeap detects a missing alloc.
+#define INSIGHTS_VALIDATE_HEAP_UNMARK_ALLOC_EVENTS 1
 
 #define INSIGHTS_DEBUG_METADATA 0
 
+// Generates warnings for alloc/free events with nullptr, for the System root heap.
+// Note: Allocating nullptr with size != 0 will always generate warnings.
+#define INSIGHTS_WARNINGS_FOR_NULLPTR_ALLOCS 0
+
 ////////////////////////////////////////////////////////////////////////////////////////////////////
 
 // Debug functionality to discard some of the Alloc/Free/MarkAllocAsHeap/UnmarkAllocAsHeap events.
@@ -226,11 +242,9 @@ void FTagTracker::AddTagSpec(TagIdType InTag, TagIdType InParentTag, const TCHAR
 			BuildTagPath(FullName, Display, InParentTag);
 		}
 
-		const FTagEntry& Entry = TagMap.Emplace(InTag, FTagEntry{
-			Session.StoreString(DisplayName),
-			Session.StoreString(FullName.ToString()),
-			InParentTag
-		});
+		const TCHAR* TagDisplayName = Session.StoreString(DisplayName);
+		const TCHAR* TagFullPath = Session.StoreString(FullName.ToString());
+		const FTagEntry& Entry = TagMap.Emplace(InTag, FTagEntry{ TagDisplayName, TagFullPath, InParentTag });
 
 		// Check if this new tag has been referenced before by a child tag
 		for (const TTuple<TagIdType,FString>& Pending : PendingTags)
@@ -249,7 +263,7 @@ void FTagTracker::AddTagSpec(TagIdType InTag, TagIdType InParentTag, const TCHAR
 		}
 		else
 		{
-			UE_LOG(LogTraceServices, Verbose, TEXT("[MemAlloc] Added Tag '%s' ('%s') with id %u."), Entry.Display, Entry.FullPath, InTag);
+			UE_LOG(LogTraceServices, Verbose, TEXT("[MemAlloc] Added Tag '%s' ('%s') with id %u (ParentTag=%u)."), Entry.Display, Entry.FullPath, InTag, InParentTag);
 		}
 	}
 	else
@@ -257,7 +271,7 @@ void FTagTracker::AddTagSpec(TagIdType InTag, TagIdType InParentTag, const TCHAR
 		++NumErrors;
 		if (NumErrors <= MaxLogMessagesPerErrorType)
 		{
-			UE_LOG(LogTraceServices, Error, TEXT("[MemAlloc] Tag with id %u (ParentTag=%u, Display=%s) already added!"), InTag, InParentTag, InDisplay);
+			UE_LOG(LogTraceServices, Error, TEXT("[MemAlloc] Tag with id %u (ParentTag=%u, Display='%s') already added!"), InTag, InParentTag, InDisplay);
 		}
 	}
 }
@@ -1339,6 +1353,12 @@ void FAllocationsProvider::EditAlloc(double Time, uint32 CallstackId, uint64 Add
 
 	if (Address == 0 && RootHeapId == EMemoryTraceRootHeap::SystemMemory)
 	{
+#if !INSIGHTS_WARNINGS_FOR_NULLPTR_ALLOCS
+		if (Size == 0)
+		{
+			return;
+		}
+#endif // INSIGHTS_WARNINGS_FOR_NULLPTR_ALLOCS
 		++AllocWarnings;
 		if (AllocWarnings <= MaxLogMessagesPerWarningType)
 		{
@@ -1365,36 +1385,38 @@ void FAllocationsProvider::EditAlloc(double Time, uint32 CallstackId, uint64 Add
 
 	const TagIdType Tag = TagTracker.GetCurrentTag(CurrentSystemThreadId, CurrentTracker);
 
-	FAllocationItem* AllocationPtr = nullptr;
-
 #if INSIGHTS_VALIDATE_ALLOC_EVENTS
-	AllocationPtr = RootHeap.LiveAllocs->FindRef(Address);
-#if INSIGHTS_REMAP_INVALID_ALLOC_EVENTS
-	if (AllocationPtr)
+	FAllocationItem* ExistingAllocationPtr = RootHeap.LiveAllocs->FindRef(Address);
+	if (ExistingAllocationPtr)
 	{
-		const uint64 OriginalAddress = Address;
-		check(OriginalAddress & (0x3FFull << 54) == 0);
-		for (uint64 Offset = 1; Offset <= 0x3FFull; ++Offset)
+#if INSIGHTS_DOUBLE_ALLOC_FREE_PREVIOUS
+
+		++AllocErrors;
+		if (AllocErrors <= MaxLogMessagesPerErrorType)
 		{
-			Address = OriginalAddress | (Offset << 54);
-			AllocationPtr = RootHeap.LiveAllocs->FindRef(Address);
-			if (!AllocationPtr)
-			{
-				break;
-			}
+			UE_LOG(LogTraceServices, Error, TEXT("[MemAlloc] Invalid ALLOC event (Address=0x%llX, Size=%llu, Tag=%u, RootHeap=%u, Time=%f, CallstackId=%u)! The previous alloc will be freed."), Address, Size, Tag, RootHeapId, Time, CallstackId);
 		}
-		++AllocWarnings;
-		if (AllocWarnings <= MaxLogMessagesPerWarningType)
+
+		// Free the previous allocation.
+		INSIGHTS_WATCH_INDIRECT_API_LOGF(TEXT("Free"), Address, Time);
+		constexpr uint32 FreeCallstackId = 0; // no callstack
+		EditFree(Time, FreeCallstackId, Address, RootHeapId);
+		RootHeap.SbTree->SetTimeForEvent(RootHeap.EventIndex, Time); // for the case where the next event is first event in a new SbTree column after changing the heap alloc
+
+#else // INSIGHTS_DOUBLE_ALLOC_FREE_PREVIOUS
+
+		++AllocErrors;
+		if (AllocErrors <= MaxLogMessagesPerErrorType)
 		{
-			UE_LOG(LogTraceServices, Warning, TEXT("[MemAlloc] Invalid ALLOC event remapped (Address=0x%llX --> 0x%llX, Size=%llu, Tag=%u, RootHeap=%u, Time=%f, CallstackId=%u)!"), OriginalAddress, Address, Size, Tag, RootHeapId, Time, CallstackId);
+			UE_LOG(LogTraceServices, Error, TEXT("[MemAlloc] Invalid ALLOC event (Address=0x%llX, Size=%llu, Tag=%u, RootHeap=%u, Time=%f, CallstackId=%u)!"), Address, Size, Tag, RootHeapId, Time, CallstackId);
 		}
+
+#endif // INSIGHTS_DOUBLE_ALLOC_FREE_PREVIOUS
 	}
-#endif // INSIGHTS_REMAP_INVALID_ALLOC_EVENTS
 #endif // INSIGHTS_VALIDATE_ALLOC_EVENTS
 
-	if (!AllocationPtr)
 	{
-		AllocationPtr = RootHeap.LiveAllocs->AddNew(Address);
+		FAllocationItem* AllocationPtr = RootHeap.LiveAllocs->AddNew(Address);
 		FAllocationItem& Allocation = *AllocationPtr;
 
 		uint32 MetadataId = MetadataProvider.InvalidMetadataId;
@@ -1444,16 +1466,6 @@ void FAllocationsProvider::EditAlloc(double Time, uint32 CallstackId, uint64 Add
 		SampleMaxLiveAllocations = FMath::Max(SampleMaxLiveAllocations, TotalLiveAllocations);
 		++SampleAllocEvents;
 	}
-#if INSIGHTS_VALIDATE_ALLOC_EVENTS
-	else
-	{
-		++AllocErrors;
-		if (AllocErrors <= MaxLogMessagesPerErrorType)
-		{
-			UE_LOG(LogTraceServices, Error, TEXT("[MemAlloc] Invalid ALLOC event (Address=0x%llX, Size=%llu, Tag=%u, RootHeap=%u, Time=%f, CallstackId=%u)!"), Address, Size, Tag, RootHeapId, Time, CallstackId);
-		}
-	}
-#endif
 
 	++AllocCount;
 	if (RootHeap.EventIndex != ~0u)
@@ -1483,11 +1495,13 @@ void FAllocationsProvider::EditFree(double Time, uint32 CallstackId, uint64 Addr
 
 	if (Address == 0 && RootHeapId == EMemoryTraceRootHeap::SystemMemory)
 	{
+#if INSIGHTS_WARNINGS_FOR_NULLPTR_ALLOCS
 		++FreeWarnings;
 		if (FreeWarnings <= MaxLogMessagesPerWarningType)
 		{
 			UE_LOG(LogTraceServices, Warning, TEXT("[MemAlloc] Free for address 0 : RootHeap=%u, Time=%f, CallstackId=%u"), RootHeapId, Time, CallstackId);
 		}
+#endif // INSIGHTS_WARNINGS_FOR_NULLPTR_ALLOCS
 		return;
 	}
 
@@ -1518,7 +1532,7 @@ void FAllocationsProvider::EditFree(double Time, uint32 CallstackId, uint64 Addr
 			HeapId Heap = AllocationPtr->RootHeap;
 			INSIGHTS_WATCH_INDIRECT_API_LOGF(TEXT("UnmarkAllocationAsHeap"), Address, Time);
 			EditUnmarkAllocationAsHeap(Time, CallstackId, Address, Heap);
-			RootHeap.SbTree->SetTimeForEvent(RootHeap.EventIndex, Time); // for the case where the free event is first event in a new SbTree column after changing the heap alloc
+			RootHeap.SbTree->SetTimeForEvent(RootHeap.EventIndex, Time); // for the case where the next event is first event in a new SbTree column after changing the heap alloc
 			AllocationPtr = RootHeap.LiveAllocs->Remove(Address); // we take ownership of AllocationPtr
 		}
 	}
@@ -1527,28 +1541,7 @@ void FAllocationsProvider::EditFree(double Time, uint32 CallstackId, uint64 Addr
 		INSIGHTS_SLOW_CHECK(!AllocationPtr->IsHeap());
 	}
 
-#if INSIGHTS_REMAP_INVALID_FREE_EVENTS
-	if (!AllocationPtr)
-	{
-		const uint64 OriginalAddress = Address;
-		check(OriginalAddress & (0x3FFull << 54) == 0);
-		for (uint64 Offset = 1; Offset <= 0x3FFull; ++Offset)
-		{
-			Address = OriginalAddress | (Offset << 54);
-			AllocationPtr = RootHeap.LiveAllocs->FindRef(Address);
-			if (AllocationPtr)
-			{
-				break;
-			}
-		}
-		++AllocWarnings;
-		if (AllocWarnings <= MaxLogMessagesPerWarningType)
-		{
-			UE_LOG(LogTraceServices, Warning, TEXT("[MemAlloc] Invalid FREE event remapped (Address=0x%llX --> 0x%llX, RootHeap=%u, Time=%f, CallstackId=%u)!"), OriginalAddress, Address, RootHeapId, Time, CallstackId);
-		}
-	}
-#endif // INSIGHTS_REMAP_INVALID_FREE_EVENTS
-
+#if INSIGHTS_VALIDATE_FREE_EVENTS
 	if (!AllocationPtr)
 	{
 		++FreeErrors;
@@ -1561,9 +1554,11 @@ void FAllocationsProvider::EditFree(double Time, uint32 CallstackId, uint64 Addr
 		constexpr uint32 FakeAllocAlignment = 0;
 		INSIGHTS_WATCH_INDIRECT_API_LOGF(TEXT("Alloc"), Address, Time);
 		EditAlloc(Time, CallstackId, Address, FakeAllocSize, FakeAllocAlignment, RootHeapId);
-		RootHeap.SbTree->SetTimeForEvent(RootHeap.EventIndex, Time); // for the case where the free event is first event in a new SbTree column after adding the fake alloc
+		RootHeap.SbTree->SetTimeForEvent(RootHeap.EventIndex, Time); // for the case where the next event is first event in a new SbTree column after adding the fake alloc
 		AllocationPtr = RootHeap.LiveAllocs->Remove(Address); // we take ownership of AllocationPtr
 	}
+#endif // INSIGHTS_VALIDATE_FREE_EVENTS
+
 	if (AllocationPtr)
 	{
 		check(RootHeap.EventIndex > AllocationPtr->StartEventIndex);
@@ -1811,6 +1806,26 @@ void FAllocationsProvider::EditMarkAllocationAsHeap(double Time, uint32 Callstac
 
 	// Remove the allocation from the Live allocs.
 	FAllocationItem* Alloc = RootHeap.LiveAllocs->Remove(Address); // we take ownership of Alloc
+
+#if INSIGHTS_VALIDATE_HEAP_MARK_ALLOC_EVENTS
+	if (!Alloc)
+	{
+		++HeapErrors;
+		if (HeapErrors <= MaxLogMessagesPerErrorType)
+		{
+			UE_LOG(LogTraceServices, Error, TEXT("[MemAlloc] HeapMarkAlloc: Could not find alloc with address 0x%llX (Heap=%u \"%s\", Flags=%u, Time=%f, CallstackId=%u)! A fake alloc will be created with size 0."), Address, Heap, HeapSpec.Name, uint32(Flags), Time, CallstackId);
+		}
+
+		// Fake the missing alloc.
+		constexpr uint64 FakeAllocSize = 0;
+		constexpr uint32 FakeAllocAlignment = 0;
+		INSIGHTS_WATCH_INDIRECT_API_LOGF(TEXT("Alloc"), Address, Time);
+		EditAlloc(Time, CallstackId, Address, FakeAllocSize, FakeAllocAlignment, RootHeap.HeapSpec->Id);
+		RootHeap.SbTree->SetTimeForEvent(RootHeap.EventIndex, Time); // for the case where the next event is first event in a new SbTree column after adding the fake alloc
+		Alloc = RootHeap.LiveAllocs->Remove(Address); // we take ownership of Alloc
+	}
+#endif // INSIGHTS_VALIDATE_HEAP_MARK_ALLOC_EVENTS
+
 	if (Alloc)
 	{
 		check(Address == Alloc->Address);
@@ -1847,7 +1862,7 @@ void FAllocationsProvider::EditMarkAllocationAsHeap(double Time, uint32 Callstac
 		++HeapErrors;
 		if (HeapErrors <= MaxLogMessagesPerErrorType)
 		{
-			UE_LOG(LogTraceServices, Error, TEXT("[MemAlloc] HeapMarkAlloc: Could not find address 0x%llX (Heap=%u \"%s\", Flags=%u, Time=%f, CallstackId=%u)!"), Address, Heap, HeapSpec.Name, uint32(Flags), Time, CallstackId);
+			UE_LOG(LogTraceServices, Error, TEXT("[MemAlloc] HeapMarkAlloc: Could not find alloc with address 0x%llX (Heap=%u \"%s\", Flags=%u, Time=%f, CallstackId=%u)!"), Address, Heap, HeapSpec.Name, uint32(Flags), Time, CallstackId);
 		}
 	}
 
@@ -1895,6 +1910,48 @@ void FAllocationsProvider::EditUnmarkAllocationAsHeap(double Time, uint32 Callst
 
 	// Remove the heap allocation from the Live allocs.
 	FAllocationItem* Alloc = RootHeap.LiveAllocs->RemoveHeap(Address); // we take ownership of Alloc
+
+#if INSIGHTS_VALIDATE_HEAP_UNMARK_ALLOC_EVENTS
+	if (!Alloc)
+	{
+		// Remove the allocation from the Live allocs.
+		Alloc = RootHeap.LiveAllocs->Remove(Address); // we take ownership of Alloc
+
+		if (!Alloc)
+		{
+			++HeapErrors;
+			if (HeapErrors <= MaxLogMessagesPerErrorType)
+			{
+				UE_LOG(LogTraceServices, Error, TEXT("[MemAlloc] HeapUnmarkAlloc: Could not find heap with address 0x%llX (Heap=%u \"%s\", Time=%f, CallstackId=%u)! A fake heap alloc will be created with size 0."), Address, Heap, HeapSpec.Name, Time, CallstackId);
+			}
+
+			// Fake the missing alloc.
+			constexpr uint64 FakeAllocSize = 0;
+			constexpr uint32 FakeAllocAlignment = 0;
+			INSIGHTS_WATCH_INDIRECT_API_LOGF(TEXT("Alloc"), Address, Time);
+			EditAlloc(Time, CallstackId, Address, FakeAllocSize, FakeAllocAlignment, RootHeap.HeapSpec->Id);
+			RootHeap.SbTree->SetTimeForEvent(RootHeap.EventIndex, Time); // for the case where the next event is first event in a new SbTree column after adding the fake alloc
+			Alloc = RootHeap.LiveAllocs->Remove(Address); // we take ownership of Alloc
+			check(Alloc != nullptr);
+		}
+		else
+		{
+			++HeapErrors;
+			if (HeapErrors <= MaxLogMessagesPerErrorType)
+			{
+				UE_LOG(LogTraceServices, Error, TEXT("[MemAlloc] HeapUnmarkAlloc: Could not find heap with address 0x%llX (Heap=%u \"%s\", Time=%f, CallstackId=%u)! An alloc with this address exists. It will be marked as heap."), Address, Heap, HeapSpec.Name, Time, CallstackId);
+			}
+		}
+
+		check(RootHeap.LiveAllocs->FindRef(Address) == nullptr);
+
+		// Mark allocation as a "heap" allocation.
+		Alloc->Flags = Alloc->Flags | EMemoryTraceHeapAllocationFlags::Heap;
+		Alloc->RootHeap = static_cast<uint8>(Heap);
+		Alloc->CallstackId = CallstackId;
+	}
+#endif // INSIGHTS_VALIDATE_HEAP_UNMARK_ALLOC_EVENTS
+
 	if (Alloc)
 	{
 		check(Address == Alloc->Address);
@@ -1994,7 +2051,7 @@ void FAllocationsProvider::EditUnmarkAllocationAsHeap(double Time, uint32 Callst
 		++HeapErrors;
 		if (HeapErrors <= MaxLogMessagesPerErrorType)
 		{
-			UE_LOG(LogTraceServices, Error, TEXT("[MemAlloc] HeapUnmarkAlloc: Could not find address 0x%llX (Heap=%u \"%s\", Time=%f, CallstackId=%u)!"), Address, Heap, HeapSpec.Name, Time, CallstackId);
+			UE_LOG(LogTraceServices, Error, TEXT("[MemAlloc] HeapUnmarkAlloc: Could not find heap with address 0x%llX (Heap=%u \"%s\", Time=%f, CallstackId=%u)!"), Address, Heap, HeapSpec.Name, Time, CallstackId);
 		}
 	}
 }
@@ -2683,9 +2740,12 @@ const IAllocationsProvider* ReadAllocationsProvider(const IAnalysisSession& Sess
 #undef INSIGHTS_SLA_USE_ADDRESS_MAP
 #undef INSIGHTS_USE_LAST_ALLOC
 #undef INSIGHTS_VALIDATE_ALLOC_EVENTS
-#undef INSIGHTS_REMAP_INVALID_ALLOC_EVENTS
-#undef INSIGHTS_REMAP_INVALID_FREE_EVENTS
+#undef INSIGHTS_DOUBLE_ALLOC_FREE_PREVIOUS
+#undef INSIGHTS_VALIDATE_FREE_EVENTS
+#undef INSIGHTS_VALIDATE_HEAP_MARK_ALLOC_EVENTS
+#undef INSIGHTS_VALIDATE_HEAP_UNMARK_ALLOC_EVENTS
 #undef INSIGHTS_DEBUG_METADATA
+#undef INSIGHTS_WARNINGS_FOR_NULLPTR_ALLOCS
 #undef INSIGHTS_FILTER_EVENTS_ENABLED
 #undef INSIGHTS_FILTER_EVENT
 #undef INSIGHTS_DEBUG_WATCH
