[Insights] Memory Insights / Modules: Fixed duplicate modules issue (when a module is loaded/unloaded multiple times with the same name, base address and size).

diff --git a/Engine/Source/Developer/TraceServices/Private/Model/ModuleProvider.cpp b/Engine/Source/Developer/TraceServices/Private/Model/ModuleProvider.cpp
--- a/Engine/Source/Developer/TraceServices/Private/Model/ModuleProvider.cpp
+++ b/Engine/Source/Developer/TraceServices/Private/Model/ModuleProvider.cpp
@@ -193,9 +193,9 @@ template <typename SymbolResolverType>
 void TModuleProvider<SymbolResolverType>::EnumerateModules(uint32 Start, TFunctionRef<void(const FModule& Module)> Callback) const
 {
 	FReadScopeLock _(ModulesLock);
-	for (uint32 i = Start; i < Modules.Num(); ++i)
+	for (uint32 ModuleIndex = Start; ModuleIndex < Modules.Num(); ++ModuleIndex)
 	{
-		Callback(Modules[i]);
+		Callback(Modules[ModuleIndex]);
 	}
 }
 
@@ -303,18 +303,31 @@ void TModuleProvider<SymbolResolverType>::GetStats(FStats* OutStats) const
 ////////////////////////////////////////////////////////////////////////////////////////////////////
 
 template<typename SymbolResolverType>
-void TModuleProvider<SymbolResolverType>::OnModuleLoad(const FStringView& Module, uint64 Base, uint32 Size, const uint8* ImageId, uint32 ImageIdSize)
+void TModuleProvider<SymbolResolverType>::OnModuleLoad(const FStringView& ModuleName, uint64 Base, uint32 Size, const uint8* ImageId, uint32 ImageIdSize)
 {
-	if (Module.Len() == 0)
+	if (ModuleName.Len() == 0)
 	{
 		return;
 	}
 
-	const FStringView NameTemp = FPathViews::GetCleanFilename(Module);
+	const FStringView NameTemp = FPathViews::GetCleanFilename(ModuleName);
 	const TCHAR* Name = Session.StoreString(NameTemp);
-	const TCHAR* FullName = Session.StoreString(Module);
+	const TCHAR* FullName = Session.StoreString(ModuleName);
 
 	FWriteScopeLock _(ModulesLock);
+
+	// Check if the module was already added.
+	for (uint32 ModuleIndex = 0; ModuleIndex < Modules.Num(); ++ModuleIndex)
+	{
+		const FModule& LoadedModule = Modules[ModuleIndex];
+		if (LoadedModule.Name == Name && // only comparing pointers from session's string store
+			LoadedModule.Base == Base &&
+			LoadedModule.Size == Size)
+		{
+			return;
+		}
+	}
+
 	FModule& NewModule = Modules.EmplaceBack(Name, FullName, Base, Size, EModuleStatus::Discovered);
 
 	// The number of cached symbols for this module is equal to the number
