[Insights] Empty spaces are now imported as 0 for numeric columns.

diff --git a/Engine/Source/Developer/TraceServices/Private/Model/TableImportTask.cpp b/Engine/Source/Developer/TraceServices/Private/Model/TableImportTask.cpp
--- a/Engine/Source/Developer/TraceServices/Private/Model/TableImportTask.cpp
+++ b/Engine/Source/Developer/TraceServices/Private/Model/TableImportTask.cpp
@@ -125,7 +125,12 @@ bool FTableImportTask::CreateLayout(const FString& Line)
 			ColumnNames[Index] = FString::Format(TEXT("Column {0}"), {Index});
 		}
 
-		if (Values[Index].IsNumeric())
+		if (Values[Index].IsEmpty())
+		{
+			// If the first line has an empty value assume it is an int, the most restrictive type and downgrade if we encounter other types.  
+			Layout.AddColumn<uint32>(*ColumnNames[Index], ProjectorFunc, TableColumnDisplayHint_Summable);
+		}
+		else if (Values[Index].IsNumeric())
 		{
 			if (Values[Index].Contains(TEXT(".")))
 			{
@@ -150,6 +155,9 @@ bool FTableImportTask::ParseData(TArray<FString>& Lines)
 	TTableLayout<FImportTableRow>& Layout = Table->EditLayout();
 	bool Restart = false;
 
+	TArray<bool> HasNonEmptyValues;
+	HasNonEmptyValues.AddDefaulted(ColumnNames.Num());
+
 	for (int32 LineIndex = 1; LineIndex < Lines.Num(); ++LineIndex)
 	{
 		TArray<FString> Values;
@@ -169,23 +177,35 @@ bool FTableImportTask::ParseData(TArray<FString>& Lines)
 			const TCHAR* Value = *Values[ValueIndex];
 			if (ColumnType == TableColumnType_CString)
 			{
+				HasNonEmptyValues[ValueIndex] = true;
 				const TCHAR* StoredValue = Table->GetStringStore().Store(Value);
 				NewRow.SetValue(ValueIndex, StoredValue);
 			}
 			else if (ColumnType == TableColumnType_Double)
 			{
-				if (!Values[ValueIndex].IsNumeric())
+				if (Values[ValueIndex].IsEmpty())
+				{
+					NewRow.SetValue(ValueIndex, 0.0f);
+					continue;
+				}
+				else if (!Values[ValueIndex].IsNumeric())
 				{
 					Layout.SetColumnType(ValueIndex, TableColumnType_CString);
 					Restart = true;
 					break;
 				}
 
+				HasNonEmptyValues[ValueIndex] = true;
 				NewRow.SetValue(ValueIndex, FCString::Atod(Value));
 			}
 			else if (ColumnType == TableColumnType_Int)
 			{
-				if (!Values[ValueIndex].IsNumeric())
+				if (Values[ValueIndex].IsEmpty())
+				{
+					NewRow.SetValue(ValueIndex, 0.0f);
+					continue;
+				}
+				else if (!Values[ValueIndex].IsNumeric())
 				{
 					Layout.SetColumnType(ValueIndex, TableColumnType_CString);
 					Restart = true;
@@ -198,6 +218,7 @@ bool FTableImportTask::ParseData(TArray<FString>& Lines)
 					break;
 				}
 
+				HasNonEmptyValues[ValueIndex] = true;
 				NewRow.SetValue(ValueIndex, FCString::Atoi64(Value));
 			}
 		}
@@ -209,10 +230,29 @@ bool FTableImportTask::ParseData(TArray<FString>& Lines)
 			Table = NewTable;
 			ParseData(Lines);
 
-			break;
+			return true;
 		}
 	}
 
+	// If we have columns with only empty values, switch their type to string and reprocess the data.
+	const TCHAR* EmptyValue = Table->GetStringStore().Store(TEXT(""));
+	for (int32 Index = 0; Index < HasNonEmptyValues.Num(); ++Index)
+	{
+		if (HasNonEmptyValues[Index] == false)
+		{
+			Layout.SetColumnType(Index, TableColumnType_CString);
+			Restart = true;
+		}
+	}
+
+	if (Restart)
+	{
+		TSharedPtr<TImportTable<FImportTableRow >> NewTable = MakeShared<TImportTable<FImportTableRow>>();
+		NewTable->EditLayout() = Layout;
+		Table = NewTable;
+		ParseData(Lines);
+	}
+
 	return true;
 }
 
