Enable SymsLib resolver for Linux/Mac

This change enables SymsLib resolver in TraceAnalysis for Mac and Linux. In order for this to be useful we also need a way to remap the stripped dwarf binaries on those platforms to the full ".debug" binaries.

diff --git a/Engine/Source/Developer/TraceServices/Private/Model/ModuleProvider.cpp b/Engine/Source/Developer/TraceServices/Private/Model/ModuleProvider.cpp
--- a/Engine/Source/Developer/TraceServices/Private/Model/ModuleProvider.cpp
+++ b/Engine/Source/Developer/TraceServices/Private/Model/ModuleProvider.cpp
@@ -515,12 +515,12 @@ TSharedPtr<IModuleAnalysisProvider> CreateModuleProvider(IAnalysisSession& InSes
 {
 	TSharedPtr<IModuleAnalysisProvider> Provider;
 
-#if PLATFORM_WINDOWS && USE_SYMSLIB
+#if USE_SYMSLIB
 	if (!Provider && (InSymbolFormat.Equals("pdb") || InSymbolFormat.Equals("dwarf")))
 	{
 		Provider = MakeShared<TModuleProvider<FSymslibResolver>>(InSession);
 	}
-#endif // PLATFORM_WINDOWS && USE_SYMSLIB
+#endif //USE_SYMSLIB
 #if PLATFORM_WINDOWS && USE_DBGHELP
 	if (!Provider && InSymbolFormat.Equals("pdb"))
 	{
diff --git a/Engine/Source/Developer/TraceServices/Private/Model/SymslibResolver.cpp b/Engine/Source/Developer/TraceServices/Private/Model/SymslibResolver.cpp
--- a/Engine/Source/Developer/TraceServices/Private/Model/SymslibResolver.cpp
+++ b/Engine/Source/Developer/TraceServices/Private/Model/SymslibResolver.cpp
@@ -2,8 +2,6 @@
 
 #include "SymslibResolver.h"
 
-#if PLATFORM_WINDOWS
-
 #include "Algo/ForEach.h"
 #include "Algo/Sort.h"
 #include "Async/MappedFileHandle.h"
@@ -36,6 +34,8 @@ DEFINE_LOG_CATEGORY_STATIC(LogSymslib, Log, All);
 		#include "DbgHelp.h"
 		#include <Microsoft/HideMicrosoftPlatformTypes.h>
 	#endif
+#else
+	#define USE_DBG_HELP_UNDECORATOR 0
 #endif
 
 /////////////////////////////////////////////////////////////////////
@@ -71,10 +71,23 @@ namespace
 		}
 	};
 
-	static FString FindBinaryFileInPath(const FString& File, const FString& SearchPath)
+	static FString FindBinaryFileInPath(const FString& File, const FString& SearchPath, const FString& Platform)
 	{
 		IPlatformFile* PlatformFile = &FPlatformFileManager::Get().GetPlatformFile();
 
+		// On Linux and Mac find the non-stripped binary if available
+		// todo: We cannot actually rely on the Platform string being set yet, we can enable that check when that's part of trace files metadata.
+		// 		 for now we always override the binary if a .debug file exists
+		//if (Platform.Equals(TEXT("Linux")) || Platform.Equals(TEXT("Mac")))
+		{
+			FString NonStrippedFile = FPaths::SetExtension(File, TEXT("debug"));
+			FString Result = FPaths::Combine(SearchPath, NonStrippedFile);
+			if (PlatformFile->FileExists(*Result))
+			{
+				return Result;
+			}
+		}
+
 		// If the search path is an absolute path to a file use this. Filenames does
 		// not need to to match (e.g. eboot.bin <-> gamename.self)
 		if (PlatformFile->FileExists(*SearchPath))
@@ -231,7 +244,7 @@ namespace
 		}
 
 		// First lookup file in symbol path
-		FString BinaryPath = FindBinaryFileInPath(FilePath, SearchPath);
+		FString BinaryPath = FindBinaryFileInPath(FilePath, SearchPath, Platform);
 		if (BinaryPath.IsEmpty() && !Platform.IsEmpty())
 		{
 			BinaryPath = FindFileInEngineFolder(FilePath, SearchPath, Platform, AppName);
@@ -1189,7 +1202,7 @@ bool FSymslibResolver::ResolveSymbol(uint64 Address, FResolvedSymbol& Target, FS
 	}
 
 	// this includes skipping symbols without name (empty string)
-	if (!SymsSymbol || !SourceFilePersistent || SymsSymbol && SymsSymbol->Name[0] == 0)
+	if (!SymsSymbol || !SourceFilePersistent || (SymsSymbol && SymsSymbol->Name[0] == 0))
 	{
 		UpdateResolvedSymbol(Target,
 			ESymbolQueryResult::NotFound,
@@ -1236,5 +1249,3 @@ void FSymslibResolver::WaitForTasks()
 /////////////////////////////////////////////////////////////////////
 
 } // namespace TraceServices
-
-#endif // PLATFORM_WINDOWS
diff --git a/Engine/Source/Developer/TraceServices/Private/Model/SymslibResolver.h b/Engine/Source/Developer/TraceServices/Private/Model/SymslibResolver.h
--- a/Engine/Source/Developer/TraceServices/Private/Model/SymslibResolver.h
+++ b/Engine/Source/Developer/TraceServices/Private/Model/SymslibResolver.h
@@ -1,14 +1,11 @@
 // Copyright Epic Games, Inc. All Rights Reserved.
 #pragma once
 
-#include "HAL/Platform.h"
-
-#if PLATFORM_WINDOWS
-
 #include "Async/MappedFileHandle.h"
 #include "Async/TaskGraphInterfaces.h"
 #include "Common/PagedArray.h"
 #include "HAL/CriticalSection.h"
+#include "HAL/Platform.h"
 #include "TraceServices/Model/AnalysisSession.h"
 #include "TraceServices/Model/Modules.h"
 #include <atomic>
@@ -130,4 +127,3 @@ private:
 
 } // namespace TraceServices
 
-#endif // PLATFORM_WINDOWS
