Disable Mac for SymsLibResolver

We don't yet have binary version of the library sybmitted.

diff --git a/Engine/Source/Developer/TraceServices/Private/Model/ModuleProvider.cpp b/Engine/Source/Developer/TraceServices/Private/Model/ModuleProvider.cpp
--- a/Engine/Source/Developer/TraceServices/Private/Model/ModuleProvider.cpp
+++ b/Engine/Source/Developer/TraceServices/Private/Model/ModuleProvider.cpp
@@ -515,7 +515,7 @@ TSharedPtr<IModuleAnalysisProvider> CreateModuleProvider(IAnalysisSession& InSes
 {
 	TSharedPtr<IModuleAnalysisProvider> Provider;
 
-#if USE_SYMSLIB
+#if USE_SYMSLIB && UE_SYMSLIB_AVAILABLE
 	if (!Provider && (InSymbolFormat.Equals("pdb") || InSymbolFormat.Equals("dwarf")))
 	{
 		Provider = MakeShared<TModuleProvider<FSymslibResolver>>(InSession);
diff --git a/Engine/Source/Developer/TraceServices/Private/Model/SymslibResolver.cpp b/Engine/Source/Developer/TraceServices/Private/Model/SymslibResolver.cpp
--- a/Engine/Source/Developer/TraceServices/Private/Model/SymslibResolver.cpp
+++ b/Engine/Source/Developer/TraceServices/Private/Model/SymslibResolver.cpp
@@ -2,6 +2,8 @@
 
 #include "SymslibResolver.h"
 
+#if UE_SYMSLIB_AVAILABLE
+
 #include "Algo/ForEach.h"
 #include "Algo/Sort.h"
 #include "Async/MappedFileHandle.h"
@@ -1249,3 +1251,5 @@ void FSymslibResolver::WaitForTasks()
 /////////////////////////////////////////////////////////////////////
 
 } // namespace TraceServices
+
+#endif // UE_SYMSLIB_AVAILABLE
diff --git a/Engine/Source/Developer/TraceServices/Private/Model/SymslibResolver.h b/Engine/Source/Developer/TraceServices/Private/Model/SymslibResolver.h
--- a/Engine/Source/Developer/TraceServices/Private/Model/SymslibResolver.h
+++ b/Engine/Source/Developer/TraceServices/Private/Model/SymslibResolver.h
@@ -1,6 +1,15 @@
 // Copyright Epic Games, Inc. All Rights Reserved.
 #pragma once
 
+// We currently only have symslib available on selected platforms.
+#if PLATFORM_WINDOWS || PLATFORM_LINUX
+	#define UE_SYMSLIB_AVAILABLE 1
+#else
+	#define UE_SYMSLIB_AVAILABLE 0
+#endif
+
+#if UE_SYMSLIB_AVAILABLE
+
 #include "Async/MappedFileHandle.h"
 #include "Async/TaskGraphInterfaces.h"
 #include "Common/PagedArray.h"
@@ -127,3 +136,4 @@ private:
 
 } // namespace TraceServices
 
+#endif //UE_SYMSLIB_AVAILABLE
