Enable SymsLib resolver on Mac

diff --git a/Engine/Source/Developer/TraceServices/Private/Model/SymslibResolver.h b/Engine/Source/Developer/TraceServices/Private/Model/SymslibResolver.h
--- a/Engine/Source/Developer/TraceServices/Private/Model/SymslibResolver.h
+++ b/Engine/Source/Developer/TraceServices/Private/Model/SymslibResolver.h
@@ -2,7 +2,7 @@
 #pragma once
 
 // We currently only have symslib available on selected platforms.
-#if PLATFORM_WINDOWS || PLATFORM_LINUX
+#if PLATFORM_DESKTOP
 	#define UE_SYMSLIB_AVAILABLE 1
 #else
 	#define UE_SYMSLIB_AVAILABLE 0
