Back out of recent changes

diff --git a/Engine/Source/Developer/TraceServices/Private/Model/SymslibResolver.h b/Engine/Source/Developer/TraceServices/Private/Model/SymslibResolver.h
--- a/Engine/Source/Developer/TraceServices/Private/Model/SymslibResolver.h
+++ b/Engine/Source/Developer/TraceServices/Private/Model/SymslibResolver.h
@@ -4,7 +4,7 @@
 #include "HAL/Platform.h"
 
 // We currently only have symslib available on selected platforms.
-#if PLATFORM_DESKTOP
+#if PLATFORM_WINDOWS || PLATFORM_LINUX
 	#define UE_SYMSLIB_AVAILABLE 1
 #else
 	#define UE_SYMSLIB_AVAILABLE 0
