[Insights] Fix default int value for table import

diff --git a/Engine/Source/Developer/TraceServices/Private/Model/TableImportTask.cpp b/Engine/Source/Developer/TraceServices/Private/Model/TableImportTask.cpp
--- a/Engine/Source/Developer/TraceServices/Private/Model/TableImportTask.cpp
+++ b/Engine/Source/Developer/TraceServices/Private/Model/TableImportTask.cpp
@@ -202,7 +202,7 @@ bool FTableImportTask::ParseData(TArray<FString>& Lines)
 			{
 				if (Values[ValueIndex].IsEmpty())
 				{
-					NewRow.SetValue(ValueIndex, 0.0f);
+					NewRow.SetValue(ValueIndex, 0);
 					continue;
 				}
 				else if (!Values[ValueIndex].IsNumeric())
