[Insights]
  - LoadTimeTraceAnalysis: Added log messages for warnings and errors detected during analysis. Also added a log message with analysis stats (when the analysis is completed).
  - LoadTimeTraceAnalysis: Added debug log messages for the analyzed trace events (disabled by default).
  - LoadTimeTraceAnalysis: Fixed crash when an AsyncPackage is still referenced after a deleted request.
  - LoadTimeTraceAnalysis: Re-introduced partial backward compatibility with very old traces (including UE 4.2x).

diff --git a/Engine/Source/Developer/TraceServices/Private/Analyzers/LoadTimeTraceAnalysis.cpp b/Engine/Source/Developer/TraceServices/Private/Analyzers/LoadTimeTraceAnalysis.cpp
--- a/Engine/Source/Developer/TraceServices/Private/Analyzers/LoadTimeTraceAnalysis.cpp
+++ b/Engine/Source/Developer/TraceServices/Private/Analyzers/LoadTimeTraceAnalysis.cpp
@@ -11,6 +11,49 @@
 
 #include <limits>
 
+// Enables backward compatibility with previous versions of "LoadTime" trace.
+#define UE_LOAD_TIME_TRACE_ANALYSIS_BACKWARD_COMPATIBILITY 1
+
+// Macro to log the errors detected during "LoadTime" trace analysis.
+#define UE_LOAD_TIME_TRACE_ANALYSIS_ERROR(...) \
+{ \
+	constexpr uint64 MaxErrorCount = 100; \
+	if (++ErrorCount <= MaxErrorCount) \
+	{ \
+		UE_LOG(LogTraceServices, Error, TEXT("[LoadTime] ") __VA_ARGS__); \
+		if (WarningCount == MaxWarningCount) \
+		{ \
+			UE_LOG(LogTraceServices, Error, TEXT("[LoadTime] Too many errors! Further LoadTime analysis errors will not be reported anymore.")); \
+		} \
+	} \
+}
+
+// Macro to log the warnings detected during "LoadTime" trace analysis.
+#define UE_LOAD_TIME_TRACE_ANALYSIS_WARNING(...) \
+{ \
+	constexpr uint64 MaxWarningCount = 100; \
+	if (++WarningCount <= MaxWarningCount) \
+	{ \
+		UE_LOG(LogTraceServices, Warning, TEXT("[LoadTime] ") __VA_ARGS__); \
+		if (WarningCount == MaxWarningCount) \
+		{ \
+			UE_LOG(LogTraceServices, Warning, TEXT("[LoadTime] Too many warnings! Further LoadTime analysis warnings will not be reported anymore.")); \
+		} \
+	} \
+}
+
+// Debug macro to log extra warnings detected during "LoadTime" trace analysis.
+#define UE_LOAD_TIME_TRACE_ANALYSIS_WARNING_EX(...) //UE_LOAD_TIME_TRACE_ANALYSIS_WARNING(__VA_ARGS__)
+
+// Macro to log the important "LoadTime" trace analysis stats.
+#define UE_LOAD_TIME_TRACE_ANALYSIS_LOG(Severity, ...) { UE_LOG(LogTraceServices, Severity, TEXT("[LoadTime] ") __VA_ARGS__); }
+
+// Debug macro to log the name and parameters of "LoadTime" trace events analyzed.
+#define UE_LOAD_TIME_TRACE_ANALYSIS_LOG1(...) //{ UE_LOG(LogTraceServices, Log, TEXT("[LoadTime] ") __VA_ARGS__); }
+
+// Debug macro to log additional info while analyzing a "LoadTime" trace event.
+#define UE_LOAD_TIME_TRACE_ANALYSIS_LOG2(...) //{ UE_LOG(LogTraceServices, Log, TEXT("[LoadTime]   --> ") __VA_ARGS__); }
+
 namespace TraceServices
 {
 
@@ -99,11 +142,34 @@ void FAsyncLoadingTraceAnalyzer::OnAnalysisBegin(const FOnAnalysisContext& Conte
 	Builder.RouteEvent(RouteId_BatchIssued, "IoDispatcher", "BatchIssued");
 	Builder.RouteEvent(RouteId_BatchResolved, "IoDispatcher", "BatchResolved");
 
-	// Backwards compatibility
+#if UE_LOAD_TIME_TRACE_ANALYSIS_BACKWARD_COMPATIBILITY
 	Builder.RouteEvent(RouteId_BeginObjectScope, "LoadTime", "BeginObjectScope");
 	Builder.RouteEvent(RouteId_EndObjectScope, "LoadTime", "EndObjectScope");
 	Builder.RouteEvent(RouteId_BeginPostLoadExport, "LoadTime", "BeginPostLoadExport");
 	Builder.RouteEvent(RouteId_EndPostLoadExport, "LoadTime", "EndPostLoadExport");
+#endif // UE_LOAD_TIME_TRACE_ANALYSIS_BACKWARD_COMPATIBILITY
+}
+
+void FAsyncLoadingTraceAnalyzer::OnAnalysisEnd()
+{
+	if (ErrorCount > 0)
+	{
+		UE_LOAD_TIME_TRACE_ANALYSIS_LOG(Error, TEXT("%llu errors, %llu warnings"), ErrorCount, WarningCount);
+	}
+	else if (WarningCount > 0)
+	{
+		UE_LOAD_TIME_TRACE_ANALYSIS_LOG(Warning, TEXT("%llu errors, %llu warnings"), ErrorCount, WarningCount);
+	}
+
+	UE_LOAD_TIME_TRACE_ANALYSIS_LOG(Log, TEXT("Analysis completed (%d threads, %d active async packages, %d active requests + %d fake requests, %d active linkers, %d active batches, %d exports, %d classes)."),
+		ThreadStatesMap.Num(),
+		ActiveAsyncPackagesMap.Num(),
+		ActiveRequestsMap.Num(),
+		FakeRequestsStack.Num(),
+		ActiveLinkersMap.Num(),
+		ActiveBatchesMap.Num(),
+		ExportsMap.Num(),
+		ClassInfosMap.Num());
 }
 
 bool FAsyncLoadingTraceAnalyzer::OnEvent(uint16 RouteId, EStyle Style, const FOnEventContext& Context)
@@ -113,45 +179,102 @@ bool FAsyncLoadingTraceAnalyzer::OnEvent(uint16 RouteId, EStyle Style, const FOn
 	const auto& EventData = Context.EventData;
 	switch (RouteId)
 	{
+	case RouteId_StartAsyncLoading:
+	{
+		uint64 Cycle = EventData.GetValue<uint64>("Cycle");
+		double Time = Context.EventTime.AsSeconds(Cycle);
+		UE_LOAD_TIME_TRACE_ANALYSIS_LOG1(TEXT("StartAsyncLoading Time=%f"), Time);
+		FAnalysisSessionEditScope _(Session);
+		Session.UpdateDurationSeconds(Time);
+		break;
+	}
+
+	case RouteId_SuspendAsyncLoading:
+	{
+		uint64 Cycle = EventData.GetValue<uint64>("Cycle");
+		double Time = Context.EventTime.AsSeconds(Cycle);
+		UE_LOAD_TIME_TRACE_ANALYSIS_LOG1(TEXT("SuspendAsyncLoading Time=%f"), Time);
+		FAnalysisSessionEditScope _(Session);
+		Session.UpdateDurationSeconds(Time);
+		break;
+	}
+
+	case RouteId_ResumeAsyncLoading:
+	{
+		uint64 Cycle = EventData.GetValue<uint64>("Cycle");
+		double Time = Context.EventTime.AsSeconds(Cycle);
+		UE_LOAD_TIME_TRACE_ANALYSIS_LOG1(TEXT("ResumeAsyncLoading Time=%f"), Time);
+		FAnalysisSessionEditScope _(Session);
+		Session.UpdateDurationSeconds(Time);
+		break;
+	}
+
 	case RouteId_PackageSummary:
 	{
 		uint64 AsyncPackagePtr = EventData.GetValue<uint64>("AsyncPackage");
+		UE_LOAD_TIME_TRACE_ANALYSIS_LOG1(TEXT("PackageSummary AsyncPackage=%llu"), AsyncPackagePtr);
+
 		FAsyncPackageState* AsyncPackage = ActiveAsyncPackagesMap.FindRef(AsyncPackagePtr);
 		if (!AsyncPackage)
 		{
+#if UE_LOAD_TIME_TRACE_ANALYSIS_BACKWARD_COMPATIBILITY
+			if (AsyncPackagePtr == 0)
+			{
+				AsyncPackagePtr = EventData.GetValue<uint64>("Linker");
+			}
+#endif // UE_LOAD_TIME_TRACE_ANALYSIS_BACKWARD_COMPATIBILITY
 			FLinkerLoadState* Linker = ActiveLinkersMap.FindRef(AsyncPackagePtr);
 			if (Linker)
 			{
+				UE_LOAD_TIME_TRACE_ANALYSIS_LOG2(TEXT("%llu (%p) is a Linker; Linker->AsyncPackage=%p"), AsyncPackagePtr, Linker, Linker->AsyncPackage);
 				AsyncPackage = Linker->AsyncPackage;
 			}
 		}
 		if (AsyncPackage)
 		{
-			FPackageInfo* Package = AsyncPackage->PackageInfo;
-			if (Package)
+			FPackageInfo* PackageInfo = AsyncPackage->PackageInfo;
+			if (PackageInfo)
 			{
 				FAnalysisSessionEditScope _(Session);
 				FString PackageName = FTraceAnalyzerUtils::LegacyAttachmentString<TCHAR>("Name", Context);
 				if (PackageName.Len())
 				{
-					Package->Name = Session.StoreString(*PackageName);
+					PackageInfo->Name = Session.StoreString(*PackageName);
 				}
-				FPackageSummaryInfo& Summary = Package->Summary;
+				FPackageSummaryInfo& Summary = PackageInfo->Summary;
 				Summary.TotalHeaderSize = EventData.GetValue<uint32>("TotalHeaderSize");
 				Summary.ImportCount = EventData.GetValue<uint32>("ImportCount");
 				Summary.ExportCount = EventData.GetValue<uint32>("ExportCount");
+				UE_LOAD_TIME_TRACE_ANALYSIS_LOG1(TEXT("TotalHeaderSize=%u ImportCount=%u ExportCount=%u"), Summary.TotalHeaderSize, Summary.ImportCount, Summary.ExportCount);
 			}
+			else
+			{
+				UE_LOAD_TIME_TRACE_ANALYSIS_WARNING(TEXT("Unknown PackageInfo!"));
+			}
+		}
+		else
+		{
+			UE_LOAD_TIME_TRACE_ANALYSIS_WARNING(TEXT("Unknown AsyncPackage!"));
 		}
 		break;
 	}
+
 	case RouteId_BeginProcessSummary:
 	{
-		FPackageInfo* Package = nullptr;
+		uint64 Cycle = EventData.GetValue<uint64>("Cycle");
+		double Time = Context.EventTime.AsSeconds(Cycle);
 		uint64 AsyncPackagePtr = EventData.GetValue<uint64>("AsyncPackage");
+		UE_LOAD_TIME_TRACE_ANALYSIS_LOG1(TEXT("BeginProcessSummary Time=%f AsyncPackage=%llu"), Time, AsyncPackagePtr);
+
+		FPackageInfo* PackageInfo = nullptr;
 		FAsyncPackageState* AsyncPackage = ActiveAsyncPackagesMap.FindRef(AsyncPackagePtr);
 		if (AsyncPackage)
 		{
-			Package = AsyncPackage->PackageInfo;
+			PackageInfo = AsyncPackage->PackageInfo;
+			if (!PackageInfo)
+			{
+				UE_LOAD_TIME_TRACE_ANALYSIS_WARNING(TEXT("Unknown PackageInfo!"));
+			}
 		}
 		else
 		{
@@ -160,108 +283,176 @@ bool FAsyncLoadingTraceAnalyzer::OnEvent(uint16 RouteId, EStyle Style, const FOn
 			{
 				if (Linker->AsyncPackage)
 				{
-					Package = Linker->AsyncPackage->PackageInfo;
-					Linker->PackageInfo = Package;
+					PackageInfo = Linker->AsyncPackage->PackageInfo;
+					Linker->PackageInfo = PackageInfo;
+					if (!PackageInfo)
+					{
+						UE_LOAD_TIME_TRACE_ANALYSIS_WARNING(TEXT("Unknown PackageInfo!"));
+					}
 				}
 				else
 				{
 					{
 						FAnalysisSessionEditScope _(Session);
-						Package = &LoadTimeProfilerProvider.CreatePackage();
+						PackageInfo = &LoadTimeProfilerProvider.CreatePackage();
 					}
 					FAsyncPackageState* FakeAsyncPackageState = new FAsyncPackageState();
-					FakeAsyncPackageState->PackageInfo = Package;
+					FakeAsyncPackageState->PackageInfo = PackageInfo;
 					FakeAsyncPackageState->Linker = Linker;
-					Linker->PackageInfo = Package;
+					Linker->PackageInfo = PackageInfo;
 					Linker->AsyncPackage = FakeAsyncPackageState;
 					Linker->bHasFakeAsyncPackageState = true;
+					UE_LOAD_TIME_TRACE_ANALYSIS_LOG2(TEXT("Linker %llu (%p) : set AsyncPackage = new (fake) FAsyncPackageState %p"), AsyncPackagePtr, Linker, FakeAsyncPackageState);
 				}
 			}
+			else
+			{
+				UE_LOAD_TIME_TRACE_ANALYSIS_WARNING(TEXT("Unknown AsyncPackage!"));
+			}
 		}
+
 		uint32 ThreadId = FTraceAnalyzerUtils::GetThreadIdField(Context);
 		FThreadState& ThreadState = GetThreadState(ThreadId);
-		uint64 Cycle = EventData.GetValue<uint64>("Cycle");
-		double Time = Context.EventTime.AsSeconds(Cycle);
 		{
 			FAnalysisSessionEditScope _(Session);
-			ThreadState.EnterScope(Time, Package);
+			ThreadState.EnterScope(Time, PackageInfo);
+			Session.UpdateDurationSeconds(Time);
 		}
 		break;
 	}
+
 	case RouteId_EndProcessSummary:
 	{
 		uint32 ThreadId = FTraceAnalyzerUtils::GetThreadIdField(Context);
 		FThreadState& ThreadState = GetThreadState(ThreadId);
+
+		uint64 Cycle = EventData.GetValue<uint64>("Cycle");
+		double Time = Context.EventTime.AsSeconds(Cycle);
+		UE_LOAD_TIME_TRACE_ANALYSIS_LOG1(TEXT("EndProcessSummary Time=%f"), Time);
+
 		if (ensure(ThreadState.GetCurrentScopeEventType() == LoadTimeProfilerObjectEventType_None))
 		{
-			uint64 Cycle = EventData.GetValue<uint64>("Cycle");
-			double Time = Context.EventTime.AsSeconds(Cycle);
 			FAnalysisSessionEditScope _(Session);
 			ThreadState.LeaveScope(Time);
+			Session.UpdateDurationSeconds(Time);
 		}
 		break;
 	}
+
 	case RouteId_BeginCreateExport:
 	{
+		uint32 ThreadId = FTraceAnalyzerUtils::GetThreadIdField(Context);
+		FThreadState& ThreadState = GetThreadState(ThreadId);
+
+		uint64 Cycle = EventData.GetValue<uint64>("Cycle");
+		double Time = Context.EventTime.AsSeconds(Cycle);
 		uint64 AsyncPackagePtr = EventData.GetValue<uint64>("AsyncPackage");
+		UE_LOAD_TIME_TRACE_ANALYSIS_LOG1(TEXT("BeginCreateExport Time=%f AsyncPackage=%llu"), Time, AsyncPackagePtr);
+
+		FLinkerLoadState* Linker = nullptr;
 		FAsyncPackageState* AsyncPackage = ActiveAsyncPackagesMap.FindRef(AsyncPackagePtr);
 		if (!AsyncPackage)
 		{
-			FLinkerLoadState* Linker = ActiveLinkersMap.FindRef(AsyncPackagePtr);
+#if UE_LOAD_TIME_TRACE_ANALYSIS_BACKWARD_COMPATIBILITY
+			if (AsyncPackagePtr == 0)
+			{
+				AsyncPackagePtr = EventData.GetValue<uint64>("Linker");
+			}
+#endif // UE_LOAD_TIME_TRACE_ANALYSIS_BACKWARD_COMPATIBILITY
+			Linker = ActiveLinkersMap.FindRef(AsyncPackagePtr);
 			if (Linker)
 			{
+				UE_LOAD_TIME_TRACE_ANALYSIS_LOG2(TEXT("%llu (%p) is a Linker; Linker->AsyncPackage=%p"), AsyncPackagePtr, Linker, Linker->AsyncPackage);
 				AsyncPackage = Linker->AsyncPackage;
 			}
 		}
-		uint32 ThreadId = FTraceAnalyzerUtils::GetThreadIdField(Context);
-		FThreadState& ThreadState = GetThreadState(ThreadId);
-		uint64 Cycle = EventData.GetValue<uint64>("Cycle");
-		double Time = Context.EventTime.AsSeconds(Cycle);
 
 		FAnalysisSessionEditScope _(Session);
 		FPackageExportInfo& Export = LoadTimeProfilerProvider.CreateExport();
-		Export.SerialSize = EventData.GetValue<uint64>("SerialSize"); // Backwards compatibility
+#if UE_LOAD_TIME_TRACE_ANALYSIS_BACKWARD_COMPATIBILITY
+		Export.SerialSize = EventData.GetValue<uint64>("SerialSize");
+#endif // UE_LOAD_TIME_TRACE_ANALYSIS_BACKWARD_COMPATIBILITY
 		if (AsyncPackage)
 		{
-			if (FPackageInfo* PackageInfo = AsyncPackage->PackageInfo)
+			FPackageInfo* PackageInfo = AsyncPackage->PackageInfo;
+			if (PackageInfo)
 			{
 				PackageInfo->Exports.Add(&Export);
 				Export.Package = PackageInfo;
 			}
+			else
+			{
+				UE_LOAD_TIME_TRACE_ANALYSIS_WARNING(TEXT("Unknown AsyncPackage!"));
+			}
+		}
+#if UE_LOAD_TIME_TRACE_ANALYSIS_BACKWARD_COMPATIBILITY
+		else if (Linker)
+		{
+			if (!Linker->PackageInfo)
+			{
+				UE_LOAD_TIME_TRACE_ANALYSIS_WARNING(TEXT("Unknown PackageInfo!"));
+				if (!UnknownPackageInfo)
+				{
+					UnknownPackageInfo = &LoadTimeProfilerProvider.CreatePackage();
+					UnknownPackageInfo->Name = Session.StoreString(TEXT("[unknown]"));
+				}
+				Linker->PackageInfo = UnknownPackageInfo;
+			}
+			Linker->PackageInfo->Exports.Add(&Export);
+			Export.Package = Linker->PackageInfo;
+		}
+#endif // UE_LOAD_TIME_TRACE_ANALYSIS_BACKWARD_COMPATIBILITY
+		else
+		{
+			UE_LOAD_TIME_TRACE_ANALYSIS_WARNING(TEXT("Unknown AsyncPackage!"));
 		}
 		ThreadState.EnterScope(Time, &Export, LoadTimeProfilerObjectEventType_Create);
+		Session.UpdateDurationSeconds(Time);
 		break;
 	}
+
 	case RouteId_EndCreateExport:
 	{
 		uint32 ThreadId = FTraceAnalyzerUtils::GetThreadIdField(Context);
 		FThreadState& ThreadState = GetThreadState(ThreadId);
+
+		uint64 Cycle = EventData.GetValue<uint64>("Cycle");
+		double Time = Context.EventTime.AsSeconds(Cycle);
+		uint64 ObjectPtr = EventData.GetValue<uint64>("Object");
+		UE_LOAD_TIME_TRACE_ANALYSIS_LOG1(TEXT("EndCreateExport Time=%f Object=%llu"), Time, ObjectPtr);
+
 		if (ensure(ThreadState.GetCurrentScopeEventType() == LoadTimeProfilerObjectEventType_Create))
 		{
 			FAnalysisSessionEditScope _(Session);
 			FPackageExportInfo* Export = ThreadState.GetCurrentExportScope();
 			if (Export)
 			{
-				uint64 ObjectPtr = EventData.GetValue<uint64>("Object");
 				ExportsMap.Add(ObjectPtr, Export);
 				const FClassInfo* ObjectClass = GetClassInfo(EventData.GetValue<uint64>("Class"));
 				Export->Class = ObjectClass;
 			}
-
-			uint64 Cycle = EventData.GetValue<uint64>("Cycle");
-			double Time = Context.EventTime.AsSeconds(Cycle);
+			else
+			{
+				UE_LOAD_TIME_TRACE_ANALYSIS_WARNING(TEXT("Unknown export scope!"));
+			}
 			ThreadState.LeaveScope(Time);
+			Session.UpdateDurationSeconds(Time);
 		}
 		break;
 	}
+
 	case RouteId_BeginSerializeExport:
 	{
 		uint32 ThreadId = FTraceAnalyzerUtils::GetThreadIdField(Context);
 		FThreadState& ThreadState = GetThreadState(ThreadId);
+
 		uint64 Cycle = EventData.GetValue<uint64>("Cycle");
 		double Time = Context.EventTime.AsSeconds(Cycle);
 		uint64 ObjectPtr = EventData.GetValue<uint64>("Object");
+		UE_LOAD_TIME_TRACE_ANALYSIS_LOG1(TEXT("BeginSerializeExport Time=%f Object=%llu"), Time, ObjectPtr);
+
 		FPackageExportInfo* Export = ExportsMap.FindRef(ObjectPtr);
+
 		FAnalysisSessionEditScope _(Session);
 		if (Export)
 		{
@@ -275,89 +466,151 @@ bool FAsyncLoadingTraceAnalyzer::OnEvent(uint16 RouteId, EStyle Style, const FOn
 				const_cast<FPackageInfo*>(Export->Package)->TotalExportsSerialSize += SerialSize;
 			}
 		}
+		else
+		{
+			UE_LOAD_TIME_TRACE_ANALYSIS_WARNING_EX(TEXT("Unknown export Object!"));
+		}
 		ThreadState.EnterScope(Time, Export, LoadTimeProfilerObjectEventType_Serialize);
+		Session.UpdateDurationSeconds(Time);
 		break;
 	}
+
 	case RouteId_EndSerializeExport:
 	{
 		uint32 ThreadId = FTraceAnalyzerUtils::GetThreadIdField(Context);
 		FThreadState& ThreadState = GetThreadState(ThreadId);
+
+		uint64 Cycle = EventData.GetValue<uint64>("Cycle");
+		double Time = Context.EventTime.AsSeconds(Cycle);
+		UE_LOAD_TIME_TRACE_ANALYSIS_LOG1(TEXT("EndSerializeExport Time=%f"), Time);
+
 		if (ensure(ThreadState.GetCurrentScopeEventType() == LoadTimeProfilerObjectEventType_Serialize))
 		{
-			uint64 Cycle = EventData.GetValue<uint64>("Cycle");
-			double Time = Context.EventTime.AsSeconds(Cycle);
 			FAnalysisSessionEditScope _(Session);
 			ThreadState.LeaveScope(Time);
+			Session.UpdateDurationSeconds(Time);
 		}
 		break;
 	}
+
 	case RouteId_BeginPostLoad:
 	{
 		uint32 ThreadId = FTraceAnalyzerUtils::GetThreadIdField(Context);
 		FThreadState& ThreadState = GetThreadState(ThreadId);
+
+		UE_LOAD_TIME_TRACE_ANALYSIS_LOG1(TEXT("BeginPostLoad ThreadId=%u --> PostLoadScopeDepth=%lli"), ThreadId, ThreadState.PostLoadScopeDepth + 1);
+
 		++ThreadState.PostLoadScopeDepth;
 		break;
 	}
+
 	case RouteId_EndPostLoad:
 	{
 		uint32 ThreadId = FTraceAnalyzerUtils::GetThreadIdField(Context);
 		FThreadState& ThreadState = GetThreadState(ThreadId);
-		--ThreadState.PostLoadScopeDepth;
+
+		UE_LOAD_TIME_TRACE_ANALYSIS_LOG1(TEXT("EndPostLoad ThreadId=%u --> PostLoadScopeDepth=%lli"), ThreadId, ThreadState.PostLoadScopeDepth - 1);
+
+		if (ThreadState.PostLoadScopeDepth > 0)
+		{
+			--ThreadState.PostLoadScopeDepth;
+		}
+		else
+		{
+			UE_LOAD_TIME_TRACE_ANALYSIS_WARNING(TEXT("Invalid PostLoad scope!"));
+		}
 		break;
 	}
+
+#if UE_LOAD_TIME_TRACE_ANALYSIS_BACKWARD_COMPATIBILITY
 	case RouteId_BeginPostLoadExport:
+#endif // UE_LOAD_TIME_TRACE_ANALYSIS_BACKWARD_COMPATIBILITY
 	case RouteId_BeginPostLoadObject:
 	{
 		uint32 ThreadId = FTraceAnalyzerUtils::GetThreadIdField(Context);
 		FThreadState& ThreadState = GetThreadState(ThreadId);
+
+		uint64 Cycle = EventData.GetValue<uint64>("Cycle");
+		double Time = Context.EventTime.AsSeconds(Cycle);
+		uint64 ObjectPtr = EventData.GetValue<uint64>("Object");
+		UE_LOAD_TIME_TRACE_ANALYSIS_LOG1(TEXT("BeginPostLoadObject ThreadId=%u Time=%f Object=%llu"), ThreadId, Time, ObjectPtr);
+
 		if (RouteId == RouteId_BeginPostLoadObject && !ThreadState.PostLoadScopeDepth)
 		{
+			UE_LOAD_TIME_TRACE_ANALYSIS_WARNING_EX(TEXT("Not in a PostLoad scope!"));
 			break;
 		}
-		uint64 ObjectPtr = EventData.GetValue<uint64>("Object");
+
 		FPackageExportInfo* Export = ExportsMap.FindRef(ObjectPtr);
-		uint64 Cycle = EventData.GetValue<uint64>("Cycle");
-		double Time = Context.EventTime.AsSeconds(Cycle);
+		if (!Export)
+		{
+			UE_LOAD_TIME_TRACE_ANALYSIS_WARNING_EX(TEXT("Unknown export Object!"));
+		}
+
 		FAnalysisSessionEditScope _(Session);
 		ThreadState.EnterScope(Time, Export, LoadTimeProfilerObjectEventType_PostLoad);
+		Session.UpdateDurationSeconds(Time);
 		break;
 	}
+
+#if UE_LOAD_TIME_TRACE_ANALYSIS_BACKWARD_COMPATIBILITY
 	case RouteId_EndPostLoadExport:
+#endif // UE_LOAD_TIME_TRACE_ANALYSIS_BACKWARD_COMPATIBILITY
 	case RouteId_EndPostLoadObject:
 	{
 		uint32 ThreadId = FTraceAnalyzerUtils::GetThreadIdField(Context);
 		FThreadState& ThreadState = GetThreadState(ThreadId);
+
+		uint64 Cycle = EventData.GetValue<uint64>("Cycle");
+		double Time = Context.EventTime.AsSeconds(Cycle);
+		UE_LOAD_TIME_TRACE_ANALYSIS_LOG1(TEXT("EndPostLoadObject ThreadId=%u Time=%f"), ThreadId, Time);
+
 		if (RouteId == RouteId_EndPostLoadObject && !ThreadState.PostLoadScopeDepth)
 		{
+			UE_LOAD_TIME_TRACE_ANALYSIS_WARNING_EX(TEXT("Not in a PostLoad scope!"));
 			break;
 		}
+
 		if (ensure(ThreadState.GetCurrentScopeEventType() == LoadTimeProfilerObjectEventType_PostLoad))
 		{
-			uint64 Cycle = EventData.GetValue<uint64>("Cycle");
-			double Time = Context.EventTime.AsSeconds(Cycle);
 			FAnalysisSessionEditScope _(Session);
 			ThreadState.LeaveScope(Time);
+			Session.UpdateDurationSeconds(Time);
 		}
 		break;
 	}
+
 	case RouteId_BeginRequest:
 	{
-		uint64 RequestId = EventData.GetValue<uint64>("RequestId");
 		uint32 ThreadId = FTraceAnalyzerUtils::GetThreadIdField(Context);
+		FThreadState& ThreadState = GetThreadState(ThreadId);
+
+		uint64 Cycle = EventData.GetValue<uint64>("Cycle");
+		double Time = Context.EventTime.AsSeconds(Cycle);
+		uint64 RequestId = EventData.GetValue<uint64>("RequestId");
+		UE_LOAD_TIME_TRACE_ANALYSIS_LOG1(TEXT("BeginRequest Time=%f RequestId=%llu"), Time, RequestId);
+
+		{
+			FAnalysisSessionEditScope _(Session);
+			Session.UpdateDurationSeconds(Time);
+		}
+
 		FRequestState* RequestState = new FRequestState();
-		RequestState->WallTimeStartCycle = EventData.GetValue<uint64>("Cycle");
+		RequestState->WallTimeStartCycle = Cycle;
 		RequestState->WallTimeEndCycle = 0;
 		RequestState->ThreadId = ThreadId;
-	
+
 		if (RequestId == 0)
 		{
+			UE_LOAD_TIME_TRACE_ANALYSIS_LOG2(TEXT("new (fake) FRequestState %p"), RequestState);
 			FakeRequestsStack.Push(RequestState);
 		}
 		else
 		{
+			UE_LOAD_TIME_TRACE_ANALYSIS_LOG2(TEXT("new FRequestState %p"), RequestState);
 			ActiveRequestsMap.Add(RequestId, RequestState);
 		}
-		FThreadState& ThreadState = GetThreadState(ThreadId);
+
 		TSharedPtr<FRequestGroupState> RequestGroup = ThreadState.RequestGroupStack.Num() ? ThreadState.RequestGroupStack[0] : nullptr;
 		if (!RequestGroup)
 		{
@@ -369,9 +622,19 @@ bool FAsyncLoadingTraceAnalyzer::OnEvent(uint16 RouteId, EStyle Style, const FOn
 		RequestState->Group = RequestGroup;
 		break;
 	}
+
 	case RouteId_EndRequest:
 	{
+		uint64 Cycle = EventData.GetValue<uint64>("Cycle");
+		double Time = Context.EventTime.AsSeconds(Cycle);
 		uint64 RequestId = EventData.GetValue<uint64>("RequestId");
+		UE_LOAD_TIME_TRACE_ANALYSIS_LOG1(TEXT("EndRequest Time=%f RequestId=%llu"), Time, RequestId);
+
+		{
+			FAnalysisSessionEditScope _(Session);
+			Session.UpdateDurationSeconds(Time);
+		}
+
 		FRequestState* RequestState = nullptr;
 		if (RequestId == 0)
 		{
@@ -386,20 +649,40 @@ bool FAsyncLoadingTraceAnalyzer::OnEvent(uint16 RouteId, EStyle Style, const FOn
 		}
 		if (RequestState)
 		{
-			RequestState->WallTimeEndCycle = EventData.GetValue<uint64>("Cycle");
+			RequestState->WallTimeEndCycle = Cycle;
 			RequestState->Group->LatestEndCycle = FMath::Max(RequestState->Group->LatestEndCycle, RequestState->WallTimeEndCycle);
 			--RequestState->Group->ActiveRequestsCount;
 			if (RequestState->Group->LoadRequest && RequestState->Group->ActiveRequestsCount == 0)
 			{
 				FAnalysisSessionEditScope _(Session);
-				RequestState->Group->LoadRequest->EndTime = Context.EventTime.AsSeconds(RequestState->Group->LatestEndCycle);
+				double EndTime = Context.EventTime.AsSeconds(RequestState->Group->LatestEndCycle);
+				RequestState->Group->LoadRequest->EndTime = EndTime;
+				Session.UpdateDurationSeconds(EndTime);
 			}
+
+			for (FAsyncPackageState* State : RequestState->AsyncPackages)
+			{
+				if (State->Request == RequestState)
+				{
+					State->Request = nullptr;
+				}
+			}
+
+			UE_LOAD_TIME_TRACE_ANALYSIS_LOG2(TEXT("delete FRequestState %p"), RequestState);
 			delete RequestState;
 		}
+		else
+		{
+			UE_LOAD_TIME_TRACE_ANALYSIS_WARNING(TEXT("Unknown Request!"));
+		}
 		break;
 	}
+
 	case RouteId_BeginRequestGroup:
 	{
+		uint32 ThreadId = FTraceAnalyzerUtils::GetThreadIdField(Context);
+		FThreadState& ThreadState = GetThreadState(ThreadId);
+
 		FString FormatString;
 		const uint8* FormatArgs;
 		if (EventData.GetString("Format", FormatString))
@@ -412,62 +695,92 @@ bool FAsyncLoadingTraceAnalyzer::OnEvent(uint16 RouteId, EStyle Style, const FOn
 			FormatArgs = EventData.GetAttachment() + (FCString::Strlen(FormatStringPtr) + 1) * sizeof(TCHAR);
 			FormatString = FormatStringPtr;
 		}
+		FFormatArgsHelper::Format(FormatBuffer, FormatBufferSize - 1, TempBuffer, FormatBufferSize - 1, *FormatString, FormatArgs);
+		UE_LOAD_TIME_TRACE_ANALYSIS_LOG1(TEXT("BeginRequestGroup ThreadId=%u Name=\"%s\""), ThreadId, FormatBuffer);
 
-		uint32 ThreadId = FTraceAnalyzerUtils::GetThreadIdField(Context);
-		FThreadState& ThreadState = GetThreadState(ThreadId);
+		TSharedRef<FRequestGroupState> GroupState = MakeShared<FRequestGroupState>();
 
 		FAnalysisSessionEditScope _(Session);
-		TSharedRef<FRequestGroupState> GroupState = MakeShared<FRequestGroupState>();
-		FFormatArgsHelper::Format(FormatBuffer, FormatBufferSize - 1, TempBuffer, FormatBufferSize - 1, *FormatString, FormatArgs);
 		GroupState->Name = Session.StoreString(FormatBuffer);
 		ThreadState.RequestGroupStack.Push(GroupState);
 		break;
 	}
+
 	case RouteId_EndRequestGroup:
 	{
 		uint32 ThreadId = FTraceAnalyzerUtils::GetThreadIdField(Context);
 		FThreadState& ThreadState = GetThreadState(ThreadId);
+
+		UE_LOAD_TIME_TRACE_ANALYSIS_LOG1(TEXT("EndRequestGroup ThreadId=%u"), ThreadId);
+
 		if (ThreadState.RequestGroupStack.Num())
 		{
 			ThreadState.RequestGroupStack.Pop(false);
 		}
+		else
+		{
+			UE_LOAD_TIME_TRACE_ANALYSIS_WARNING(TEXT("Invalid request group!"));
+		}
 		break;
 	}
+
 	case RouteId_NewAsyncPackage:
 	{
 		uint64 AsyncPackagePtr = EventData.GetValue<uint64>("AsyncPackage");
+		UE_LOAD_TIME_TRACE_ANALYSIS_LOG1(TEXT("NewAsyncPackage AsyncPackage=%llu"), AsyncPackagePtr);
 		FAsyncPackageState* AsyncPackageState = new FAsyncPackageState();
+		UE_LOAD_TIME_TRACE_ANALYSIS_LOG2(TEXT("new FAsyncPackageState %p"), AsyncPackageState);
 		{
 			FAnalysisSessionEditScope _(Session);
 			AsyncPackageState->PackageInfo = &LoadTimeProfilerProvider.CreatePackage();
+#if UE_LOAD_TIME_TRACE_ANALYSIS_BACKWARD_COMPATIBILITY
+			if (EventData.GetAttachmentSize())
+			{
+				const TCHAR* PackageName = reinterpret_cast<const TCHAR*>(EventData.GetAttachment());
+				AsyncPackageState->PackageInfo->Name = Session.StoreString(PackageName);
+			}
+#endif // UE_LOAD_TIME_TRACE_ANALYSIS_BACKWARD_COMPATIBILITY
 		}
 		ActiveAsyncPackagesMap.Add(AsyncPackagePtr, AsyncPackageState);
 		break;
 	}
+
 	case RouteId_DestroyAsyncPackage:
 	{
 		uint64 AsyncPackagePtr = EventData.GetValue<uint64>("AsyncPackage");
+		UE_LOAD_TIME_TRACE_ANALYSIS_LOG1(TEXT("DestroyAsyncPackage AsyncPackage=%llu"), AsyncPackagePtr);
 		FAsyncPackageState* AsyncPackageState;
 		if (ActiveAsyncPackagesMap.RemoveAndCopyValue(AsyncPackagePtr, AsyncPackageState))
 		{
 			if (AsyncPackageState->Linker)
 			{
 				AsyncPackageState->Linker->AsyncPackage = nullptr;
+				UE_LOAD_TIME_TRACE_ANALYSIS_LOG2(TEXT("Linker (%p) : reset AsyncPackage"), AsyncPackageState->Linker);
 			}
+			UE_LOAD_TIME_TRACE_ANALYSIS_LOG2(TEXT("delete FAsyncPackageState %p"), AsyncPackageState);
 			delete AsyncPackageState;
 		}
+		else
+		{
+			UE_LOAD_TIME_TRACE_ANALYSIS_WARNING(TEXT("Unknown AsyncPackage!"));
+		}
 		break;
 	}
+
 	case RouteId_NewLinker:
 	{
 		uint64 LinkerPtr = EventData.GetValue<uint64>("Linker");
+		UE_LOAD_TIME_TRACE_ANALYSIS_LOG1(TEXT("NewLinker Linker=%llu"), LinkerPtr);
 		FLinkerLoadState* LinkerState = new FLinkerLoadState();
+		UE_LOAD_TIME_TRACE_ANALYSIS_LOG2(TEXT("new FLinkerLoadState %p"), LinkerState);
 		ActiveLinkersMap.Add(LinkerPtr, LinkerState);
 		break;
 	}
+
 	case RouteId_DestroyLinker:
 	{
 		uint64 LinkerPtr = EventData.GetValue<uint64>("Linker");
+		UE_LOAD_TIME_TRACE_ANALYSIS_LOG1(TEXT("DestroyLinker Linker=%llu"), LinkerPtr);
 		FLinkerLoadState* LinkerState;
 		if (ActiveLinkersMap.RemoveAndCopyValue(LinkerPtr, LinkerState))
 		{
@@ -477,16 +790,24 @@ bool FAsyncLoadingTraceAnalyzer::OnEvent(uint16 RouteId, EStyle Style, const FOn
 			}
 			if (LinkerState->bHasFakeAsyncPackageState)
 			{
+				UE_LOAD_TIME_TRACE_ANALYSIS_LOG2(TEXT("delete (fake) FAsyncPackageState %p"), LinkerState->AsyncPackage);
 				delete LinkerState->AsyncPackage;
 			}
+			UE_LOAD_TIME_TRACE_ANALYSIS_LOG2(TEXT("delete FLinkerLoadState %p"), LinkerState);
 			delete LinkerState;
 		}
+		else
+		{
+			UE_LOAD_TIME_TRACE_ANALYSIS_WARNING(TEXT("Unknown Linker!"));
+		}
 		break;
 	}
+
 	case RouteId_AsyncPackageImportDependency:
 	{
 		uint64 AsyncPackagePtr = EventData.GetValue<uint64>("AsyncPackage");
 		uint64 ImportedAsyncPackagePtr = EventData.GetValue<uint64>("ImportedAsyncPackage");
+		UE_LOAD_TIME_TRACE_ANALYSIS_LOG1(TEXT("AsyncPackageImportDependency AsyncPackage=%llu ImportedAsyncPackage=%llu"), AsyncPackagePtr, ImportedAsyncPackagePtr);
 
 		FAsyncPackageState* AsyncPackage = ActiveAsyncPackagesMap.FindRef(AsyncPackagePtr);
 		FAsyncPackageState* ImportedAsyncPackage = nullptr;
@@ -499,10 +820,13 @@ bool FAsyncLoadingTraceAnalyzer::OnEvent(uint16 RouteId, EStyle Style, const FOn
 			FLinkerLoadState* LinkerState = ActiveLinkersMap.FindRef(AsyncPackagePtr);
 			if (LinkerState)
 			{
+				UE_LOAD_TIME_TRACE_ANALYSIS_LOG2(TEXT("%llu (%p) is a Linker; Linker->AsyncPackage=%p"), AsyncPackagePtr, LinkerState, LinkerState->AsyncPackage);
 				AsyncPackage = LinkerState->AsyncPackage;
+
 				FLinkerLoadState* ImportedLinkerState = ActiveLinkersMap.FindRef(ImportedAsyncPackagePtr);
 				if (ImportedLinkerState)
 				{
+					UE_LOAD_TIME_TRACE_ANALYSIS_LOG2(TEXT("%llu (%p) is a Linker; ImportedLinker->AsyncPackage=%p"), ImportedAsyncPackagePtr, ImportedLinkerState, ImportedLinkerState->AsyncPackage);
 					ImportedAsyncPackage = ImportedLinkerState->AsyncPackage;
 				}
 			}
@@ -522,11 +846,29 @@ bool FAsyncLoadingTraceAnalyzer::OnEvent(uint16 RouteId, EStyle Style, const FOn
 		}
 		break;
 	}
+
 	case RouteId_AsyncPackageRequestAssociation:
 	{
 		uint64 AsyncPackagePtr = EventData.GetValue<uint64>("AsyncPackage");
-		FAsyncPackageState* AsyncPackageState = nullptr;
 		uint64 RequestId = EventData.GetValue<uint64>("RequestId");
+		UE_LOAD_TIME_TRACE_ANALYSIS_LOG1(TEXT("AsyncPackageRequestAssociation AsyncPackage=%llu RequestId=%llu"), AsyncPackagePtr, RequestId);
+
+		FAsyncPackageState* AsyncPackageState = ActiveAsyncPackagesMap.FindRef(AsyncPackagePtr);
+		if (!AsyncPackageState)
+		{
+			FLinkerLoadState* Linker = ActiveLinkersMap.FindRef(AsyncPackagePtr);
+			if (Linker)
+			{
+				UE_LOAD_TIME_TRACE_ANALYSIS_LOG2(TEXT("%llu (%p) is a Linker; Linker->AsyncPackage=%p"), AsyncPackagePtr, Linker, Linker->AsyncPackage);
+				AsyncPackageState = Linker->AsyncPackage;
+			}
+		}
+		if (!AsyncPackageState)
+		{
+			UE_LOAD_TIME_TRACE_ANALYSIS_WARNING(TEXT("Unknown AsyncPackage!"));
+			break;
+		}
+
 		FRequestState* RequestState = nullptr;
 		if (RequestId == 0)
 		{
@@ -539,25 +881,49 @@ bool FAsyncLoadingTraceAnalyzer::OnEvent(uint16 RouteId, EStyle Style, const FOn
 		{
 			RequestState = ActiveRequestsMap.FindRef(RequestId);
 		}
-		if (RequestState)
+		if (!RequestState)
 		{
-			AsyncPackageState = ActiveAsyncPackagesMap.FindRef(AsyncPackagePtr);
-			if (!AsyncPackageState)
-			{
-				FLinkerLoadState* Linker = ActiveLinkersMap.FindRef(AsyncPackagePtr);
-				if (Linker)
-				{
-					AsyncPackageState = Linker->AsyncPackage;
-				}
-			}
-			if (AsyncPackageState)
-			{
-				FAnalysisSessionEditScope _(Session);
-				PackageRequestAssociation(Context, AsyncPackageState, RequestState);
-			}
+			UE_LOAD_TIME_TRACE_ANALYSIS_WARNING(TEXT("Unknown Request!"));
+			break;
+		}
+
+		FAnalysisSessionEditScope _(Session);
+		PackageRequestAssociation(Context, AsyncPackageState, RequestState);
+		break;
+	}
+
+	case RouteId_AsyncPackageLinkerAssociation:
+	{
+		uint64 AsyncPackagePtr = EventData.GetValue<uint64>("AsyncPackage");
+		uint64 LinkerPtr = EventData.GetValue<uint64>("Linker");
+		UE_LOAD_TIME_TRACE_ANALYSIS_LOG1(TEXT("AsyncPackageLinkerAssociation AsyncPackage=%llu Linker=%llu"), AsyncPackagePtr, LinkerPtr);
+
+		FAsyncPackageState* AsyncPackageState = ActiveAsyncPackagesMap.FindRef(AsyncPackagePtr);
+		if (!AsyncPackageState)
+		{
+			UE_LOAD_TIME_TRACE_ANALYSIS_WARNING(TEXT("Unknown AsyncPackage!"));
+			break;
 		}
+
+		FLinkerLoadState* LinkerState = ActiveLinkersMap.FindRef(LinkerPtr);
+		if (!LinkerState)
+		{
+			UE_LOAD_TIME_TRACE_ANALYSIS_WARNING(TEXT("Unknown Linker!"));
+			break;
+		}
+
+		if (LinkerState->PackageInfo)
+		{
+			FAnalysisSessionEditScope _(Session);
+			AsyncPackageState->PackageInfo->Name = LinkerState->PackageInfo->Name;
+			AsyncPackageState->PackageInfo->Summary = LinkerState->PackageInfo->Summary;
+		}
+		LinkerState->AsyncPackage = AsyncPackageState;
+		AsyncPackageState->Linker = LinkerState;
+		UE_LOAD_TIME_TRACE_ANALYSIS_LOG2(TEXT("AsyncPackage %llu (%p) <-> Linker %llu (%p)"), AsyncPackagePtr, AsyncPackageState, LinkerPtr, LinkerState);
 		break;
 	}
+
 	case RouteId_ClassInfo:
 	{
 		FAnalysisSessionEditScope _(Session);
@@ -567,77 +933,98 @@ bool FAsyncLoadingTraceAnalyzer::OnEvent(uint16 RouteId, EStyle Style, const FOn
 		ClassInfosMap.Add(ClassPtr, &ClassInfo);
 		break;
 	}
-	case RouteId_BatchIssued:
+
+	case RouteId_BatchIssued: // IoDispatcher
 	{
-		FAnalysisSessionEditScope _(Session);
 		uint64 Cycle = EventData.GetValue<uint64>("Cycle");
-		uint64 BatchId = EventData.GetValue<uint64>("BatchId");
 		double Time = Context.EventTime.AsSeconds(Cycle);
-		uint64 BatchHandle = LoadTimeProfilerProvider.BeginIoDispatcherBatch(BatchId, Time);
+		uint64 BatchId = EventData.GetValue<uint64>("BatchId");
+
+		UE_LOAD_TIME_TRACE_ANALYSIS_LOG1(TEXT("IoDispatcher.BatchIssued Time=%f BatchId=%llu"), Time, BatchId);
+
+		uint64 BatchHandle;
+		{
+			FAnalysisSessionEditScope _(Session);
+			BatchHandle = LoadTimeProfilerProvider.BeginIoDispatcherBatch(BatchId, Time);
+			Session.UpdateDurationSeconds(Time);
+		}
 		ActiveBatchesMap.Add(BatchId, BatchHandle);
 		break;
 	}
-	case RouteId_BatchResolved:
+
+	case RouteId_BatchResolved: // IoDispatcher
 	{
-		FAnalysisSessionEditScope _(Session);
 		uint64 Cycle = EventData.GetValue<uint64>("Cycle");
+		double Time = Context.EventTime.AsSeconds(Cycle);
 		uint64 BatchId = EventData.GetValue<uint64>("BatchId");
 		uint64 TotalSize = EventData.GetValue<uint64>("TotalSize");
+
+		UE_LOAD_TIME_TRACE_ANALYSIS_LOG1(TEXT("IoDispatcher.BatchResolved Time=%f BatchId=%llu TotalSize=%llu"), Time, BatchId, TotalSize);
+
 		uint64* FindBatchHandle = ActiveBatchesMap.Find(BatchId);
 		if (FindBatchHandle)
 		{
-			double Time = Context.EventTime.AsSeconds(Cycle);
+			FAnalysisSessionEditScope _(Session);
 			LoadTimeProfilerProvider.EndIoDispatcherBatch(*FindBatchHandle, Time, TotalSize);
+			Session.UpdateDurationSeconds(Time);
+		}
+		else
+		{
+			UE_LOAD_TIME_TRACE_ANALYSIS_WARNING(TEXT("Unknown Batch!"));
 		}
 		break;
 	}
+
+#if UE_LOAD_TIME_TRACE_ANALYSIS_BACKWARD_COMPATIBILITY
 	case RouteId_BeginObjectScope:
 	{
 		uint32 ThreadId = FTraceAnalyzerUtils::GetThreadIdField(Context);
 		FThreadState& ThreadState = GetThreadState(ThreadId);
-		uint64 ObjectPtr = EventData.GetValue<uint64>("Object");
-		FPackageExportInfo* Export = ExportsMap.FindRef(ObjectPtr);
-		ELoadTimeProfilerObjectEventType EventType = static_cast<ELoadTimeProfilerObjectEventType>(EventData.GetValue<uint8>("EventType"));
-		//EventType = static_cast<ELoadTimeProfilerObjectEventType>(99); // debug
+
 		uint64 Cycle = EventData.GetValue<uint64>("Cycle");
 		double Time = Context.EventTime.AsSeconds(Cycle);
 
-		FAnalysisSessionEditScope _(Session);
-		ThreadState.EnterScope(Time, Export, EventType);
+		uint64 ObjectPtr = EventData.GetValue<uint64>("Object");
+		ELoadTimeProfilerObjectEventType EventType = static_cast<ELoadTimeProfilerObjectEventType>(EventData.GetValue<uint8>("EventType"));
+
+		UE_LOAD_TIME_TRACE_ANALYSIS_LOG1(TEXT("BeginObjectScope Time=%f Object=%llu EventType=%u"), Time, ObjectPtr, uint32(EventType));
+
+		FPackageExportInfo* Export = ExportsMap.FindRef(ObjectPtr);
+		if (!Export)
+		{
+			UE_LOAD_TIME_TRACE_ANALYSIS_WARNING_EX(TEXT("Unknown export Object!"));
+		}
+
+		//EventType = static_cast<ELoadTimeProfilerObjectEventType>(99); // debug
+		{
+			FAnalysisSessionEditScope _(Session);
+			ThreadState.EnterScope(Time, Export, EventType);
+			Session.UpdateDurationSeconds(Time);
+		}
 		break;
 	}
+#endif // UE_LOAD_TIME_TRACE_ANALYSIS_BACKWARD_COMPATIBILITY
+
+#if UE_LOAD_TIME_TRACE_ANALYSIS_BACKWARD_COMPATIBILITY
 	case RouteId_EndObjectScope:
 	{
 		uint32 ThreadId = FTraceAnalyzerUtils::GetThreadIdField(Context);
 		FThreadState& ThreadState = GetThreadState(ThreadId);
+
+		uint64 Cycle = EventData.GetValue<uint64>("Cycle");
+		double Time = Context.EventTime.AsSeconds(Cycle);
+
+		UE_LOAD_TIME_TRACE_ANALYSIS_LOG1(TEXT("EndObjectScope Time=%f"), Time);
+
 		//if (ensure(ThreadState.GetCurrentExportScopeEventType() == static_cast<ELoadTimeProfilerObjectEventType>(99))) // debug
 		{
-			uint64 Cycle = EventData.GetValue<uint64>("Cycle");
-			double Time = Context.EventTime.AsSeconds(Cycle);
 			FAnalysisSessionEditScope _(Session);
 			ThreadState.LeaveScope(Time);
+			Session.UpdateDurationSeconds(Time);
 		}
 		break;
 	}
-	case RouteId_AsyncPackageLinkerAssociation:
-	{
-		uint64 LinkerPtr = EventData.GetValue<uint64>("Linker");
-		uint64 AsyncPackagePtr = EventData.GetValue<uint64>("AsyncPackage");
-		FAsyncPackageState* AsyncPackageState = ActiveAsyncPackagesMap.FindRef(AsyncPackagePtr);
-		FLinkerLoadState* LinkerState = ActiveLinkersMap.FindRef(LinkerPtr);
-		if (AsyncPackageState && LinkerState)
-		{
-			if (LinkerState->PackageInfo)
-			{
-				FAnalysisSessionEditScope _(Session);
-				AsyncPackageState->PackageInfo->Name = LinkerState->PackageInfo->Name;
-				AsyncPackageState->PackageInfo->Summary = LinkerState->PackageInfo->Summary;
-			}
-			LinkerState->AsyncPackage = AsyncPackageState;
-			AsyncPackageState->Linker = LinkerState;
-		}
-		break;
-	}
+#endif // UE_LOAD_TIME_TRACE_ANALYSIS_BACKWARD_COMPATIBILITY
 	}
 
 	return true;
@@ -708,6 +1095,7 @@ void FAsyncLoadingTraceAnalyzer::PackageRequestAssociation(const FOnEventContext
 {
 	if (!AsyncPackageState->Request)
 	{
+		UE_LOAD_TIME_TRACE_ANALYSIS_LOG2(TEXT("PackageRequestAssociation(AsyncPackageState=%p, RequestState=%p)"), AsyncPackageState, RequestState);
 		RequestState->AsyncPackages.Add(AsyncPackageState);
 		AsyncPackageState->Request = RequestState;
 		FLoadRequest* LoadRequest = RequestState->Group->LoadRequest;
@@ -719,6 +1107,7 @@ void FAsyncLoadingTraceAnalyzer::PackageRequestAssociation(const FOnEventContext
 			LoadRequest->Name = Session.StoreString(*RequestState->Group->Name);
 			LoadRequest->ThreadId = RequestState->ThreadId;
 			RequestState->Group->LoadRequest = LoadRequest;
+			UE_LOAD_TIME_TRACE_ANALYSIS_LOG2(TEXT("Provider.CreateRequest \"%s\""), LoadRequest->Name);
 		}
 		if (AsyncPackageState->PackageInfo)
 		{
@@ -733,3 +1122,11 @@ void FAsyncLoadingTraceAnalyzer::PackageRequestAssociation(const FOnEventContext
 }
 
 } // namespace TraceServices
+
+#undef UE_LOAD_TIME_TRACE_ANALYSIS_BACKWARD_COMPATIBILITY
+#undef UE_LOAD_TIME_TRACE_ANALYSIS_ERROR
+#undef UE_LOAD_TIME_TRACE_ANALYSIS_WARNING
+#undef UE_LOAD_TIME_TRACE_ANALYSIS_WARNING_EX
+#undef UE_LOAD_TIME_TRACE_ANALYSIS_LOG
+#undef UE_LOAD_TIME_TRACE_ANALYSIS_LOG1
+#undef UE_LOAD_TIME_TRACE_ANALYSIS_LOG2
diff --git a/Engine/Source/Developer/TraceServices/Private/Analyzers/LoadTimeTraceAnalysis.h b/Engine/Source/Developer/TraceServices/Private/Analyzers/LoadTimeTraceAnalysis.h
--- a/Engine/Source/Developer/TraceServices/Private/Analyzers/LoadTimeTraceAnalysis.h
+++ b/Engine/Source/Developer/TraceServices/Private/Analyzers/LoadTimeTraceAnalysis.h
@@ -26,8 +26,9 @@ public:
 	virtual ~FAsyncLoadingTraceAnalyzer();
 
 	virtual void OnAnalysisBegin(const FOnAnalysisContext& Context) override;
+	virtual void OnAnalysisEnd() override;
 	virtual bool OnEvent(uint16 RouteId, EStyle Style, const FOnEventContext& Context) override;
-	
+
 private:
 	struct FRequestState;
 	struct FAsyncPackageState;
@@ -79,7 +80,7 @@ private:
 		int64 PostLoadScopeDepth = 0;
 		FLoadTimeProfilerCpuEvent CurrentEvent;
 		TArray<TSharedPtr<FRequestGroupState>> RequestGroupStack;
-		
+
 		FLoadTimeProfilerProvider::CpuTimelineInternal* CpuTimeline;
 
 		void EnterScope(double Time, const FPackageInfo* PackageInfo);
@@ -177,6 +178,9 @@ private:
 	TPointerMap<uint64> ActiveBatchesMap;
 	TMap<uint32, FThreadState*> ThreadStatesMap;
 	TPointerMap<const FClassInfo*> ClassInfosMap;
+	FPackageInfo* UnknownPackageInfo = nullptr;
+	uint64 ErrorCount = 0;
+	uint64 WarningCount = 0;
 };
 
 } // namespace TraceServices
