[Insights] Add infinity and NaN support to table import and table diff.

diff --git a/Engine/Source/Developer/TraceInsights/Private/Insights/Table/ViewModels/TableCellValueFormatter.cpp b/Engine/Source/Developer/TraceInsights/Private/Insights/Table/ViewModels/TableCellValueFormatter.cpp
--- a/Engine/Source/Developer/TraceInsights/Private/Insights/Table/ViewModels/TableCellValueFormatter.cpp
+++ b/Engine/Source/Developer/TraceInsights/Private/Insights/Table/ViewModels/TableCellValueFormatter.cpp
@@ -11,6 +11,8 @@
 #include "Insights/Table/ViewModels/TableCellValueGetter.h"
 #include "Insights/Table/ViewModels/TableColumn.h"
 
+#include <cmath>
+
 #define LOCTEXT_NAMESPACE "Insights::FTableCellValueFormatter"
 
 namespace Insights
@@ -205,7 +207,11 @@ FText FFloatValueFormatterAsNumber::FormatValue(const TOptional<FTableCellValue>
 	if (InValue.IsSet())
 	{
 		const float Value = InValue.GetValue().Float;
-		if (Value == 0.0f)
+		if (std::isnan(Value))
+		{
+			return FText::FromString(TEXT("NaN"));
+		}
+		else if (Value == 0.0f)
 		{
 			return FText::FromString(TEXT("0"));
 		}
@@ -255,7 +261,11 @@ FText FDoubleValueFormatterAsNumber::FormatValue(const TOptional<FTableCellValue
 	if (InValue.IsSet())
 	{
 		const double Value = InValue.GetValue().Double;
-		if (Value == 0.0)
+		if (std::isnan(Value))
+		{
+			return FText::FromString(TEXT("NaN"));
+		}
+		else if (Value == 0.0)
 		{
 			return FText::FromString(TEXT("0"));
 		}
diff --git a/Engine/Source/Developer/TraceServices/Private/Model/TableImportTask.cpp b/Engine/Source/Developer/TraceServices/Private/Model/TableImportTask.cpp
--- a/Engine/Source/Developer/TraceServices/Private/Model/TableImportTask.cpp
+++ b/Engine/Source/Developer/TraceServices/Private/Model/TableImportTask.cpp
@@ -7,6 +7,8 @@
 #include "Misc/FileHelper.h"
 #include "Tasks/Task.h"
 
+#include <limits>
+
 namespace TraceServices
 {
 
@@ -125,14 +127,19 @@ bool FTableImportTask::CreateLayout(const FString& Line)
 			ColumnNames[Index] = FString::Format(TEXT("Column {0}"), {Index});
 		}
 
-		if (Values[Index].IsEmpty())
+		const FString& Value = Values[Index];
+		if (Value.IsEmpty())
 		{
 			// If the first line has an empty value assume it is an int, the most restrictive type and downgrade if we encounter other types.  
 			Layout.AddColumn<uint32>(*ColumnNames[Index], ProjectorFunc, TableColumnDisplayHint_Summable);
 		}
-		else if (Values[Index].IsNumeric())
+		else if (Value.Equals(TEXT("inf")) || Value.Equals(TEXT("-inf")) || Value.Equals(TEXT("infinity")) || Value.Equals(TEXT("-infinity")) || Value.Equals(TEXT("nan")))
+		{
+			Layout.AddColumn<double>(*ColumnNames[Index], ProjectorFunc, TableColumnDisplayHint_Summable);
+		}
+		else if (Value.IsNumeric())
 		{
-			if (Values[Index].Contains(TEXT(".")))
+			if (Value.Contains(TEXT(".")))
 			{
 				Layout.AddColumn<double>(*ColumnNames[Index], ProjectorFunc, TableColumnDisplayHint_Summable);
 			}
@@ -174,44 +181,63 @@ bool FTableImportTask::ParseData(TArray<FString>& Lines)
 		for (int32 ValueIndex = 0; ValueIndex < Values.Num(); ++ValueIndex)
 		{
 			ETableColumnType ColumnType = Layout.GetColumnType(ValueIndex);
-			const TCHAR* Value = *Values[ValueIndex];
+			const FString& Value = Values[ValueIndex];
 			if (ColumnType == TableColumnType_CString)
 			{
 				HasNonEmptyValues[ValueIndex] = true;
-				const TCHAR* StoredValue = Table->GetStringStore().Store(Value);
+				const TCHAR* StoredValue = Table->GetStringStore().Store(*Value);
 				NewRow.SetValue(ValueIndex, StoredValue);
 			}
 			else if (ColumnType == TableColumnType_Double)
 			{
-				if (Values[ValueIndex].IsEmpty())
+				if (Value.IsEmpty())
 				{
 					NewRow.SetValue(ValueIndex, 0.0f);
 					continue;
 				}
-				else if (!Values[ValueIndex].IsNumeric())
+				else if (Value.IsNumeric())
 				{
-					Layout.SetColumnType(ValueIndex, TableColumnType_CString);
-					Restart = true;
-					break;
+					HasNonEmptyValues[ValueIndex] = true;
+					NewRow.SetValue(ValueIndex, FCString::Atod(*Value));
+					continue;
+				}
+				else if (Value.Equals(TEXT("inf")) || Value.Equals(TEXT("infinity")))
+				{
+					HasNonEmptyValues[ValueIndex] = true;
+					NewRow.SetValue(ValueIndex, std::numeric_limits<double>::infinity());
+					continue;
+				}
+				else if (Value.Equals(TEXT("-inf")) || Value.Equals(TEXT("-infinity")))
+				{
+					HasNonEmptyValues[ValueIndex] = true;
+					NewRow.SetValue(ValueIndex, -std::numeric_limits<double>::infinity());
+					continue;
+				}
+				else if (Value.Equals(TEXT("nan")))
+				{
+					HasNonEmptyValues[ValueIndex] = true;
+					NewRow.SetValue(ValueIndex, std::numeric_limits<double>::quiet_NaN());
+					continue;
 				}
 
-				HasNonEmptyValues[ValueIndex] = true;
-				NewRow.SetValue(ValueIndex, FCString::Atod(Value));
+				Layout.SetColumnType(ValueIndex, TableColumnType_CString);
+				Restart = true;
+				break;
 			}
 			else if (ColumnType == TableColumnType_Int)
 			{
-				if (Values[ValueIndex].IsEmpty())
+				if (Value.IsEmpty())
 				{
 					NewRow.SetValue(ValueIndex, 0);
 					continue;
 				}
-				else if (!Values[ValueIndex].IsNumeric())
+				else if (!Value.IsNumeric())
 				{
 					Layout.SetColumnType(ValueIndex, TableColumnType_CString);
 					Restart = true;
 					break;
 				}
-				else if (Values[ValueIndex].Contains(TEXT(".")))
+				else if (Value.Contains(TEXT(".")))
 				{
 					Layout.SetColumnType(ValueIndex, TableColumnType_Double);
 					Restart = true;
@@ -219,7 +245,7 @@ bool FTableImportTask::ParseData(TArray<FString>& Lines)
 				}
 
 				HasNonEmptyValues[ValueIndex] = true;
-				NewRow.SetValue(ValueIndex, FCString::Atoi64(Value));
+				NewRow.SetValue(ValueIndex, FCString::Atoi64(*Value));
 			}
 		}
 
