[Insights]
  - Metadata / MemoryInsights: Fixed mem tags and metadata association with mem allocations for threads where SystemThreadId is not available (ex. missing $Trace.ThreadInfo event). Now, both taggging and metadata uses the "trace asigned" thread ids.

diff --git a/Engine/Source/Developer/TraceServices/Private/Analyzers/AllocationsAnalysis.cpp b/Engine/Source/Developer/TraceServices/Private/Analyzers/AllocationsAnalysis.cpp
--- a/Engine/Source/Developer/TraceServices/Private/Analyzers/AllocationsAnalysis.cpp
+++ b/Engine/Source/Developer/TraceServices/Private/Analyzers/AllocationsAnalysis.cpp
@@ -68,7 +68,7 @@ void FAllocationsAnalyzer::OnAnalysisBegin(const FOnAnalysisContext& Context)
 
 #if INSIGHTS_MEM_TRACE_METADATA_TEST
 	{
-		FMetadataProvider::FEditScopeLock _(MetadataProvider);
+		FProviderEditScopeLock _(MetadataProvider);
 		TagIdMetadataType = MetadataProvider.RegisterMetadataType(TEXT("MemTagId"), sizeof(TagIdType));
 	}
 #endif
@@ -156,6 +156,7 @@ bool FAllocationsAnalyzer::OnEvent(uint16 RouteId, EStyle Style, const FOnEventC
 			//     FAllocEvent Event; // aligned
 			//     EventData.CopyData(&Event, sizeof(Event));
 
+			const uint32 ThreadId = Context.ThreadInfo.GetId();
 			const double Time = GetCurrentTime();
 
 			// CallstackId is optional. If the field is not present CallstackId will be 0 (i.e. "no callstack").
@@ -172,16 +173,12 @@ bool FAllocationsAnalyzer::OnEvent(uint16 RouteId, EStyle Style, const FOnEventC
 			const uint64 Size = (SizeUpper << SizeShift) | static_cast<uint64>(AlignmentPow2_SizeLower & SizeLowerMask);
 			const uint32 Alignment = 1u << (AlignmentPow2_SizeLower >> SizeShift);
 
-			const uint32 TraceThreadId = Context.ThreadInfo.GetId();
-			const uint32 SystemThreadId = Context.ThreadInfo.GetSystemId();
-
 			FProviderEditScopeLock _(AllocationsProvider);
-			AllocationsProvider.SetCurrentThreadId(TraceThreadId, SystemThreadId);
-			AllocationsProvider.EditAlloc(Time, CallstackId, Address, Size, Alignment, RootHeap);
+			AllocationsProvider.EditAlloc(ThreadId, Time, CallstackId, Address, Size, Alignment, RootHeap);
 			if (RouteId == RouteId_ReallocAlloc || RouteId == RouteId_ReallocAllocSystem)
 			{
 				const uint8 Tracker = 0; // We only care about the default tracker for now.
-				AllocationsProvider.EditPopTagFromPtr(SystemThreadId, Tracker);
+				AllocationsProvider.EditPopTagFromPtr(ThreadId, Tracker);
 			}
 			break;
 		}
@@ -196,6 +193,7 @@ bool FAllocationsAnalyzer::OnEvent(uint16 RouteId, EStyle Style, const FOnEventC
 		case RouteId_Free:
 		case RouteId_ReallocFree:
 		{
+			const uint32 ThreadId = Context.ThreadInfo.GetId();
 			const double Time = GetCurrentTime();
 
 			// CallstackId is optional. If the field is not present CallstackId will be 0 (i.e. "no callstack").
@@ -205,22 +203,19 @@ bool FAllocationsAnalyzer::OnEvent(uint16 RouteId, EStyle Style, const FOnEventC
 
 			RootHeap = EventData.GetValue<uint8>("RootHeap", static_cast<uint8>(RootHeap));
 
-			const uint32 TraceThreadId = Context.ThreadInfo.GetId();
-			const uint32 SystemThreadId = Context.ThreadInfo.GetSystemId();
-
 			FProviderEditScopeLock _(AllocationsProvider);
-			AllocationsProvider.SetCurrentThreadId(TraceThreadId, SystemThreadId);
 			if (RouteId == RouteId_ReallocFree || RouteId == RouteId_ReallocFreeSystem)
 			{
 				const uint8 Tracker = 0; // We only care about the default tracker for now.
-				AllocationsProvider.EditPushTagFromPtr(SystemThreadId, Tracker, Address, RootHeap);
+				AllocationsProvider.EditPushTagFromPtr(ThreadId, Tracker, Address, RootHeap);
 			}
-			AllocationsProvider.EditFree(Time, CallstackId, Address, RootHeap);
+			AllocationsProvider.EditFree(ThreadId, Time, CallstackId, Address, RootHeap);
 			break;
 		}
 
 		case RouteId_HeapMarkAlloc:
 		{
+			const uint32 ThreadId = Context.ThreadInfo.GetId();
 			const double Time = GetCurrentTime();
 
 			// CallstackId is optional. If the field is not present CallstackId will be 0 (i.e. "no callstack").
@@ -230,17 +225,14 @@ bool FAllocationsAnalyzer::OnEvent(uint16 RouteId, EStyle Style, const FOnEventC
 			const HeapId Heap = EventData.GetValue<uint16>("Heap", 0);
 			const EMemoryTraceHeapAllocationFlags Flags = EventData.GetValue<EMemoryTraceHeapAllocationFlags>("Flags");
 
-			const uint32 TraceThreadId = Context.ThreadInfo.GetId();
-			const uint32 SystemThreadId = Context.ThreadInfo.GetSystemId();
-
 			FProviderEditScopeLock _(AllocationsProvider);
-			AllocationsProvider.SetCurrentThreadId(TraceThreadId, SystemThreadId);
-			AllocationsProvider.EditMarkAllocationAsHeap(Time, CallstackId, Address, Heap, Flags);
+			AllocationsProvider.EditMarkAllocationAsHeap(ThreadId, Time, CallstackId, Address, Heap, Flags);
 			break;
 		}
 
 		case RouteId_HeapUnmarkAlloc:
 		{
+			const uint32 ThreadId = Context.ThreadInfo.GetId();
 			const double Time = GetCurrentTime();
 
 			// CallstackId is optional. If the field is not present CallstackId will be 0 (i.e. "no callstack").
@@ -249,12 +241,8 @@ bool FAllocationsAnalyzer::OnEvent(uint16 RouteId, EStyle Style, const FOnEventC
 			const uint64 Address = EventData.GetValue<uint64>("Address");
 			const HeapId Heap = EventData.GetValue<uint16>("Heap", 0);
 
-			const uint32 TraceThreadId = Context.ThreadInfo.GetId();
-			const uint32 SystemThreadId = Context.ThreadInfo.GetSystemId();
-
 			FProviderEditScopeLock _(AllocationsProvider);
-			AllocationsProvider.SetCurrentThreadId(TraceThreadId, SystemThreadId);
-			AllocationsProvider.EditUnmarkAllocationAsHeap(Time, CallstackId, Address, Heap);
+			AllocationsProvider.EditUnmarkAllocationAsHeap(ThreadId, Time, CallstackId, Address, Heap);
 			break;
 		}
 
@@ -300,7 +288,7 @@ bool FAllocationsAnalyzer::OnEvent(uint16 RouteId, EStyle Style, const FOnEventC
 
 		case RouteId_MemScopeTag: // "MemoryScope", see UE_MEMSCOPE
 		{
-			const uint32 ThreadId = Context.ThreadInfo.GetSystemId();
+			const uint32 ThreadId = Context.ThreadInfo.GetId();
 			const uint8 Tracker = 0; // We only care about the default tracker for now.
 
 			if (Style == EStyle::EnterScope)
@@ -312,7 +300,7 @@ bool FAllocationsAnalyzer::OnEvent(uint16 RouteId, EStyle Style, const FOnEventC
 				}
 #if INSIGHTS_MEM_TRACE_METADATA_TEST
 				{
-					FMetadataProvider::FEditScopeLock _(MetadataProvider);
+					FProviderEditScopeLock _(MetadataProvider);
 					MetadataProvider.PushScopedMetadata(ThreadId, TagIdMetadataType, (void*)&Tag, sizeof(TagIdType));
 				}
 #endif
@@ -321,7 +309,7 @@ bool FAllocationsAnalyzer::OnEvent(uint16 RouteId, EStyle Style, const FOnEventC
 			{
 #if INSIGHTS_MEM_TRACE_METADATA_TEST
 				{
-					FMetadataProvider::FEditScopeLock _(MetadataProvider);
+					FProviderEditScopeLock _(MetadataProvider);
 					MetadataProvider.PopScopedMetadata(ThreadId, TagIdMetadataType);
 				}
 #endif
@@ -335,7 +323,7 @@ bool FAllocationsAnalyzer::OnEvent(uint16 RouteId, EStyle Style, const FOnEventC
 
 		case RouteId_MemScopePtr: // "MemoryScopePtr", see UE_MEMSCOPE_PTR
 		{
-			const uint32 ThreadId = Context.ThreadInfo.GetSystemId();
+			const uint32 ThreadId = Context.ThreadInfo.GetId();
 			const uint8 Tracker = 0; // We only care about the default tracker for now.
 
 			if (Style == EStyle::EnterScope)
@@ -351,7 +339,7 @@ bool FAllocationsAnalyzer::OnEvent(uint16 RouteId, EStyle Style, const FOnEventC
 				}
 #if INSIGHTS_MEM_TRACE_METADATA_TEST
 				{
-					FMetadataProvider::FEditScopeLock _(MetadataProvider);
+					FProviderEditScopeLock _(MetadataProvider);
 					TagIdType Tag = 0; //TODO: AllocationsProvider.GetTagFromPtr(ThreadId, Tracker, Ptr, RootHeapId);
 					MetadataProvider.PushScopedMetadata(ThreadId, TagIdMetadataType, (void*)&Tag, sizeof(TagIdType));
 				}
@@ -361,7 +349,7 @@ bool FAllocationsAnalyzer::OnEvent(uint16 RouteId, EStyle Style, const FOnEventC
 			{
 #if INSIGHTS_MEM_TRACE_METADATA_TEST
 				{
-					FMetadataProvider::FEditScopeLock _(MetadataProvider);
+					FProviderEditScopeLock _(MetadataProvider);
 					MetadataProvider.PopScopedMetadata(ThreadId, TagIdMetadataType);
 				}
 #endif
diff --git a/Engine/Source/Developer/TraceServices/Private/Analyzers/MetadataAnalysis.cpp b/Engine/Source/Developer/TraceServices/Private/Analyzers/MetadataAnalysis.cpp
--- a/Engine/Source/Developer/TraceServices/Private/Analyzers/MetadataAnalysis.cpp
+++ b/Engine/Source/Developer/TraceServices/Private/Analyzers/MetadataAnalysis.cpp
@@ -65,7 +65,7 @@ bool FMetadataAnalysis::OnEvent(uint16 RouteId, EStyle Style, const FOnEventCont
 	{
 	case RouteId_ClearScope:
 		{
-			const uint32 ThreadId = Context.ThreadInfo.GetSystemId();
+			const uint32 ThreadId = Context.ThreadInfo.GetId();
 			if (Style == EStyle::EnterScope)
 			{
 				MetadataProvider->BeginClearStackScope(ThreadId);
@@ -79,7 +79,7 @@ bool FMetadataAnalysis::OnEvent(uint16 RouteId, EStyle Style, const FOnEventCont
 
 	case RouteId_SaveStack:
 		{
-			const uint32 ThreadId = Context.ThreadInfo.GetSystemId();
+			const uint32 ThreadId = Context.ThreadInfo.GetId();
 			const uint32 RuntimeId = Context.EventData.GetValue<uint32>("Id");
 			MetadataProvider->SaveStack(ThreadId, RuntimeId);
 		}
@@ -87,7 +87,7 @@ bool FMetadataAnalysis::OnEvent(uint16 RouteId, EStyle Style, const FOnEventCont
 
 	case RouteId_RestoreStack:
 		{
-			const uint32 ThreadId = Context.ThreadInfo.GetSystemId();
+			const uint32 ThreadId = Context.ThreadInfo.GetId();
 			if (Style == EStyle::EnterScope)
 			{
 				const uint32 RuntimeId = Context.EventData.GetValue<uint32>("Id");
@@ -109,7 +109,7 @@ bool FMetadataAnalysis::OnEvent(uint16 RouteId, EStyle Style, const FOnEventCont
 			// for this event type and register it with the provider.
 			const uint16 MetadataTypeId = GetOrRegisterType(EventInfo);
 
-			const uint32 ThreadId = Context.ThreadInfo.GetSystemId();
+			const uint32 ThreadId = Context.ThreadInfo.GetId();
 
 			// Push the actual data of the scope to the provider or pop if we are leaving the scope.
 			if (Style == EStyle::EnterScope)
diff --git a/Engine/Source/Developer/TraceServices/Private/Model/AllocationItem.h b/Engine/Source/Developer/TraceServices/Private/Model/AllocationItem.h
--- a/Engine/Source/Developer/TraceServices/Private/Model/AllocationItem.h
+++ b/Engine/Source/Developer/TraceServices/Private/Model/AllocationItem.h
@@ -32,7 +32,8 @@ struct FAllocationItem
 	uint32 EndEventIndex;
 	double StartTime;
 	double EndTime;
-	uint32 ThreadId;
+	uint16 ThreadId;
+	uint16 FreeThreadId;
 	uint32 CallstackId;
 	uint32 FreeCallstackId;
 	uint32 MetadataId;
diff --git a/Engine/Source/Developer/TraceServices/Private/Model/AllocationsProvider.cpp b/Engine/Source/Developer/TraceServices/Private/Model/AllocationsProvider.cpp
--- a/Engine/Source/Developer/TraceServices/Private/Model/AllocationsProvider.cpp
+++ b/Engine/Source/Developer/TraceServices/Private/Model/AllocationsProvider.cpp
@@ -67,7 +67,7 @@
 #if INSIGHTS_FILTER_EVENTS_ENABLED
 namespace TraceServices
 {
-bool FAllocationsProvider::ShouldIgnoreEvent(double Time, uint64 Address, HeapId RootHeapId, uint32 CallstackId)
+bool FAllocationsProvider::ShouldIgnoreEvent(uint32 ThreadId, double Time, uint64 Address, HeapId RootHeapId, uint32 CallstackId)
 {
 	//if (RootHeapId != EMemoryTraceRootHeap::VideoMemory)
 	//{
@@ -90,9 +90,9 @@ bool FAllocationsProvider::ShouldIgnoreEvent(double Time, uint64 Address, HeapId
 	return false;
 }
 } // namespace TraceServices
-#define INSIGHTS_FILTER_EVENT(Time, Address, RootHeapId, CallstackId) { if (ShouldIgnoreEvent(Time, Address, RootHeapId, CallstackId)) return; }
+#define INSIGHTS_FILTER_EVENT(ThreadId, Time, Address, RootHeapId, CallstackId) { if (ShouldIgnoreEvent(ThreadId, Time, Address, RootHeapId, CallstackId)) return; }
 #else // INSIGHTS_FILTER_EVENTS_ENABLED
-#define INSIGHTS_FILTER_EVENT(Time, Address, RootHeapId, CallstackId)
+#define INSIGHTS_FILTER_EVENT(ThreadId, Time, Address, RootHeapId, CallstackId)
 #endif // INSIGHTS_FILTER_EVENTS_ENABLED
 
 ////////////////////////////////////////////////////////////////////////////////////////////////////
@@ -121,39 +121,39 @@ namespace TraceServices
 //#define INSIGHTS_LOGF(Format, ...) { FPlatformMisc::LowLevelOutputDebugStringf(TEXT("[MemAlloc]") Format TEXT("\n"), __VA_ARGS__); }
 
 #if 0 || INSIGHTS_DEBUG_WATCH
-#define INSIGHTS_API_LOGF(Format, Time, ...) \
-	INSIGHTS_LOGF(TEXT("[API][%3u][%f] ") Format, CurrentTraceThreadId, Time, __VA_ARGS__);
+#define INSIGHTS_API_LOGF(ThreadId, Time, Address, Format, ...) \
+	INSIGHTS_LOGF(TEXT("[API][%3u][%f] ") Format, ThreadId, Time, __VA_ARGS__);
 #else
-#define INSIGHTS_API_LOGF(Format, Time, ...)
+#define INSIGHTS_API_LOGF(ThreadId, Time, Address, Format, ...)
 #endif
 
 #if 0 || INSIGHTS_DEBUG_WATCH
-#define INSIGHTS_INDIRECT_API_LOGF(ApiName, Address, Time) \
+#define INSIGHTS_INDIRECT_API_LOGF(ApiName, ThreadId, Time, Address) \
 	INSIGHTS_LOGF(TEXT("[API] --> ") ApiName TEXT(" 0x%llX"), Address);
 #else
-#define INSIGHTS_INDIRECT_API_LOGF(ApiName, Address, Time)
+#define INSIGHTS_INDIRECT_API_LOGF(ApiName, ThreadId, Time, Address)
 #endif
 
 #if INSIGHTS_DEBUG_WATCH
-	#define INSIGHTS_WATCH_API_LOGF(Address, Format, Time, ...) \
+	#define INSIGHTS_WATCH_API_LOGF(ThreadId, Time, Address, Format, ...) \
 	{\
 		if (IsAddressWatched(Address))\
 		{\
 			INSIGHTS_DEBUG_WATCH_FOUND;\
-			INSIGHTS_API_LOGF(Format, Time, __VA_ARGS__);\
+			INSIGHTS_API_LOGF(ThreadId, Time, Address, Format, __VA_ARGS__);\
 		}\
 	}
-	#define INSIGHTS_WATCH_INDIRECT_API_LOGF(ApiName, Address, Time) \
+	#define INSIGHTS_WATCH_INDIRECT_API_LOGF(ApiName, ThreadId, Time, Address) \
 	{\
 		if (IsAddressWatched(Address))\
 		{\
 			INSIGHTS_DEBUG_WATCH_FOUND;\
-			INSIGHTS_INDIRECT_API_LOGF(ApiName, Address, Time);\
+			INSIGHTS_INDIRECT_API_LOGF(ApiName, ThreadId, Time, Address);\
 		}\
 	}
 #else
-	#define INSIGHTS_WATCH_API_LOGF(Address, Format, Time, ...) INSIGHTS_API_LOGF(Format, Time, __VA_ARGS__)
-	#define INSIGHTS_WATCH_INDIRECT_API_LOGF(ApiName, Address, Time) INSIGHTS_INDIRECT_API_LOGF(ApiName, Address, Time)
+	#define INSIGHTS_WATCH_API_LOGF(ThreadId, Time, Address, Format, ...) INSIGHTS_API_LOGF(ThreadId, Time, Address, Format, __VA_ARGS__)
+	#define INSIGHTS_WATCH_INDIRECT_API_LOGF(ApiName, ThreadId, Time, Address) INSIGHTS_INDIRECT_API_LOGF(ApiName, ThreadId, Time, Address)
 #endif // INSIGHTS_DEBUG_WATCH
 
 ////////////////////////////////////////////////////////////////////////////////////////////////////
@@ -1324,7 +1324,7 @@ void FAllocationsProvider::EditInit(double InTime, uint8 InMinAlignment)
 		return;
 	}
 
-	INSIGHTS_API_LOGF(TEXT("Init : MinAlignment=%u)"), InTime, uint32(InMinAlignment));
+	INSIGHTS_API_LOGF(0u, InTime, 0ull, TEXT("Init : MinAlignment=%u)"), uint32(InMinAlignment));
 
 	InitTime = InTime;
 	MinAlignment = InMinAlignment;
@@ -1339,7 +1339,7 @@ void FAllocationsProvider::EditInit(double InTime, uint8 InMinAlignment)
 
 ////////////////////////////////////////////////////////////////////////////////////////////////////
 
-void FAllocationsProvider::EditAlloc(double Time, uint32 CallstackId, uint64 Address, uint64 Size, uint32 Alignment, HeapId RootHeapId)
+void FAllocationsProvider::EditAlloc(uint32 ThreadId, double Time, uint32 CallstackId, uint64 Address, uint64 Size, uint32 Alignment, HeapId RootHeapId)
 {
 	EditAccessCheck();
 
@@ -1348,8 +1348,8 @@ void FAllocationsProvider::EditAlloc(double Time, uint32 CallstackId, uint64 Add
 		return;
 	}
 
-	INSIGHTS_WATCH_API_LOGF(Address, TEXT("Alloc 0x%llX : Size=%llu Alignment=%u RootHeap=%u CallstackId=%u"), Time, Address, Size, Alignment, RootHeapId, CallstackId);
-	INSIGHTS_FILTER_EVENT(Time, Address, RootHeapId, CallstackId);
+	INSIGHTS_WATCH_API_LOGF(ThreadId, Time, Address, TEXT("Alloc 0x%llX : Size=%llu Alignment=%u RootHeap=%u CallstackId=%u"), Address, Size, Alignment, RootHeapId, CallstackId);
+	INSIGHTS_FILTER_EVENT(ThreadId, Time, Address, RootHeapId, CallstackId);
 
 	if (Address == 0 && RootHeapId == EMemoryTraceRootHeap::SystemMemory)
 	{
@@ -1383,7 +1383,8 @@ void FAllocationsProvider::EditAlloc(double Time, uint32 CallstackId, uint64 Add
 
 	AdvanceTimelines(Time);
 
-	const TagIdType Tag = TagTracker.GetCurrentTag(CurrentSystemThreadId, CurrentTracker);
+	const uint8 Tracker = 0;
+	const TagIdType Tag = TagTracker.GetCurrentTag(ThreadId, Tracker);
 
 #if INSIGHTS_VALIDATE_ALLOC_EVENTS
 	FAllocationItem* ExistingAllocationPtr = RootHeap.LiveAllocs->FindRef(Address);
@@ -1398,9 +1399,9 @@ void FAllocationsProvider::EditAlloc(double Time, uint32 CallstackId, uint64 Add
 		}
 
 		// Free the previous allocation.
-		INSIGHTS_WATCH_INDIRECT_API_LOGF(TEXT("Free"), Address, Time);
+		INSIGHTS_WATCH_INDIRECT_API_LOGF(TEXT("Free"), ThreadId, Time, Address);
 		constexpr uint32 FreeCallstackId = 0; // no callstack
-		EditFree(Time, FreeCallstackId, Address, RootHeapId);
+		EditFree(ThreadId, Time, FreeCallstackId, Address, RootHeapId);
 		RootHeap.SbTree->SetTimeForEvent(RootHeap.EventIndex, Time); // for the case where the next event is first event in a new SbTree column after changing the heap alloc
 
 #else // INSIGHTS_DOUBLE_ALLOC_FREE_PREVIOUS
@@ -1422,17 +1423,17 @@ void FAllocationsProvider::EditAlloc(double Time, uint32 CallstackId, uint64 Add
 		uint32 MetadataId = MetadataProvider.InvalidMetadataId;
 		{
 			FProviderEditScopeLock _(MetadataProvider);
-			MetadataId = MetadataProvider.PinAndGetId(CurrentSystemThreadId);
+			MetadataId = MetadataProvider.PinAndGetId(ThreadId);
 #if INSIGHTS_DEBUG_METADATA
 			if (MetadataId != FMetadataProvider::InvalidMetadataId &&
-				!TagTracker.HasTagFromPtrScope(CurrentSystemThreadId, CurrentTracker))
+				!TagTracker.HasTagFromPtrScope(ThreadId, Tracker))
 			{
-				const uint32 MetadataStackSize = MetadataProvider.GetMetadataStackSize(CurrentSystemThreadId, MetadataId);
+				const uint32 MetadataStackSize = MetadataProvider.GetMetadataStackSize(ThreadId, MetadataId);
 				check(MetadataStackSize > 0);
 				uint16 MetaType;
 				const void* MetaData;
 				uint32 MetaDataSize;
-				MetadataProvider.GetMetadata(CurrentSystemThreadId, MetadataId, MetadataStackSize - 1, MetaType, MetaData, MetaDataSize);
+				MetadataProvider.GetMetadata(ThreadId, MetadataId, MetadataStackSize - 1, MetaType, MetaData, MetaDataSize);
 				if (MetaType == 0) // "MemTagId"
 				{
 					TagIdType MetaMemTag = *(TagIdType*)MetaData;
@@ -1449,7 +1450,9 @@ void FAllocationsProvider::EditAlloc(double Time, uint32 CallstackId, uint64 Add
 		Allocation.EndEventIndex = (uint32)-1;
 		Allocation.StartTime = Time;
 		Allocation.EndTime = std::numeric_limits<double>::infinity();
-		Allocation.ThreadId = CurrentSystemThreadId;
+		Allocation.ThreadId = static_cast<uint16>(ThreadId);
+		check(uint32(Allocation.ThreadId) == ThreadId);
+		Allocation.FreeThreadId = 0;
 		Allocation.CallstackId = CallstackId;
 		Allocation.FreeCallstackId = 0; // no callstack yet
 		Allocation.MetadataId = MetadataId;
@@ -1481,7 +1484,7 @@ void FAllocationsProvider::EditAlloc(double Time, uint32 CallstackId, uint64 Add
 
 ////////////////////////////////////////////////////////////////////////////////////////////////////
 
-void FAllocationsProvider::EditFree(double Time, uint32 CallstackId, uint64 Address, HeapId RootHeapId)
+void FAllocationsProvider::EditFree(uint32 ThreadId, double Time, uint32 CallstackId, uint64 Address, HeapId RootHeapId)
 {
 	EditAccessCheck();
 
@@ -1490,8 +1493,8 @@ void FAllocationsProvider::EditFree(double Time, uint32 CallstackId, uint64 Addr
 		return;
 	}
 
-	INSIGHTS_WATCH_API_LOGF(Address, TEXT("Free 0x%llX : RootHeap=%u CallstackId=%u"), Time, Address, RootHeapId, CallstackId);
-	INSIGHTS_FILTER_EVENT(Time, Address, RootHeapId, CallstackId);
+	INSIGHTS_WATCH_API_LOGF(ThreadId, Time, Address, TEXT("Free 0x%llX : RootHeap=%u CallstackId=%u"), Address, RootHeapId, CallstackId);
+	INSIGHTS_FILTER_EVENT(ThreadId, Time, Address, RootHeapId, CallstackId);
 
 	if (Address == 0 && RootHeapId == EMemoryTraceRootHeap::SystemMemory)
 	{
@@ -1530,8 +1533,8 @@ void FAllocationsProvider::EditFree(double Time, uint32 CallstackId, uint64 Addr
 		{
 			INSIGHTS_SLOW_CHECK(AllocationPtr->IsHeap());
 			HeapId Heap = AllocationPtr->RootHeap;
-			INSIGHTS_WATCH_INDIRECT_API_LOGF(TEXT("UnmarkAllocationAsHeap"), Address, Time);
-			EditUnmarkAllocationAsHeap(Time, CallstackId, Address, Heap);
+			INSIGHTS_WATCH_INDIRECT_API_LOGF(TEXT("UnmarkAllocationAsHeap"), ThreadId, Time, Address);
+			EditUnmarkAllocationAsHeap(ThreadId, Time, CallstackId, Address, Heap);
 			RootHeap.SbTree->SetTimeForEvent(RootHeap.EventIndex, Time); // for the case where the next event is first event in a new SbTree column after changing the heap alloc
 			AllocationPtr = RootHeap.LiveAllocs->Remove(Address); // we take ownership of AllocationPtr
 		}
@@ -1552,8 +1555,8 @@ void FAllocationsProvider::EditFree(double Time, uint32 CallstackId, uint64 Addr
 		// Fake the missing alloc.
 		constexpr uint64 FakeAllocSize = 0;
 		constexpr uint32 FakeAllocAlignment = 0;
-		INSIGHTS_WATCH_INDIRECT_API_LOGF(TEXT("Alloc"), Address, Time);
-		EditAlloc(Time, CallstackId, Address, FakeAllocSize, FakeAllocAlignment, RootHeapId);
+		INSIGHTS_WATCH_INDIRECT_API_LOGF(TEXT("Alloc"), ThreadId, Time, Address);
+		EditAlloc(ThreadId, Time, CallstackId, Address, FakeAllocSize, FakeAllocAlignment, RootHeapId);
 		RootHeap.SbTree->SetTimeForEvent(RootHeap.EventIndex, Time); // for the case where the next event is first event in a new SbTree column after adding the fake alloc
 		AllocationPtr = RootHeap.LiveAllocs->Remove(Address); // we take ownership of AllocationPtr
 	}
@@ -1565,6 +1568,9 @@ void FAllocationsProvider::EditFree(double Time, uint32 CallstackId, uint64 Addr
 		AllocationPtr->EndEventIndex = RootHeap.EventIndex;
 		AllocationPtr->EndTime = Time;
 
+		AllocationPtr->FreeThreadId = static_cast<uint16>(ThreadId);
+		check(uint32(AllocationPtr->FreeThreadId) == ThreadId);
+
 		AllocationPtr->FreeCallstackId = CallstackId;
 
 		const uint64 OldSize = AllocationPtr->GetSize();
@@ -1611,7 +1617,7 @@ void FAllocationsProvider::EditHeapSpec(HeapId Id, HeapId ParentId, const FStrin
 {
 	EditAccessCheck();
 
-	INSIGHTS_API_LOGF(TEXT("HeapSpec : Id=%u ParentId=%u Name=\"%.*s\" Flags=0x%X"), -1.0, Id, ParentId, Name.Len(), Name.GetData(), uint32(Flags));
+	INSIGHTS_API_LOGF(0u, 0.0, 0ull, TEXT("HeapSpec : Id=%u ParentId=%u Name=\"%.*s\" Flags=0x%X"), Id, ParentId, Name.Len(), Name.GetData(), uint32(Flags));
 
 	AddHeapSpec(Id, ParentId, Name, Flags);
 }
@@ -1767,7 +1773,7 @@ void FAllocationsProvider::AddHeapSpec(HeapId Id, HeapId ParentId, const FString
 
 ////////////////////////////////////////////////////////////////////////////////////////////////////
 
-void FAllocationsProvider::EditMarkAllocationAsHeap(double Time, uint32 CallstackId, uint64 Address, HeapId Heap, EMemoryTraceHeapAllocationFlags Flags)
+void FAllocationsProvider::EditMarkAllocationAsHeap(uint32 ThreadId, double Time, uint32 CallstackId, uint64 Address, HeapId Heap, EMemoryTraceHeapAllocationFlags Flags)
 {
 	EditAccessCheck();
 
@@ -1776,7 +1782,7 @@ void FAllocationsProvider::EditMarkAllocationAsHeap(double Time, uint32 Callstac
 		return;
 	}
 
-	INSIGHTS_WATCH_API_LOGF(Address, TEXT("MarkAllocAsHeap 0x%llX : Heap=%u Flags=0x%X CallstackId=%u"), Time, Address, Heap, uint32(Flags), CallstackId);
+	INSIGHTS_WATCH_API_LOGF(ThreadId, Time, Address, TEXT("MarkAllocAsHeap 0x%llX : Heap=%u Flags=0x%X CallstackId=%u"), Address, Heap, uint32(Flags), CallstackId);
 
 	if (!IsValidHeap(Heap))
 	{
@@ -1791,7 +1797,7 @@ void FAllocationsProvider::EditMarkAllocationAsHeap(double Time, uint32 Callstac
 	const FHeapSpec& HeapSpec = GetHeapSpecUnchecked(Heap);
 	const FRootHeap& RootHeap = FindParentRootHeapUnchecked(Heap);
 
-	INSIGHTS_FILTER_EVENT(Time, Address, RootHeap.HeapSpec->Id, CallstackId);
+	INSIGHTS_FILTER_EVENT(ThreadId, Time, Address, RootHeap.HeapSpec->Id, CallstackId);
 
 #if 0 // TODO
 	if (Heap == RootHeap.HeapSpec->Id)
@@ -1819,8 +1825,8 @@ void FAllocationsProvider::EditMarkAllocationAsHeap(double Time, uint32 Callstac
 		// Fake the missing alloc.
 		constexpr uint64 FakeAllocSize = 0;
 		constexpr uint32 FakeAllocAlignment = 0;
-		INSIGHTS_WATCH_INDIRECT_API_LOGF(TEXT("Alloc"), Address, Time);
-		EditAlloc(Time, CallstackId, Address, FakeAllocSize, FakeAllocAlignment, RootHeap.HeapSpec->Id);
+		INSIGHTS_WATCH_INDIRECT_API_LOGF(TEXT("Alloc"), ThreadId, Time, Address);
+		EditAlloc(ThreadId, Time, CallstackId, Address, FakeAllocSize, FakeAllocAlignment, RootHeap.HeapSpec->Id);
 		RootHeap.SbTree->SetTimeForEvent(RootHeap.EventIndex, Time); // for the case where the next event is first event in a new SbTree column after adding the fake alloc
 		Alloc = RootHeap.LiveAllocs->Remove(Address); // we take ownership of Alloc
 	}
@@ -1871,7 +1877,7 @@ void FAllocationsProvider::EditMarkAllocationAsHeap(double Time, uint32 Callstac
 
 ////////////////////////////////////////////////////////////////////////////////////////////////////
 
-void FAllocationsProvider::EditUnmarkAllocationAsHeap(double Time, uint32 CallstackId, uint64 Address, HeapId Heap)
+void FAllocationsProvider::EditUnmarkAllocationAsHeap(uint32 ThreadId, double Time, uint32 CallstackId, uint64 Address, HeapId Heap)
 {
 	EditAccessCheck();
 
@@ -1880,7 +1886,7 @@ void FAllocationsProvider::EditUnmarkAllocationAsHeap(double Time, uint32 Callst
 		return;
 	}
 
-	INSIGHTS_WATCH_API_LOGF(Address, TEXT("UnmarkAllocAsHeap 0x%llX : Heap=%u CallstackId=%u"), Time, Address, Heap, CallstackId);
+	INSIGHTS_WATCH_API_LOGF(ThreadId, Time, Address, TEXT("UnmarkAllocAsHeap 0x%llX : Heap=%u CallstackId=%u"), Address, Heap, CallstackId);
 
 	if (!IsValidHeap(Heap))
 	{
@@ -1895,7 +1901,7 @@ void FAllocationsProvider::EditUnmarkAllocationAsHeap(double Time, uint32 Callst
 	const FHeapSpec& HeapSpec = GetHeapSpecUnchecked(Heap);
 	const FRootHeap& RootHeap = FindParentRootHeapUnchecked(Heap);
 
-	INSIGHTS_FILTER_EVENT(Time, Address, RootHeap.HeapSpec->Id, CallstackId);
+	INSIGHTS_FILTER_EVENT(ThreadId, Time, Address, RootHeap.HeapSpec->Id, CallstackId);
 
 #if 0 // TODO
 	if (Heap == RootHeap.HeapSpec->Id)
@@ -1928,8 +1934,8 @@ void FAllocationsProvider::EditUnmarkAllocationAsHeap(double Time, uint32 Callst
 			// Fake the missing alloc.
 			constexpr uint64 FakeAllocSize = 0;
 			constexpr uint32 FakeAllocAlignment = 0;
-			INSIGHTS_WATCH_INDIRECT_API_LOGF(TEXT("Alloc"), Address, Time);
-			EditAlloc(Time, CallstackId, Address, FakeAllocSize, FakeAllocAlignment, RootHeap.HeapSpec->Id);
+			INSIGHTS_WATCH_INDIRECT_API_LOGF(TEXT("Alloc"), ThreadId, Time, Address);
+			EditAlloc(ThreadId, Time, CallstackId, Address, FakeAllocSize, FakeAllocAlignment, RootHeap.HeapSpec->Id);
 			RootHeap.SbTree->SetTimeForEvent(RootHeap.EventIndex, Time); // for the case where the next event is first event in a new SbTree column after adding the fake alloc
 			Alloc = RootHeap.LiveAllocs->Remove(Address); // we take ownership of Alloc
 			check(Alloc != nullptr);
@@ -2022,8 +2028,8 @@ void FAllocationsProvider::EditUnmarkAllocationAsHeap(double Time, uint32 Callst
 				// Free automatically all allocs in this heap.
 				for (const FAllocInHeap& AllocInHeap : AllocsInHeap)
 				{
-					INSIGHTS_WATCH_INDIRECT_API_LOGF(TEXT("Free"), AllocInHeap.Address, Time);
-					EditFree(Time, CallstackId, AllocInHeap.Address, RootHeap.HeapSpec->Id);
+					INSIGHTS_WATCH_INDIRECT_API_LOGF(TEXT("Free"), ThreadId, Time, AllocInHeap.Address);
+					EditFree(ThreadId, Time, CallstackId, AllocInHeap.Address, RootHeap.HeapSpec->Id);
 				}
 			}
 		}
@@ -2037,14 +2043,13 @@ void FAllocationsProvider::EditUnmarkAllocationAsHeap(double Time, uint32 Callst
 
 		// We cannot just unmark the allocation as heap, there is no timestamp support, instead fake a "free"
 		// event and an "alloc" event. Make sure the new allocation retains the tag from the original.
-		CurrentTracker = 1;
-		TagTracker.PushTagFromPtr(CurrentSystemThreadId, CurrentTracker, Alloc->Tag);
-		INSIGHTS_WATCH_INDIRECT_API_LOGF(TEXT("Free"), Address, Time);
-		EditFree(Time, CallstackId, Address, RootHeap.HeapSpec->Id);
-		INSIGHTS_WATCH_INDIRECT_API_LOGF(TEXT("Alloc"), Address, Time);
-		EditAlloc(Time, AllocCallstackId, Address, Size, Alignment, RootHeap.HeapSpec->Id);
-		TagTracker.PopTagFromPtr(CurrentSystemThreadId, CurrentTracker);
-		CurrentTracker = 0;
+		const uint8 Tracker = 1;
+		TagTracker.PushTagFromPtr(ThreadId, Tracker, Alloc->Tag);
+		INSIGHTS_WATCH_INDIRECT_API_LOGF(TEXT("Free"), ThreadId, Time, Address);
+		EditFree(ThreadId, Time, CallstackId, Address, RootHeap.HeapSpec->Id);
+		INSIGHTS_WATCH_INDIRECT_API_LOGF(TEXT("Alloc"), ThreadId, Time, Address);
+		EditAlloc(ThreadId, Time, AllocCallstackId, Address, Size, Alignment, RootHeap.HeapSpec->Id);
+		TagTracker.PopTagFromPtr(ThreadId, Tracker);
 	}
 	else
 	{
@@ -2192,7 +2197,7 @@ void FAllocationsProvider::EditPushTagFromPtr(uint32 ThreadId, uint8 Tracker, ui
 	const TagIdType Tag = Alloc ? Alloc->Tag : 0; // If ptr is not found use "Untagged"
 	TagTracker.PushTagFromPtr(ThreadId, Tracker, Tag);
 
-	INSIGHTS_FILTER_EVENT(0.0, Ptr, RootHeapId, 0);
+	INSIGHTS_FILTER_EVENT(ThreadId, 0.0, Ptr, RootHeapId, 0);
 
 	if (!Alloc)
 	{
diff --git a/Engine/Source/Developer/TraceServices/Private/Model/AllocationsProvider.h b/Engine/Source/Developer/TraceServices/Private/Model/AllocationsProvider.h
--- a/Engine/Source/Developer/TraceServices/Private/Model/AllocationsProvider.h
+++ b/Engine/Source/Developer/TraceServices/Private/Model/AllocationsProvider.h
@@ -341,12 +341,12 @@ public:
 
 	void EditInit(double Time, uint8 MinAlignment);
 
-	void EditAlloc(double Time, uint32 CallstackId, uint64 Address, uint64 Size, uint32 Alignment, HeapId RootHeap);
-	void EditFree(double Time, uint32 CallstackId, uint64 Address, HeapId RootHeap);
+	void EditAlloc(uint32 ThreadId, double Time, uint32 CallstackId, uint64 Address, uint64 Size, uint32 Alignment, HeapId RootHeap);
+	void EditFree(uint32 ThreadId, double Time, uint32 CallstackId, uint64 Address, HeapId RootHeap);
 
 	void EditHeapSpec(HeapId Id, HeapId ParentId, const FStringView& Name, EMemoryTraceHeapFlags Flags);
-	void EditMarkAllocationAsHeap(double Time, uint32 CallstackId, uint64 Address, HeapId Heap, EMemoryTraceHeapAllocationFlags Flags);
-	void EditUnmarkAllocationAsHeap(double Time, uint32 CallstackId, uint64 Address, HeapId Heap);
+	void EditMarkAllocationAsHeap(uint32 ThreadId, double Time, uint32 CallstackId, uint64 Address, HeapId Heap, EMemoryTraceHeapAllocationFlags Flags);
+	void EditUnmarkAllocationAsHeap(uint32 ThreadId, double Time, uint32 CallstackId, uint64 Address, HeapId Heap);
 
 	void EditAddTagSpec(TagIdType Tag, TagIdType ParentTag, const TCHAR* Display) { EditAccessCheck(); TagTracker.AddTagSpec(Tag, ParentTag, Display); }
 	void EditPushTag(uint32 ThreadId, uint8 Tracker, TagIdType Tag);
@@ -356,16 +356,10 @@ public:
 
 	void EditOnAnalysisCompleted(double Time);
 
-	void SetCurrentThreadId(uint32 InThreadId, uint32 InSystemThreadId)
-	{
-		CurrentTraceThreadId = InThreadId;
-		CurrentSystemThreadId = InSystemThreadId;
-	}
-
 	//////////////////////////////////////////////////
 
 private:
-	bool ShouldIgnoreEvent(double Time, uint64 Address, HeapId RootHeapId, uint32 CallstackId);
+	bool ShouldIgnoreEvent(uint32 ThreadId, double Time, uint64 Address, HeapId RootHeapId, uint32 CallstackId);
 	void AdvanceTimelines(double Time);
 	void AddHeapSpec(HeapId Id, HeapId ParentId, const FStringView& Name, EMemoryTraceHeapFlags Flags);
 
@@ -406,11 +400,7 @@ private:
 	uint8 SummarySizeShift = 0;
 	bool bInitialized = false;
 
-	uint32 CurrentTraceThreadId = 0;
-	uint32 CurrentSystemThreadId = 0;
-
 	FTagTracker TagTracker;
-	uint8 CurrentTracker = 0;
 
 	TArray<FHeapSpec*> HeapSpecs;
 	FRootHeap* RootHeaps[MaxRootHeaps] = { nullptr };
