[Insights] Implemented workaround for trace files with broken stats counters events (-trace=stats; i.e. incorrect timestamps).

diff --git a/Engine/Source/Developer/TraceServices/Private/Analyzers/StatsTraceAnalysis.cpp b/Engine/Source/Developer/TraceServices/Private/Analyzers/StatsTraceAnalysis.cpp
--- a/Engine/Source/Developer/TraceServices/Private/Analyzers/StatsTraceAnalysis.cpp
+++ b/Engine/Source/Developer/TraceServices/Private/Analyzers/StatsTraceAnalysis.cpp
@@ -122,10 +122,9 @@ bool FStatsAnalyzer::OnEvent(uint16 RouteId, EStyle Style, const FOnEventContext
 		{
 			uint32 ThreadId = FTraceAnalyzerUtils::GetThreadIdField(Context);
 			TSharedRef<FThreadState> ThreadState = GetThreadState(ThreadId);
-			if (RouteId == RouteId_EventBatch2)
-			{
-				ThreadState->LastCycle = 0; // each batch starts with an absolute timestamp value
-			}
+
+			const uint64 BaseTimestamp = Context.EventTime.AsCycle64() - Context.EventTime.GetTimestamp();
+
 			TArrayView<const uint8> DataView = FTraceAnalyzerUtils::LegacyAttachmentArray("Data", Context);
 			uint64 BufferSize = DataView.Num();
 			const uint8* BufferPtr = DataView.GetData();
@@ -155,6 +154,15 @@ bool FStatsAnalyzer::OnEvent(uint16 RouteId, EStyle Style, const FOnEventContext
 
 				uint8 Op = DecodedIdAndOp & 0x7;
 				uint64 CycleDiff = FTraceAnalyzerUtils::Decode7bit(BufferPtr);
+
+				if (RouteId == RouteId_EventBatch2)
+				{
+					if (CycleDiff >= BaseTimestamp)
+					{
+						ThreadState->LastCycle = 0;
+					}
+				}
+
 				uint64 Cycle = ThreadState->LastCycle + CycleDiff;
 				double Time = Context.EventTime.AsSeconds(Cycle);
 				ThreadState->LastCycle = Cycle;
