[Insights] Added missing TRACESERVICES_API to FProviderLock and IEditableRegionProvider.

diff --git a/Engine/Source/Developer/TraceServices/Public/Common/ProviderLock.h b/Engine/Source/Developer/TraceServices/Public/Common/ProviderLock.h
--- a/Engine/Source/Developer/TraceServices/Public/Common/ProviderLock.h
+++ b/Engine/Source/Developer/TraceServices/Public/Common/ProviderLock.h
@@ -51,7 +51,7 @@ struct TRACESERVICES_API FProviderReadScopeLock
 class FProviderLock
 {
 public:
-	struct FThreadLocalState
+	struct TRACESERVICES_API FThreadLocalState
 	{
 		FProviderLock* Lock;
 		int32 ReadLockCount;
@@ -59,14 +59,14 @@ public:
 	};
 
 public:
-	void ReadAccessCheck(FThreadLocalState& State) const;
-	void WriteAccessCheck(FThreadLocalState& State) const;
+	TRACESERVICES_API void ReadAccessCheck(FThreadLocalState& State) const;
+	TRACESERVICES_API void WriteAccessCheck(FThreadLocalState& State) const;
 
-	void BeginRead(FThreadLocalState& State);
-	void EndRead(FThreadLocalState& State);
+	TRACESERVICES_API void BeginRead(FThreadLocalState& State);
+	TRACESERVICES_API void EndRead(FThreadLocalState& State);
 
-	void BeginWrite(FThreadLocalState& State);
-	void EndWrite(FThreadLocalState& State);
+	TRACESERVICES_API void BeginWrite(FThreadLocalState& State);
+	TRACESERVICES_API void EndWrite(FThreadLocalState& State);
 
 	UE_DEPRECATED(5.3, "Please use the FThreadLocalState overload")
 	void ReadAccessCheck(const FProviderLock* CurrentProviderLock, const int32& CurrentReadProviderLockCount, const int32& CurrentWriteProviderLockCount) const;
diff --git a/Engine/Source/Developer/TraceServices/Public/TraceServices/Model/Regions.h b/Engine/Source/Developer/TraceServices/Public/TraceServices/Model/Regions.h
--- a/Engine/Source/Developer/TraceServices/Public/TraceServices/Model/Regions.h
+++ b/Engine/Source/Developer/TraceServices/Public/TraceServices/Model/Regions.h
@@ -75,40 +75,40 @@ public:
 	virtual void EnumerateLanes(TFunctionRef<void(const FRegionLane&, const int32)> Callback) const = 0;
 
 	/**
-	 * @return A monotonically increasing counter that that changes each time new data is added to the provider.
+	 * @return A monotonically increasing counter that changes each time new data is added to the provider.
 	 * This can be used to detect when to update any (UI-)state dependent on the provider during analysis.
 	 */
 	virtual uint64 GetUpdateCounter() const = 0;
 };
 
-/*
-* The interface to a provider that can consume mutations of region events from a session.
-*/
-class IEditableRegionProvider
+/**
+ * The interface to a provider that can consume mutations of region events from a session.
+ */
+class TRACESERVICES_API IEditableRegionProvider
 	: public IEditableProvider
 {
 public:
 	virtual ~IEditableRegionProvider() override = default;
 
-	/*
-	* Append a new begin event of a region from the trace session.
-	*
-	* @param Name		The string name of the region.
-	* @param Time		The time in seconds of the begin event of this region.
-	*/
+	/**
+	 * Append a new begin event of a region from the trace session.
+	 *
+	 * @param Name		The string name of the region.
+	 * @param Time		The time in seconds of the begin event of this region.
+	 */
 	virtual void AppendRegionBegin(const TCHAR* Name, double Time) = 0;
 
-	/*
-	* Append a new end event of a region from the trace session.
-	*
-	* @param Name		The string name of the region.
-	* @param Time		The time in seconds of the end event of this region.
-	*/
+	/**
+	 * Append a new end event of a region from the trace session.
+	 *
+	 * @param Name		The string name of the region.
+	 * @param Time		The time in seconds of the end event of this region.
+	 */
 	virtual void AppendRegionEnd(const TCHAR* Name, double Time) = 0;
 
 	/**
 	 * Called from the analyzer once all events have been processed.
-	 * Allows postprocessing and error reporting for regions that were never closed.
+	 * Allows post-processing and error reporting for regions that were never closed.
 	 */
 	virtual void OnAnalysisSessionEnded() = 0;
 };
