[Insights] CpuProfilerAnalysis:
  - Refactored code re "pending events" and detection of "still-open" timing events.
  - Added support for "CpuProfiler.EndThread" with a timestamp. "Still-open" timing events will be closed with this timestamp.
  - Removed analysis for unused "CpuProfiler.EndCaptureV2" event.
  - Separated the analysis code for the backward compatibility events.

diff --git a/Engine/Source/Developer/TraceServices/Private/Analyzers/CpuProfilerTraceAnalysis.cpp b/Engine/Source/Developer/TraceServices/Private/Analyzers/CpuProfilerTraceAnalysis.cpp
--- a/Engine/Source/Developer/TraceServices/Private/Analyzers/CpuProfilerTraceAnalysis.cpp
+++ b/Engine/Source/Developer/TraceServices/Private/Analyzers/CpuProfilerTraceAnalysis.cpp
@@ -45,13 +45,12 @@ void FCpuProfilerAnalyzer::OnAnalysisBegin(const FOnAnalysisContext& Context)
 {
 	auto& Builder = Context.InterfaceBuilder;
 
-	Builder.RouteLoggerEvents(RouteId_CpuScope, "Cpu", true);
 	Builder.RouteEvent(RouteId_EventSpec, "CpuProfiler", "EventSpec");
-	Builder.RouteEvent(RouteId_EventBatch, "CpuProfiler", "EventBatch");
 	Builder.RouteEvent(RouteId_EndThread, "CpuProfiler", "EndThread");
-	Builder.RouteEvent(RouteId_EndCapture, "CpuProfiler", "EndCapture");
-	Builder.RouteEvent(RouteId_EventBatchV2, "CpuProfiler", "EventBatchV2");
-	Builder.RouteEvent(RouteId_EndCaptureV2, "CpuProfiler", "EndCaptureV2");
+	Builder.RouteEvent(RouteId_EventBatchV2, "CpuProfiler", "EventBatchV2"); // added in UE 5.1
+	Builder.RouteEvent(RouteId_EventBatch, "CpuProfiler", "EventBatch"); // backward compatibility; removed in UE 5.1
+	Builder.RouteEvent(RouteId_EndCapture, "CpuProfiler", "EndCapture"); // backward compatibility; removed in UE 5.1
+	Builder.RouteLoggerEvents(RouteId_CpuScope, "Cpu", true);
 }
 
 ////////////////////////////////////////////////////////////////////////////////////////////////////
@@ -60,35 +59,18 @@ void FCpuProfilerAnalyzer::OnAnalysisEnd()
 {
 	LLM_SCOPE_BYNAME(TEXT("Insights/FCpuProfilerAnalyzer"));
 
-	const double Timestamp = std::numeric_limits<double>::infinity();
-
 	for (auto& KV : ThreadStatesMap)
 	{
 		FThreadState& ThreadState = *KV.Value;
 
-		if (ThreadState.LastCycle != 0 && ThreadState.LastCycle != ~0)
+		if (ThreadState.LastCycle != ~0ull) // if EndThread is not received
 		{
-			int32 RemainingPending = ThreadState.PendingEvents.Num();
-			if (RemainingPending > 0)
-			{
-				uint64 LastCycle = ThreadState.LastCycle;
-				const FPendingEvent* PendingCursor = ThreadState.PendingEvents.GetData();
-				DispatchPendingEvents(LastCycle, ~0ull, ThreadState, PendingCursor, RemainingPending);
-				check(RemainingPending == 0);
-				ThreadState.PendingEvents.Reset();
-			}
-
-			while (ThreadState.ScopeStack.Num())
-			{
-				ThreadState.ScopeStack.Pop();
-				CPUPROFILER_DEBUG_LOGF(TEXT("[%u] ~E=%llu (%.9f)\n"), ThreadState.ThreadId, ThreadState.LastCycle, Timestamp);
-				ThreadState.Timeline->AppendEndEvent(Timestamp);
-				CPUPROFILER_DEBUG_END_EVENT(Timestamp);
-			}
+			DispatchRemainingPendingEvents(ThreadState);
+			EndOpenEvents(ThreadState, std::numeric_limits<double>::infinity());
 		}
 
-		ensure(ThreadState.PendingEvents.Num() == 0);
-		ensure(ThreadState.ScopeStack.Num() == 0);
+		check(ThreadState.PendingEvents.Num() == 0); // no pending events
+		check(ThreadState.ScopeStack.Num() == 0); // no open events
 	}
 }
 
@@ -164,38 +146,53 @@ bool FCpuProfilerAnalyzer::OnEvent(uint16 RouteId, EStyle Style, const FOnEventC
 			break;
 		}
 
-		int32 RemainingPending = ThreadState.PendingEvents.Num();
-		if (RemainingPending > 0)
+		DispatchRemainingPendingEvents(ThreadState);
+
+		const uint64 Cycle = EventData.GetValue<uint64>("Cycle", ThreadState.LastCycle); // added in UE 5.4
+		if (Cycle != 0)
 		{
-			uint64 LastCycle = ThreadState.LastCycle;
-			const FPendingEvent* PendingCursor = ThreadState.PendingEvents.GetData();
-			DispatchPendingEvents(LastCycle, ~0ull, ThreadState, PendingCursor, RemainingPending);
-			check(RemainingPending == 0);
-			ThreadState.PendingEvents.Reset();
+			ensure(Cycle >= ThreadState.LastCycle);
+			double Timestamp = Context.EventTime.AsSeconds(Cycle);
+			Session.UpdateDurationSeconds(Timestamp);
+			EndOpenEvents(ThreadState, Timestamp);
 		}
 
-		ensure(ThreadState.LastCycle != 0 || ThreadState.ScopeStack.Num() == 0);
+		check(ThreadState.PendingEvents.Num() == 0); // no pending events
+		check(ThreadState.ScopeStack.Num() == 0); // no open events
+
+		ThreadState.LastCycle = ~0ull;
+		break;
+	}
+
+	case RouteId_EventBatchV2:
+	{
+		const uint32 ThreadId = Context.ThreadInfo.GetId();
+		FThreadState& ThreadState = GetThreadState(ThreadId);
+
+		if (ThreadState.LastCycle == ~0ull)
+		{
+			// Ignore timing events received after EndThread.
+			break;
+		}
+
+		TArrayView<const uint8> DataView = Context.EventData.GetArrayView<uint8>("Data");
+		const uint32 BufferSize = DataView.Num();
+		const uint8* BufferPtr = DataView.GetData();
+
+		ProcessBufferV2(Context.EventTime, ThreadState, BufferPtr, BufferSize);
 		if (ThreadState.LastCycle != 0)
 		{
 			double Timestamp = Context.EventTime.AsSeconds(ThreadState.LastCycle);
 			Session.UpdateDurationSeconds(Timestamp);
-			while (ThreadState.ScopeStack.Num())
-			{
-				ThreadState.ScopeStack.Pop();
-				CPUPROFILER_DEBUG_LOGF(TEXT("[%u] ^E=%llu (%.9f)\n"), ThreadState.ThreadId, ThreadState.LastCycle, Timestamp);
-				ThreadState.Timeline->AppendEndEvent(Timestamp);
-				CPUPROFILER_DEBUG_END_EVENT(Timestamp);
-			}
 		}
 
-		ThreadState.LastCycle = ~0ull;
+		TotalEventSize += BufferSize;
+		BytesPerScope = double(TotalEventSize) / double(TotalScopeCount);
 		break;
 	}
 
-	case RouteId_EventBatch:
-	case RouteId_EndCapture:
-	case RouteId_EventBatchV2:
-	case RouteId_EndCaptureV2:
+	case RouteId_EventBatch: // backward compatibility
+	case RouteId_EndCapture: // backward compatibility
 	{
 		const uint32 ThreadId = FTraceAnalyzerUtils::GetThreadIdField(Context);
 		FThreadState& ThreadState = GetThreadState(ThreadId);
@@ -210,29 +207,25 @@ bool FCpuProfilerAnalyzer::OnEvent(uint16 RouteId, EStyle Style, const FOnEventC
 		const uint32 BufferSize = DataView.Num();
 		const uint8* BufferPtr = DataView.GetData();
 
-		uint64 LastCycle;
-		if (RouteId == RouteId_EventBatch || RouteId == RouteId_EndCapture)
+		ProcessBuffer(Context.EventTime, ThreadState, BufferPtr, BufferSize);
+
+		if (RouteId == RouteId_EndCapture)
 		{
-			LastCycle = ProcessBuffer(Context.EventTime, ThreadState, BufferPtr, BufferSize);
+			DispatchRemainingPendingEvents(ThreadState);
+			if (ThreadState.LastCycle != 0)
+			{
+				double Timestamp = Context.EventTime.AsSeconds(ThreadState.LastCycle);
+				Session.UpdateDurationSeconds(Timestamp);
+				EndOpenEvents(ThreadState, Timestamp);
+			}
+			ThreadState.LastCycle = ~0ull;
 		}
 		else
 		{
-			LastCycle = ProcessBufferV2(Context.EventTime, ThreadState, BufferPtr, BufferSize);
-		}
-
-		if (LastCycle)
-		{
-			double LastTimestamp = Context.EventTime.AsSeconds(LastCycle);
-			Session.UpdateDurationSeconds(LastTimestamp);
-			if (RouteId == RouteId_EndCapture || RouteId == RouteId_EndCaptureV2)
+			if (ThreadState.LastCycle != 0)
 			{
-				for (int32 i = ThreadState.ScopeStack.Num(); i--;)
-				{
-					ThreadState.ScopeStack.Pop();
-					CPUPROFILER_DEBUG_LOGF(TEXT("[%u] -E=%llu (%.9f)\n"), ThreadState.ThreadId, LastCycle, LastTimestamp);
-					ThreadState.Timeline->AppendEndEvent(LastTimestamp);
-					CPUPROFILER_DEBUG_END_EVENT(LastTimestamp);
-				}
+				double Timestamp = Context.EventTime.AsSeconds(ThreadState.LastCycle);
+				Session.UpdateDurationSeconds(Timestamp);
 			}
 		}
 
@@ -258,7 +251,7 @@ bool FCpuProfilerAnalyzer::OnEvent(uint16 RouteId, EStyle Style, const FOnEventC
 
 ////////////////////////////////////////////////////////////////////////////////////////////////////
 
-uint64 FCpuProfilerAnalyzer::ProcessBuffer(const FEventTime& EventTime, FThreadState& ThreadState, const uint8* BufferPtr, uint32 BufferSize)
+void FCpuProfilerAnalyzer::ProcessBuffer(const FEventTime& EventTime, FThreadState& ThreadState, const uint8* BufferPtr, uint32 BufferSize)
 {
 	uint64 LastCycle = ThreadState.LastCycle;
 
@@ -341,12 +334,11 @@ uint64 FCpuProfilerAnalyzer::ProcessBuffer(const FEventTime& EventTime, FThreadS
 	}
 
 	ThreadState.LastCycle = LastCycle;
-	return LastCycle;
 }
 
 ////////////////////////////////////////////////////////////////////////////////////////////////////
 
-uint64 FCpuProfilerAnalyzer::ProcessBufferV2(const FEventTime& EventTime, FThreadState& ThreadState, const uint8* BufferPtr, uint32 BufferSize)
+void FCpuProfilerAnalyzer::ProcessBufferV2(const FEventTime& EventTime, FThreadState& ThreadState, const uint8* BufferPtr, uint32 BufferSize)
 {
 	uint64 LastCycle = ThreadState.LastCycle;
 
@@ -395,7 +387,7 @@ uint64 FCpuProfilerAnalyzer::ProcessBufferV2(const FEventTime& EventTime, FThrea
 
 				// Begins a "CoroTask" scoped timer.
 				{
-					if (CoroutineTimerId == ~0)
+					if (CoroutineTimerId == ~0u)
 					{
 						CoroutineTimerId = DefineNewTimerChecked(CoroutineSpecId, TEXT("Coroutine"));
 					}
@@ -424,14 +416,14 @@ uint64 FCpuProfilerAnalyzer::ProcessBufferV2(const FEventTime& EventTime, FThrea
 					CPUPROFILER_DEBUG_BEGIN_EVENT(ActualTime, Event);
 				}
 
-				// Begins the cpu scoped timers (suspended in previous coroutine execution).
+				// Begins the CPU scoped timers (suspended in previous coroutine execution).
 				{
-					if (CoroutineUnknownTimerId == ~0)
+					if (CoroutineUnknownTimerId == ~0u)
 					{
 						CoroutineUnknownTimerId = DefineNewTimerChecked(CoroutineUnknownSpecId, TEXT("<unknown>"));
 					}
 
-					//TODO: Restore the saved stack of cpu scoped timers for this CoroutineId.
+					//TODO: Restore the saved stack of CPU scoped timers for this CoroutineId.
 					for (uint32 i = 0; i < TimerScopeDepth; ++i)
 					{
 						FEventScopeState& ScopeState = ThreadState.ScopeStack.AddDefaulted_GetRef();
@@ -452,9 +444,9 @@ uint64 FCpuProfilerAnalyzer::ProcessBufferV2(const FEventTime& EventTime, FThrea
 
 				if (TimerScopeDepth != 0)
 				{
-					//TODO: Save current stack of cpu scoped timers (using id from metadata of CoroTask timer?)
+					//TODO: Save current stack of CPU scoped timers (using id from metadata of CoroTask timer?)
 
-					// Ends (suspends) the cpu scoped timers.
+					// Ends (suspends) the CPU scoped timers.
 					for (uint32 i = 0; i < TimerScopeDepth; ++i)
 					{
 						// If we receive mismatched end events ignore them for now.
@@ -541,7 +533,6 @@ uint64 FCpuProfilerAnalyzer::ProcessBufferV2(const FEventTime& EventTime, FThrea
 	}
 
 	ThreadState.LastCycle = LastCycle;
-	return LastCycle;
 }
 
 ////////////////////////////////////////////////////////////////////////////////////////////////////
@@ -578,7 +569,7 @@ void FCpuProfilerAnalyzer::DispatchPendingEvents(
 		if (PendingCycle < LastCycle)
 		{
 			// Time needs to increase monotonically.
-			// We are not allowing pending events (with metadata) older than regular cpu timing events.
+			// We are not allowing pending events (with metadata) older than regular CPU timing events.
 			// When this happens we further ignore all pending events on this thread.
 			// The issue can occur in late connect trace sessions with trace protocol <= 6 (i.e. the scoped events have relative timestamps).
 			ThreadState.bShouldIgnorePendingEvents = true;
@@ -614,6 +605,34 @@ void FCpuProfilerAnalyzer::DispatchPendingEvents(
 
 ////////////////////////////////////////////////////////////////////////////////////////////////////
 
+void FCpuProfilerAnalyzer::DispatchRemainingPendingEvents(FThreadState& ThreadState)
+{
+	int32 RemainingPending = ThreadState.PendingEvents.Num();
+	if (RemainingPending > 0)
+	{
+		uint64 LastCycle = ThreadState.LastCycle;
+		const FPendingEvent* PendingCursor = ThreadState.PendingEvents.GetData();
+		DispatchPendingEvents(LastCycle, ~0ull, ThreadState, PendingCursor, RemainingPending);
+		check(RemainingPending == 0);
+		ThreadState.PendingEvents.Reset();
+	}
+}
+
+////////////////////////////////////////////////////////////////////////////////////////////////////
+
+void FCpuProfilerAnalyzer::EndOpenEvents(FThreadState& ThreadState, double Timestamp)
+{
+	while (ThreadState.ScopeStack.Num())
+	{
+		ThreadState.ScopeStack.Pop();
+		CPUPROFILER_DEBUG_LOGF(TEXT("[%u] ~E=%llu (%.9f)\n"), ThreadState.ThreadId, ThreadState.LastCycle, Timestamp);
+		ThreadState.Timeline->AppendEndEvent(Timestamp);
+		CPUPROFILER_DEBUG_END_EVENT(Timestamp);
+	}
+}
+
+////////////////////////////////////////////////////////////////////////////////////////////////////
+
 void FCpuProfilerAnalyzer::OnCpuScopeEnter(const FOnEventContext& Context)
 {
 	if (Context.EventTime.GetTimestamp() == 0)
@@ -681,7 +700,7 @@ void FCpuProfilerAnalyzer::OnCpuScopeLeave(const FOnEventContext& Context)
 
 uint32 FCpuProfilerAnalyzer::DefineTimer(uint32 SpecId, const TCHAR* Name, const TCHAR* File, uint32 Line, bool bMergeByName)
 {
-	// The cpu scoped events (timers) will be merged by name.
+	// The CPU scoped events (timers) will be merged by name.
 	// Ex.: If there are multiple timers defined in code with same name,
 	//      those will appear in Insights as a single timer.
 
@@ -698,7 +717,7 @@ uint32 FCpuProfilerAnalyzer::DefineTimer(uint32 SpecId, const TCHAR* Name, const
 			// Update name for mapped timer.
 			EditableTimingProfilerProvider.SetTimerNameAndLocation(*FindTimerId, Name, File, Line);
 			// In this case, we do not remap the SpecId to the previously defined timer with same name.
-			// This is becasue the two timers are already used in timelines.
+			// This is because the two timers are already used in timelines.
 			// So we will continue to use separate timers, even if those have same name.
 			return *FindTimerId;
 		}
@@ -728,7 +747,7 @@ uint32 FCpuProfilerAnalyzer::DefineTimer(uint32 SpecId, const TCHAR* Name, const
 		}
 		else
 		{
-			// Define a new Cpu timer.
+			// Define a new CPU timer.
 			uint32 NewTimerId = EditableTimingProfilerProvider.AddCpuTimer(Name, File, Line);
 			// Map the SpecId to the timer.
 			SpecIdToTimerIdMap.Add(SpecId, NewTimerId);
@@ -763,7 +782,7 @@ uint32 FCpuProfilerAnalyzer::GetTimerId(uint32 SpecId)
 	else
 	{
 		// Adds a timer with an "unknown" name.
-		// The "unknown" timers are not merged by name, becasue the actual name
+		// The "unknown" timers are not merged by name, because the actual name
 		// might be updated when an EventSpec event is received (for this SpecId).
 		return DefineNewTimerChecked(SpecId, *FString::Printf(TEXT("<unknown %u>"), SpecId));
 	}
diff --git a/Engine/Source/Developer/TraceServices/Private/Analyzers/CpuProfilerTraceAnalysis.h b/Engine/Source/Developer/TraceServices/Private/Analyzers/CpuProfilerTraceAnalysis.h
--- a/Engine/Source/Developer/TraceServices/Private/Analyzers/CpuProfilerTraceAnalysis.h
+++ b/Engine/Source/Developer/TraceServices/Private/Analyzers/CpuProfilerTraceAnalysis.h
@@ -49,27 +49,26 @@ private:
 		bool bShouldIgnorePendingEvents = false; // becomes true when we detect first pending event with incorrect timestamp (i.e < LastCycle)
 	};
 
+	void ProcessBuffer(const FEventTime& EventTime, FThreadState& ThreadState, const uint8* BufferPtr, uint32 BufferSize);
+	void ProcessBufferV2(const FEventTime& EventTime, FThreadState& ThreadState, const uint8* BufferPtr, uint32 BufferSize);
+	void DispatchPendingEvents(uint64& LastCycle, uint64 CurrentCycle, FThreadState& ThreadState, const FPendingEvent*& PendingCursor, int32& RemainingPending);
+	void DispatchRemainingPendingEvents(FThreadState& ThreadState);
+	void EndOpenEvents(FThreadState& ThreadState, double Timestamp);
 	void OnCpuScopeEnter(const FOnEventContext& Context);
 	void OnCpuScopeLeave(const FOnEventContext& Context);
 	uint32 DefineTimer(uint32 SpecId, const TCHAR* TimerName, const TCHAR* File, uint32 Line, bool bMergeByName); // returns the TimerId
 	uint32 DefineNewTimerChecked(uint32 SpecId, const TCHAR* TimerName, const TCHAR* File = nullptr, uint32 Line = 0); // returns the TimerId
 	uint32 GetTimerId(uint32 SpecId);
 	FThreadState& GetThreadState(uint32 ThreadId);
-	uint64 ProcessBuffer(const FEventTime& EventTime, FThreadState& ThreadState, const uint8* BufferPtr, uint32 BufferSize);
-	uint64 ProcessBufferV2(const FEventTime& EventTime, FThreadState& ThreadState, const uint8* BufferPtr, uint32 BufferSize);
-	void DispatchPendingEvents(uint64& LastCycle, uint64 CurrentCycle, FThreadState& ThreadState, const FPendingEvent*& PendingCursor, int32& RemainingPending);
 
 	enum : uint16
 	{
 		RouteId_EventSpec,
-		RouteId_EventBatch,
-		RouteId_EventBatchV2,
 		RouteId_EndThread,
-		RouteId_EndCapture,
-		RouteId_EndCaptureV2,
+		RouteId_EventBatchV2,
+		RouteId_EventBatch, // backward compatibility
+		RouteId_EndCapture, // backward compatibility
 		RouteId_CpuScope,
-		RouteId_ChannelAnnounce,
-		RouteId_ChannelToggle,
 	};
 
 	IAnalysisSession& Session;
