[Insights] LoadTime: Added error message (instead of a check that would result in a crash) for cases when an End* event is encountered without the corresponding Begin* event. This can happen in late connect sessions with "loatime" channel enabled.

diff --git a/Engine/Source/Developer/TraceServices/Private/Analyzers/LoadTimeTraceAnalysis.cpp b/Engine/Source/Developer/TraceServices/Private/Analyzers/LoadTimeTraceAnalysis.cpp
--- a/Engine/Source/Developer/TraceServices/Private/Analyzers/LoadTimeTraceAnalysis.cpp
+++ b/Engine/Source/Developer/TraceServices/Private/Analyzers/LoadTimeTraceAnalysis.cpp
@@ -21,7 +21,7 @@
 	if (++ErrorCount <= MaxErrorCount) \
 	{ \
 		UE_LOG(LogTraceServices, Error, TEXT("[LoadTime] ") __VA_ARGS__); \
-		if (WarningCount == MaxWarningCount) \
+		if (ErrorCount == MaxErrorCount) \
 		{ \
 			UE_LOG(LogTraceServices, Error, TEXT("[LoadTime] Too many errors! Further LoadTime analysis errors will not be reported anymore.")); \
 		} \
@@ -314,11 +314,10 @@ bool FAsyncLoadingTraceAnalyzer::OnEvent(uint16 RouteId, EStyle Style, const FOn
 
 		uint32 ThreadId = FTraceAnalyzerUtils::GetThreadIdField(Context);
 		FThreadState& ThreadState = GetThreadState(ThreadId);
-		{
-			FAnalysisSessionEditScope _(Session);
-			ThreadState.EnterScope(Time, PackageInfo);
-			Session.UpdateDurationSeconds(Time);
-		}
+
+		FAnalysisSessionEditScope _(Session);
+		ThreadState.EnterScope(Time, PackageInfo);
+		Session.UpdateDurationSeconds(Time);
 		break;
 	}
 
@@ -331,12 +330,21 @@ bool FAsyncLoadingTraceAnalyzer::OnEvent(uint16 RouteId, EStyle Style, const FOn
 		double Time = Context.EventTime.AsSeconds(Cycle);
 		UE_LOAD_TIME_TRACE_ANALYSIS_LOG1(TEXT("EndProcessSummary Time=%f"), Time);
 
-		if (ensure(ThreadState.GetCurrentScopeEventType() == LoadTimeProfilerObjectEventType_None))
+		if (ThreadState.CpuScopeStackDepth <= 0)
 		{
-			FAnalysisSessionEditScope _(Session);
-			ThreadState.LeaveScope(Time);
-			Session.UpdateDurationSeconds(Time);
+			UE_LOAD_TIME_TRACE_ANALYSIS_ERROR(TEXT("EndProcessSummary without Begin event!"));
+			break;
+		}
+
+		if (ThreadState.GetCurrentScopeEventType() != LoadTimeProfilerObjectEventType_None)
+		{
+			UE_LOAD_TIME_TRACE_ANALYSIS_ERROR(TEXT("EndProcessSummary does not match the Begin event type!"));
+			break;
 		}
+
+		FAnalysisSessionEditScope _(Session);
+		ThreadState.LeaveScope(Time);
+		Session.UpdateDurationSeconds(Time);
 		break;
 	}
 
@@ -422,23 +430,32 @@ bool FAsyncLoadingTraceAnalyzer::OnEvent(uint16 RouteId, EStyle Style, const FOn
 		uint64 ObjectPtr = EventData.GetValue<uint64>("Object");
 		UE_LOAD_TIME_TRACE_ANALYSIS_LOG1(TEXT("EndCreateExport Time=%f Object=%llu"), Time, ObjectPtr);
 
-		if (ensure(ThreadState.GetCurrentScopeEventType() == LoadTimeProfilerObjectEventType_Create))
+		if (ThreadState.CpuScopeStackDepth <= 0)
 		{
-			FAnalysisSessionEditScope _(Session);
-			FPackageExportInfo* Export = ThreadState.GetCurrentExportScope();
-			if (Export)
-			{
-				ExportsMap.Add(ObjectPtr, Export);
-				const FClassInfo* ObjectClass = GetClassInfo(EventData.GetValue<uint64>("Class"));
-				Export->Class = ObjectClass;
-			}
-			else
-			{
-				UE_LOAD_TIME_TRACE_ANALYSIS_WARNING(TEXT("Unknown export scope!"));
-			}
-			ThreadState.LeaveScope(Time);
-			Session.UpdateDurationSeconds(Time);
+			UE_LOAD_TIME_TRACE_ANALYSIS_ERROR(TEXT("EndCreateExport without Begin event!"));
+			break;
+		}
+
+		if (ThreadState.GetCurrentScopeEventType() != LoadTimeProfilerObjectEventType_Create)
+		{
+			UE_LOAD_TIME_TRACE_ANALYSIS_ERROR(TEXT("EndCreateExport does not match the Begin event type!"));
+			break;
+		}
+
+		FAnalysisSessionEditScope _(Session);
+		FPackageExportInfo* Export = ThreadState.GetCurrentExportScope();
+		if (Export)
+		{
+			ExportsMap.Add(ObjectPtr, Export);
+			const FClassInfo* ObjectClass = GetClassInfo(EventData.GetValue<uint64>("Class"));
+			Export->Class = ObjectClass;
+		}
+		else
+		{
+			UE_LOAD_TIME_TRACE_ANALYSIS_WARNING(TEXT("Unknown export scope!"));
 		}
+		ThreadState.LeaveScope(Time);
+		Session.UpdateDurationSeconds(Time);
 		break;
 	}
 
@@ -485,12 +502,21 @@ bool FAsyncLoadingTraceAnalyzer::OnEvent(uint16 RouteId, EStyle Style, const FOn
 		double Time = Context.EventTime.AsSeconds(Cycle);
 		UE_LOAD_TIME_TRACE_ANALYSIS_LOG1(TEXT("EndSerializeExport Time=%f"), Time);
 
-		if (ensure(ThreadState.GetCurrentScopeEventType() == LoadTimeProfilerObjectEventType_Serialize))
+		if (ThreadState.CpuScopeStackDepth <= 0)
 		{
-			FAnalysisSessionEditScope _(Session);
-			ThreadState.LeaveScope(Time);
-			Session.UpdateDurationSeconds(Time);
+			UE_LOAD_TIME_TRACE_ANALYSIS_ERROR(TEXT("EndSerializeExport without Begin event!"));
+			break;
+		}
+
+		if (ThreadState.GetCurrentScopeEventType() != LoadTimeProfilerObjectEventType_Serialize)
+		{
+			UE_LOAD_TIME_TRACE_ANALYSIS_ERROR(TEXT("EndSerializeExport does not match the Begin event type!"));
+			break;
 		}
+
+		FAnalysisSessionEditScope _(Session);
+		ThreadState.LeaveScope(Time);
+		Session.UpdateDurationSeconds(Time);
 		break;
 	}
 
@@ -572,12 +598,21 @@ bool FAsyncLoadingTraceAnalyzer::OnEvent(uint16 RouteId, EStyle Style, const FOn
 			break;
 		}
 
-		if (ensure(ThreadState.GetCurrentScopeEventType() == LoadTimeProfilerObjectEventType_PostLoad))
+		if (ThreadState.CpuScopeStackDepth <= 0)
 		{
-			FAnalysisSessionEditScope _(Session);
-			ThreadState.LeaveScope(Time);
-			Session.UpdateDurationSeconds(Time);
+			UE_LOAD_TIME_TRACE_ANALYSIS_ERROR(TEXT("EndPostLoadObject without Begin event!"));
+			break;
 		}
+
+		if (ThreadState.GetCurrentScopeEventType() != LoadTimeProfilerObjectEventType_PostLoad)
+		{
+			UE_LOAD_TIME_TRACE_ANALYSIS_ERROR(TEXT("EndPostLoadObject does not match the Begin event type!"));
+			break;
+		}
+
+		FAnalysisSessionEditScope _(Session);
+		ThreadState.LeaveScope(Time);
+		Session.UpdateDurationSeconds(Time);
 		break;
 	}
 
@@ -996,12 +1031,13 @@ bool FAsyncLoadingTraceAnalyzer::OnEvent(uint16 RouteId, EStyle Style, const FOn
 			UE_LOAD_TIME_TRACE_ANALYSIS_WARNING_EX(TEXT("Unknown export Object!"));
 		}
 
-		//EventType = static_cast<ELoadTimeProfilerObjectEventType>(99); // debug
-		{
-			FAnalysisSessionEditScope _(Session);
-			ThreadState.EnterScope(Time, Export, EventType);
-			Session.UpdateDurationSeconds(Time);
-		}
+#if 0 // debug
+		EventType = static_cast<ELoadTimeProfilerObjectEventType>(99);
+#endif
+
+		FAnalysisSessionEditScope _(Session);
+		ThreadState.EnterScope(Time, Export, EventType);
+		Session.UpdateDurationSeconds(Time);
 		break;
 	}
 #endif // UE_LOAD_TIME_TRACE_ANALYSIS_BACKWARD_COMPATIBILITY
@@ -1017,12 +1053,23 @@ bool FAsyncLoadingTraceAnalyzer::OnEvent(uint16 RouteId, EStyle Style, const FOn
 
 		UE_LOAD_TIME_TRACE_ANALYSIS_LOG1(TEXT("EndObjectScope Time=%f"), Time);
 
-		//if (ensure(ThreadState.GetCurrentExportScopeEventType() == static_cast<ELoadTimeProfilerObjectEventType>(99))) // debug
+		if (ThreadState.CpuScopeStackDepth <= 0)
 		{
-			FAnalysisSessionEditScope _(Session);
-			ThreadState.LeaveScope(Time);
-			Session.UpdateDurationSeconds(Time);
+			UE_LOAD_TIME_TRACE_ANALYSIS_ERROR(TEXT("EndObjectScope without Begin event!"));
+			break;
 		}
+
+#if 0 // debug
+		if (ThreadState.GetCurrentExportScopeEventType() != static_cast<ELoadTimeProfilerObjectEventType>(99)))
+		{
+			UE_LOAD_TIME_TRACE_ANALYSIS_ERROR(TEXT("EndObjectScope does not match the Begin event type!"));
+			break;
+		}
+#endif
+
+		FAnalysisSessionEditScope _(Session);
+		ThreadState.LeaveScope(Time);
+		Session.UpdateDurationSeconds(Time);
 		break;
 	}
 #endif // UE_LOAD_TIME_TRACE_ANALYSIS_BACKWARD_COMPATIBILITY
