[Insights] Improved error detection for timing region events. A zero duration region will be created automatically if begin and end events do not match. Excluded warnings for "<SlowTask>" regions as begin-end events are known to mismatch.

diff --git a/Engine/Source/Developer/TraceServices/Private/Model/Regions.cpp b/Engine/Source/Developer/TraceServices/Private/Model/Regions.cpp
--- a/Engine/Source/Developer/TraceServices/Private/Model/Regions.cpp
+++ b/Engine/Source/Developer/TraceServices/Private/Model/Regions.cpp
@@ -47,15 +47,22 @@ void FRegionProvider::AppendRegionBegin(const TCHAR* Name, double Time)
 {
 	EditAccessCheck();
 
-	if (FTimeRegion** OpenRegion = OpenRegions.Find(Name))
+	FTimeRegion** OpenRegion = OpenRegions.Find(Name);
+	if (OpenRegion)
 	{
-		++NumWarnings;
-		if (NumWarnings <= MaxWarningMessages)
+		if (FCString::Strcmp(Name, TEXT("<SlowTask>")) != 0)
 		{
-			UE_LOG(LogTraceServices, Warning, TEXT("[Regions] A region begin event (%s) was encountered while a region with same name is already open."), Name)
+			++NumWarnings;
+			if (NumWarnings <= MaxWarningMessages)
+			{
+				UE_LOG(LogTraceServices, Warning, TEXT("[Regions] A region begin event (BeginTime=%f, Name=\"%s\") was encountered while a region with same name is already open."), Time, Name)
+			}
 		}
+
+		// Automatically end the previous region.
+		AppendRegionEnd(Name, Time);
 	}
-	else
+
 	{
 		FTimeRegion Region;
 		Region.BeginTime = Time;
@@ -85,21 +92,30 @@ void FRegionProvider::AppendRegionEnd(const TCHAR* Name, double Time)
 {
 	EditAccessCheck();
 
-	if (FTimeRegion** OpenRegion = OpenRegions.Find(Name))
+	FTimeRegion** OpenRegion = OpenRegions.Find(Name);
+	if (!OpenRegion)
+	{
+		if (FCString::Strcmp(Name, TEXT("<SlowTask>")) != 0)
+		{
+			++NumWarnings;
+			if (NumWarnings <= MaxWarningMessages)
+			{
+				UE_LOG(LogTraceServices, Warning, TEXT("[Regions] A region end event (EndTime=%f, Name=\"%s\") was encountered without having seen a matching region begin event first."), Time, Name)
+			}
+		}
+
+		// Automatically create a new region.
+		AppendRegionBegin(Name, Time);
+		OpenRegion = OpenRegions.Find(Name);
+		check(OpenRegion);
+	}
+
 	{
 		(*OpenRegion)->EndTime = Time;
 
 		OpenRegions.Remove(Name);
 		UpdateCounter++;
 	}
-	else
-	{
-		++NumWarnings;
-		if (NumWarnings <= MaxWarningMessages)
-		{
-			UE_LOG(LogTraceServices, Warning, TEXT("[Regions] A region end event (%s) was encountered without having seen a matching region start event first."), Name)
-		}
-	}
 
 	// Update session time
 	{
@@ -116,10 +132,13 @@ void FRegionProvider::OnAnalysisSessionEnded()
 	{
 		const FTimeRegion* Region = KV.Value;
 
-		++NumWarnings;
-		if (NumWarnings <= MaxWarningMessages)
+		if (FCString::Strcmp(Region->Text, TEXT("<SlowTask>")) != 0)
 		{
-			UE_LOG(LogTraceServices, Warning, TEXT("[Regions] A region begin event (%s) was never closed."), Region->Text)
+			++NumWarnings;
+			if (NumWarnings <= MaxWarningMessages)
+			{
+				UE_LOG(LogTraceServices, Warning, TEXT("[Regions] A region (BeginTime=%f, Name=\"%s\") was never closed."), Region->BeginTime, Region->Text)
+			}
 		}
 	}
 
