NetworkInsights
- Detected lost packets are now marked on incoming data
- Changed how we display lost data to make it easier to spot.

diff --git a/Engine/Source/Developer/TraceInsights/Private/Insights/NetworkingProfiler/ViewModels/PacketViewDrawHelper.cpp b/Engine/Source/Developer/TraceInsights/Private/Insights/NetworkingProfiler/ViewModels/PacketViewDrawHelper.cpp
--- a/Engine/Source/Developer/TraceInsights/Private/Insights/NetworkingProfiler/ViewModels/PacketViewDrawHelper.cpp
+++ b/Engine/Source/Developer/TraceInsights/Private/Insights/NetworkingProfiler/ViewModels/PacketViewDrawHelper.cpp
@@ -233,7 +233,10 @@ void FPacketViewDrawHelper::DrawCached(const FNetworkPacketSeries& Series) const
 			{
 				DrawContext.DrawBox(X, FilterMatchContentY, SampleW, FilterMatchContentH, WhiteBrush, ColorFilterMatch);
 			}
-
+		}
+		if (Sample.AggregatedStatus == TraceServices::ENetProfilerDeliveryStatus::Dropped)
+		{
+			DrawContext.DrawBox(X + (SampleW/2.f), 0.f, 1.f, Y + H, WhiteBrush, FLinearColor::Red);
 		}
 	}
 
diff --git a/Engine/Source/Developer/TraceServices/Private/Analyzers/NetTraceAnalyzer.cpp b/Engine/Source/Developer/TraceServices/Private/Analyzers/NetTraceAnalyzer.cpp
--- a/Engine/Source/Developer/TraceServices/Private/Analyzers/NetTraceAnalyzer.cpp
+++ b/Engine/Source/Developer/TraceServices/Private/Analyzers/NetTraceAnalyzer.cpp
@@ -521,6 +521,33 @@ void FNetTraceAnalyzer::HandlePacketEvent(const FOnEventContext& Context, const
 
 	// Add the packet
 	FNetProfilerConnectionData& ConnectionData = NetProfilerProvider.EditConnectionData(ConnectionState->ConnectionIndex, ConnectionMode);
+
+	if (ConnectionMode == ENetProfilerConnectionMode::Incoming)
+	{
+		if (ConnectionData.Packets.Num() > 0)
+		{
+			uint32 ExpectedSequenceNumber = ConnectionData.Packets.Last().SequenceNumber + 1U;
+			while (ExpectedSequenceNumber < SequenceNumber)
+			{
+				// Inject packets to visualize missing packets
+				FNetProfilerPacket& Packet = ConnectionData.Packets.PushBack();
+				Packet.SequenceNumber = ExpectedSequenceNumber;
+
+				// Fake it
+				Packet.StartEventIndex = ConnectionState->CurrentPacketStartIndex[ConnectionMode];
+				Packet.EventCount = 0;
+				Packet.TimeStamp = GetLastTimestamp();
+				Packet.DeliveryStatus = ENetProfilerDeliveryStatus::Dropped;
+				Packet.ConnectionState = ConnectionState->ConnectionState;
+				Packet.ContentSizeInBits = 0;
+				Packet.TotalPacketSizeInBytes = (Packet.ContentSizeInBits + 7u) >> 3u;
+			
+				++ExpectedSequenceNumber;
+				++ConnectionData.PacketChangeCount;
+			}
+		}
+	}
+
 	FNetProfilerPacket& Packet = ConnectionData.Packets.PushBack();
 	++ConnectionData.PacketChangeCount;
 
