[Insights] Fixed incorrect timestamps for cpu timing events with metadata that are adjacent to other regular cpu timing events.

diff --git a/Engine/Source/Developer/TraceServices/Private/Analyzers/CpuProfilerTraceAnalysis.cpp b/Engine/Source/Developer/TraceServices/Private/Analyzers/CpuProfilerTraceAnalysis.cpp
--- a/Engine/Source/Developer/TraceServices/Private/Analyzers/CpuProfilerTraceAnalysis.cpp
+++ b/Engine/Source/Developer/TraceServices/Private/Analyzers/CpuProfilerTraceAnalysis.cpp
@@ -284,7 +284,7 @@ void FCpuProfilerAnalyzer::ProcessBuffer(const FEventTime& EventTime, FThreadSta
 		}
 
 		// Dispatch pending events that are younger than the one we've just decoded.
-		DispatchPendingEvents(LastCycle, ActualCycle, ThreadState, PendingCursor, RemainingPending);
+		DispatchPendingEvents(LastCycle, ActualCycle, ThreadState, PendingCursor, RemainingPending, (DecodedCycle & 1ull) != 0);
 
 		double ActualTime = EventTime.AsSeconds(ActualCycle);
 
@@ -371,7 +371,7 @@ void FCpuProfilerAnalyzer::ProcessBufferV2(const FEventTime& EventTime, FThreadS
 		}
 
 		// Dispatch pending events that are younger than the one we've just decoded.
-		DispatchPendingEvents(LastCycle, ActualCycle, ThreadState, PendingCursor, RemainingPending);
+		DispatchPendingEvents(LastCycle, ActualCycle, ThreadState, PendingCursor, RemainingPending, (DecodedCycle & 1ull) != 0);
 
 		double ActualTime = EventTime.AsSeconds(ActualCycle);
 
@@ -542,7 +542,8 @@ void FCpuProfilerAnalyzer::DispatchPendingEvents(
 	uint64 CurrentCycle,
 	FThreadState& ThreadState,
 	const FPendingEvent*& PendingCursor,
-	int32& RemainingPending)
+	int32& RemainingPending,
+	bool bIsBeginEvent)
 {
 	if (ThreadState.bShouldIgnorePendingEvents)
 	{
@@ -561,7 +562,8 @@ void FCpuProfilerAnalyzer::DispatchPendingEvents(
 			bEnter = false;
 		}
 
-		if (PendingCycle > CurrentCycle)
+		if ((PendingCycle > CurrentCycle) ||
+			(PendingCycle == CurrentCycle && !bIsBeginEvent))
 		{
 			break;
 		}
@@ -612,7 +614,7 @@ void FCpuProfilerAnalyzer::DispatchRemainingPendingEvents(FThreadState& ThreadSt
 	{
 		uint64 LastCycle = ThreadState.LastCycle;
 		const FPendingEvent* PendingCursor = ThreadState.PendingEvents.GetData();
-		DispatchPendingEvents(LastCycle, ~0ull, ThreadState, PendingCursor, RemainingPending);
+		DispatchPendingEvents(LastCycle, ~0ull, ThreadState, PendingCursor, RemainingPending, true);
 		check(RemainingPending == 0);
 		ThreadState.PendingEvents.Reset();
 	}
diff --git a/Engine/Source/Developer/TraceServices/Private/Analyzers/CpuProfilerTraceAnalysis.h b/Engine/Source/Developer/TraceServices/Private/Analyzers/CpuProfilerTraceAnalysis.h
--- a/Engine/Source/Developer/TraceServices/Private/Analyzers/CpuProfilerTraceAnalysis.h
+++ b/Engine/Source/Developer/TraceServices/Private/Analyzers/CpuProfilerTraceAnalysis.h
@@ -51,7 +51,7 @@ private:
 
 	void ProcessBuffer(const FEventTime& EventTime, FThreadState& ThreadState, const uint8* BufferPtr, uint32 BufferSize);
 	void ProcessBufferV2(const FEventTime& EventTime, FThreadState& ThreadState, const uint8* BufferPtr, uint32 BufferSize);
-	void DispatchPendingEvents(uint64& LastCycle, uint64 CurrentCycle, FThreadState& ThreadState, const FPendingEvent*& PendingCursor, int32& RemainingPending);
+	void DispatchPendingEvents(uint64& LastCycle, uint64 CurrentCycle, FThreadState& ThreadState, const FPendingEvent*& PendingCursor, int32& RemainingPending, bool bIsBeginEvent);
 	void DispatchRemainingPendingEvents(FThreadState& ThreadState);
 	void EndOpenEvents(FThreadState& ThreadState, double Timestamp);
 	void OnCpuScopeEnter(const FOnEventContext& Context);
