[Insights] Fixed Unreal Insights crash when analysing LoadTime trace events. It happened in two situations: a) when EndRequest event encounters a package (associated with the request) already destroyed; b) when AsyncPackageRequestAssociation tries to associate an async package with imported packages already destroyed.

diff --git a/Engine/Source/Developer/TraceServices/Private/Analyzers/LoadTimeTraceAnalysis.cpp b/Engine/Source/Developer/TraceServices/Private/Analyzers/LoadTimeTraceAnalysis.cpp
--- a/Engine/Source/Developer/TraceServices/Private/Analyzers/LoadTimeTraceAnalysis.cpp
+++ b/Engine/Source/Developer/TraceServices/Private/Analyzers/LoadTimeTraceAnalysis.cpp
@@ -246,7 +246,7 @@ bool FAsyncLoadingTraceAnalyzer::OnEvent(uint16 RouteId, EStyle Style, const FOn
 				Summary.ImportCount = EventData.GetValue<uint32>("ImportCount");
 				Summary.ExportCount = EventData.GetValue<uint32>("ExportCount");
 				Summary.Priority = EventData.GetValue<int32>("Priority", 0);  // added in UE 5.4
-				UE_LOAD_TIME_TRACE_ANALYSIS_LOG1(TEXT("TotalHeaderSize=%u ImportCount=%u ExportCount=%u Priority=%d"), Summary.TotalHeaderSize, Summary.ImportCount, Summary.ExportCount, Summary.Priority);
+				UE_LOAD_TIME_TRACE_ANALYSIS_LOG1(TEXT("  TotalHeaderSize=%u ImportCount=%u ExportCount=%u Priority=%d"), Summary.TotalHeaderSize, Summary.ImportCount, Summary.ExportCount, Summary.Priority);
 			}
 			else
 			{
@@ -793,6 +793,25 @@ bool FAsyncLoadingTraceAnalyzer::OnEvent(uint16 RouteId, EStyle Style, const FOn
 				AsyncPackageState->Linker->AsyncPackage = nullptr;
 				UE_LOAD_TIME_TRACE_ANALYSIS_LOG2(TEXT("Linker (%p) : reset AsyncPackage"), AsyncPackageState->Linker);
 			}
+
+			if (AsyncPackageState->Request)
+			{
+				AsyncPackageState->Request->AsyncPackages.RemoveSingleSwap(AsyncPackageState);
+				UE_LOAD_TIME_TRACE_ANALYSIS_LOG2(TEXT("Request (%p) : remove AsyncPackage"), AsyncPackageState->Request);
+			}
+
+			for (FAsyncPackageState* ImportedAsyncPackage : AsyncPackageState->ImportedAsyncPackages)
+			{
+				ImportedAsyncPackage->ImportedByAsyncPackages.Remove(AsyncPackageState);
+				UE_LOAD_TIME_TRACE_ANALYSIS_LOG2(TEXT("Imported AsyncPackage (%p) : remove ImportedBy AsyncPackage"), ImportedAsyncPackage);
+			}
+
+			for (FAsyncPackageState* ImportedByAsyncPackage : AsyncPackageState->ImportedByAsyncPackages)
+			{
+				ImportedByAsyncPackage->ImportedAsyncPackages.Remove(AsyncPackageState);
+				UE_LOAD_TIME_TRACE_ANALYSIS_LOG2(TEXT("ImportedBy AsyncPackage (%p) : remove Imported AsyncPackage"), ImportedByAsyncPackage);
+			}
+
 			UE_LOAD_TIME_TRACE_ANALYSIS_LOG2(TEXT("delete FAsyncPackageState %p"), AsyncPackageState);
 			delete AsyncPackageState;
 		}
@@ -820,14 +839,17 @@ bool FAsyncLoadingTraceAnalyzer::OnEvent(uint16 RouteId, EStyle Style, const FOn
 		FLinkerLoadState* LinkerState;
 		if (ActiveLinkersMap.RemoveAndCopyValue(LinkerPtr, LinkerState))
 		{
-			if (LinkerState->AsyncPackage)
+			FAsyncPackageState* AsyncPackageState = LinkerState->AsyncPackage;
+			if (AsyncPackageState)
 			{
-				LinkerState->AsyncPackage->Linker = nullptr;
-			}
-			if (LinkerState->bHasFakeAsyncPackageState)
-			{
-				UE_LOAD_TIME_TRACE_ANALYSIS_LOG2(TEXT("delete (fake) FAsyncPackageState %p"), LinkerState->AsyncPackage);
-				delete LinkerState->AsyncPackage;
+				AsyncPackageState->Linker = nullptr;
+				if (LinkerState->bHasFakeAsyncPackageState)
+				{
+					UE_LOAD_TIME_TRACE_ANALYSIS_LOG2(TEXT("delete (fake) FAsyncPackageState %p"), AsyncPackageState);
+					ensure(AsyncPackageState->ImportedAsyncPackages.IsEmpty());
+					ensure(AsyncPackageState->ImportedByAsyncPackages.IsEmpty());
+					delete AsyncPackageState;
+				}
 			}
 			UE_LOAD_TIME_TRACE_ANALYSIS_LOG2(TEXT("delete FLinkerLoadState %p"), LinkerState);
 			delete LinkerState;
@@ -872,6 +894,7 @@ bool FAsyncLoadingTraceAnalyzer::OnEvent(uint16 RouteId, EStyle Style, const FOn
 			if (!AsyncPackage->ImportedAsyncPackages.Contains(ImportedAsyncPackage))
 			{
 				AsyncPackage->ImportedAsyncPackages.Add(ImportedAsyncPackage);
+				ImportedAsyncPackage->ImportedByAsyncPackages.Add(AsyncPackage);
 				FAnalysisSessionEditScope _(Session);
 				AsyncPackage->PackageInfo->ImportedPackages.Add(ImportedAsyncPackage->PackageInfo);
 				if (AsyncPackage->Request)
diff --git a/Engine/Source/Developer/TraceServices/Private/Analyzers/LoadTimeTraceAnalysis.h b/Engine/Source/Developer/TraceServices/Private/Analyzers/LoadTimeTraceAnalysis.h
--- a/Engine/Source/Developer/TraceServices/Private/Analyzers/LoadTimeTraceAnalysis.h
+++ b/Engine/Source/Developer/TraceServices/Private/Analyzers/LoadTimeTraceAnalysis.h
@@ -58,6 +58,7 @@ private:
 		FRequestState* Request = nullptr;
 		FLinkerLoadState* Linker = nullptr;
 		TSet<FAsyncPackageState*> ImportedAsyncPackages;
+		TSet<FAsyncPackageState*> ImportedByAsyncPackages;
 	};
 
 	struct FLinkerLoadState
