[Insights] GpuProfiler: Fixed crash when analysis encounters GPU profiler frames with invalid timestamps (ex. emitted with invalid "CalibrationBias"). Also added error / log messages when analysis is completed.

diff --git a/Engine/Source/Developer/TraceServices/Private/Analyzers/GpuProfilerTraceAnalysis.cpp b/Engine/Source/Developer/TraceServices/Private/Analyzers/GpuProfilerTraceAnalysis.cpp
--- a/Engine/Source/Developer/TraceServices/Private/Analyzers/GpuProfilerTraceAnalysis.cpp
+++ b/Engine/Source/Developer/TraceServices/Private/Analyzers/GpuProfilerTraceAnalysis.cpp
@@ -12,10 +12,7 @@ namespace TraceServices
 FGpuProfilerAnalyzer::FGpuProfilerAnalyzer(FAnalysisSession& InSession, FTimingProfilerProvider& InTimingProfilerProvider)
 	: Session(InSession)
 	, TimingProfilerProvider(InTimingProfilerProvider)
-	, Timeline(TimingProfilerProvider.EditGpuTimeline())
-	, Timeline2(TimingProfilerProvider.EditGpu2Timeline())
 {
-
 }
 
 void FGpuProfilerAnalyzer::OnAnalysisBegin(const FOnAnalysisContext& Context)
@@ -27,6 +24,19 @@ void FGpuProfilerAnalyzer::OnAnalysisBegin(const FOnAnalysisContext& Context)
 	Builder.RouteEvent(RouteId_Frame2, "GpuProfiler", "Frame2");
 }
 
+void FGpuProfilerAnalyzer::OnAnalysisEnd()
+{
+	if (NumFramesWithErrors > 0)
+	{
+		UE_LOG(LogTraceServices, Error, TEXT("[GpuProfiler] Frames with errors: %u"), NumFramesWithErrors);
+	}
+
+	if (NumFrames > 0 || EventTypeMap.Num() > 0)
+	{
+		UE_LOG(LogTraceServices, Log, TEXT("[GpuProfiler] Analysis completed (%u frames, %d timers)."), NumFrames, EventTypeMap.Num());
+	}
+}
+
 bool FGpuProfilerAnalyzer::OnEvent(uint16 RouteId, EStyle Style, const FOnEventContext& Context)
 {
 	LLM_SCOPE_BYNAME(TEXT("Insights/FGpuProfilerAnalyzer"));
@@ -57,25 +67,43 @@ bool FGpuProfilerAnalyzer::OnEvent(uint16 RouteId, EStyle Style, const FOnEventC
 	case RouteId_Frame:
 	case RouteId_Frame2:
 	{
-		TraceServices::FTimingProfilerProvider::TimelineInternal& ThisTimeline = (RouteId == RouteId_Frame) ? Timeline : Timeline2;
+		TraceServices::FTimingProfilerProvider::TimelineInternal& ThisTimeline = (RouteId == RouteId_Frame) ?
+			TimingProfilerProvider.EditGpuTimeline() :
+			TimingProfilerProvider.EditGpu2Timeline();
 		double& ThisMinTime = (RouteId == RouteId_Frame) ? MinTime : MinTime2;
+
 		const auto& Data = EventData.GetArray<uint8>("Data");
 		const uint8* BufferPtr = Data.GetData();
 		const uint8* BufferEnd = BufferPtr + Data.Num();
 
-		uint32 CurrentDepth = 0;
-
 		uint64 CalibrationBias = EventData.GetValue<uint64>("CalibrationBias");
 		uint64 LastTimestamp = EventData.GetValue<uint64>("TimestampBase");
+		uint32 RenderingFrameNumber = EventData.GetValue<uint32>("RenderingFrameNumber");
+
+		++NumFrames;
+
 		double LastTime = 0.0;
+		uint32 CurrentDepth = 0;
+		bool bHasErrors = false;
+
 		while (BufferPtr < BufferEnd)
 		{
 			uint64 DecodedTimestamp = FTraceAnalyzerUtils::Decode7bit(BufferPtr);
 			uint64 ActualTimestamp = (DecodedTimestamp >> 1) + LastTimestamp;
 			LastTimestamp = ActualTimestamp;
-			LastTime = double(ActualTimestamp + CalibrationBias) / 1000000.0;
+			LastTime = double(ActualTimestamp + CalibrationBias) * 0.000001;
 			LastTime += Context.EventTime.AsSeconds(0);
 
+			if (LastTime < 0.0)
+			{
+				if (DecodedTimestamp & 1ull)
+				{
+					BufferPtr += sizeof(uint32);
+				}
+				bHasErrors = true;
+				continue;
+			}
+
 			// The monolithic timeline assumes that timestamps are ever increasing, but
 			// with gpu/cpu calibration and drift there can be a tiny bit of overlap between
 			// frames. So we just clamp.
@@ -113,12 +141,15 @@ bool FGpuProfilerAnalyzer::OnEvent(uint16 RouteId, EStyle Style, const FOnEventC
 				ThisTimeline.AppendEndEvent(LastTime);
 			}
 		}
-		Session.UpdateDurationSeconds(LastTime);
 		check(BufferPtr == BufferEnd);
 		check(CurrentDepth == 0);
+		if (bHasErrors && ++NumFramesWithErrors <= 100)
+		{
+			UE_LOG(LogTraceServices, Error, TEXT("[GpuProfiler] The rendering frame %u has invalid timestamps!"), RenderingFrameNumber);
+		}
+		Session.UpdateDurationSeconds(LastTime);
 		break;
 	}
-		
 	}
 
 	return true;
diff --git a/Engine/Source/Developer/TraceServices/Private/Analyzers/GpuProfilerTraceAnalysis.h b/Engine/Source/Developer/TraceServices/Private/Analyzers/GpuProfilerTraceAnalysis.h
--- a/Engine/Source/Developer/TraceServices/Private/Analyzers/GpuProfilerTraceAnalysis.h
+++ b/Engine/Source/Developer/TraceServices/Private/Analyzers/GpuProfilerTraceAnalysis.h
@@ -16,6 +16,7 @@ class FGpuProfilerAnalyzer
 public:
 	FGpuProfilerAnalyzer(FAnalysisSession& Session, FTimingProfilerProvider& TimingProfilerProvider);
 	virtual void OnAnalysisBegin(const FOnAnalysisContext& Context) override;
+	virtual void OnAnalysisEnd() override;
 	virtual bool OnEvent(uint16 RouteId, EStyle Style, const FOnEventContext& Context) override;
 
 private:
@@ -26,15 +27,13 @@ private:
 		RouteId_Frame2, // GPU Index 1
 	};
 
-	double GpuTimestampToSessionTime(uint64 GpuMicroseconds);
-
 	FAnalysisSession& Session;
 	FTimingProfilerProvider& TimingProfilerProvider;
-	FTimingProfilerProvider::TimelineInternal& Timeline;
-	FTimingProfilerProvider::TimelineInternal& Timeline2;
 	TMap<uint64, uint32> EventTypeMap;
-	double MinTime = 0.0f;
-	double MinTime2 = 0.0f;
+	double MinTime = DBL_MIN;
+	double MinTime2 = DBL_MIN;
+	uint32 NumFrames = 0;
+	uint32 NumFramesWithErrors = 0;
 };
 
 } // namespace TraceServices
diff --git a/Engine/Source/Developer/TraceServices/Private/Model/AllocationsProvider.cpp b/Engine/Source/Developer/TraceServices/Private/Model/AllocationsProvider.cpp
--- a/Engine/Source/Developer/TraceServices/Private/Model/AllocationsProvider.cpp
+++ b/Engine/Source/Developer/TraceServices/Private/Model/AllocationsProvider.cpp
@@ -2298,26 +2298,26 @@ void FAllocationsProvider::EditOnAnalysisCompleted(double Time)
 
 	if (AllocWarnings > 0 || FreeWarnings > 0 || HeapWarnings > 0 || MiscWarnings > 0)
 	{
-		UE_LOG(LogTraceServices, Warning, TEXT("[MemAlloc] %u warnings (%u ALLOC + %u FREE + %u HEAP + %u other)"),
+		UE_LOG(LogTraceServices, Warning, TEXT("[MemAlloc] %llu warnings (%llu ALLOC + %llu FREE + %llu HEAP + %llu other)"),
 			AllocWarnings + FreeWarnings + HeapWarnings + MiscWarnings,
 			AllocWarnings, FreeWarnings, HeapWarnings, MiscWarnings);
 	}
 
 	if (AllocErrors > 0)
 	{
-		UE_LOG(LogTraceServices, Error, TEXT("[MemAlloc] ALLOC event errors: %u"), AllocErrors);
+		UE_LOG(LogTraceServices, Error, TEXT("[MemAlloc] ALLOC event errors: %llu"), AllocErrors);
 	}
 	if (FreeErrors > 0)
 	{
-		UE_LOG(LogTraceServices, Error, TEXT("[MemAlloc] FREE event errors: %u"), FreeErrors);
+		UE_LOG(LogTraceServices, Error, TEXT("[MemAlloc] FREE event errors: %llu"), FreeErrors);
 	}
 	if (HeapErrors > 0)
 	{
-		UE_LOG(LogTraceServices, Error, TEXT("[MemAlloc] HEAP event errors: %u"), HeapErrors);
+		UE_LOG(LogTraceServices, Error, TEXT("[MemAlloc] HEAP event errors: %llu"), HeapErrors);
 	}
 	if (MiscErrors > 0)
 	{
-		UE_LOG(LogTraceServices, Error, TEXT("[MemAlloc] Other errors: %u"), MiscErrors);
+		UE_LOG(LogTraceServices, Error, TEXT("[MemAlloc] Other errors: %llu"), MiscErrors);
 	}
 	if (TagTracker.GetNumErrors() > 0)
 	{
@@ -2341,7 +2341,7 @@ void FAllocationsProvider::EditOnAnalysisCompleted(double Time)
 			++NumHeapSpecs;
 		}
 	}
-	UE_LOG(LogTraceServices, Log, TEXT("[MemAlloc] Analysis completed (%llu events, %llu allocs, %llu frees, %llu heaps, %d heap specs)."),
+	UE_LOG(LogTraceServices, Log, TEXT("[MemAlloc] Analysis completed (%llu events, %llu allocs, %llu frees, %llu heaps, %u heap specs)."),
 		TotalEventCount, AllocCount, FreeCount, HeapCount, NumHeapSpecs);
 }
 
