[Insights] Fixed FMemory::Malloc/delete mismatch in FDefinitionProvider.

diff --git a/Engine/Source/Developer/TraceServices/Private/Model/DefinitionProvider.cpp b/Engine/Source/Developer/TraceServices/Private/Model/DefinitionProvider.cpp
--- a/Engine/Source/Developer/TraceServices/Private/Model/DefinitionProvider.cpp
+++ b/Engine/Source/Developer/TraceServices/Private/Model/DefinitionProvider.cpp
@@ -16,6 +16,16 @@ FDefinitionProvider::FDefinitionProvider(IAnalysisSession* InSession)
 
 ////////////////////////////////////////////////////////////////////////////////////////////////////
 
+FDefinitionProvider::~FDefinitionProvider()
+{
+	for (uint8* Page : Pages)
+	{
+		FMemory::Free(Page);
+	}
+}
+
+////////////////////////////////////////////////////////////////////////////////////////////////////
+
 void FDefinitionProvider::AddEntry(uint64 Hash, const void* Ptr)
 {
 	Definitions.Add(Hash, Ptr);
@@ -37,13 +47,12 @@ void* FDefinitionProvider::Allocate(uint32 Size, uint32 Alignment)
 	uint8* Dest = nullptr;
 	if (PageRemain > Size)
 	{
-		Dest = Pages.Last().Get() + (PageSize - PageRemain);
+		Dest = Pages.Last() + (PageSize - PageRemain);
 	}
 	else
 	{
-		auto& NewPage = Pages.Emplace_GetRef((uint8*)FMemory::MallocZeroed(PageSize));
-
-		Dest = NewPage.Get();
+		Dest = (uint8*)FMemory::MallocZeroed(PageSize);
+		Pages.Add(Dest);
 		PageRemain = PageSize;
 	}
 	PageRemain -= Size;
diff --git a/Engine/Source/Developer/TraceServices/Private/Model/DefinitionProvider.h b/Engine/Source/Developer/TraceServices/Private/Model/DefinitionProvider.h
--- a/Engine/Source/Developer/TraceServices/Private/Model/DefinitionProvider.h
+++ b/Engine/Source/Developer/TraceServices/Private/Model/DefinitionProvider.h
@@ -16,6 +16,7 @@ class FDefinitionProvider : public IDefinitionProvider
 {
 public:
 	FDefinitionProvider(IAnalysisSession* InSession);
+	virtual ~FDefinitionProvider();
 
 	virtual void BeginRead() const override { Lock.ReadLock(); }
 	virtual void EndRead() const override { Lock.ReadUnlock(); }
@@ -33,7 +34,7 @@ private:
 private:
 	mutable FRWLock Lock;
 
-	TArray<TUniquePtr<uint8>> Pages;
+	TArray<uint8*> Pages;
 	static constexpr uint32 PageSize = 1024;
 	uint32 PageRemain;
 
