Unreal Insights support for tracing swap events

diff --git a/Engine/Source/Developer/TraceInsights/Private/Insights/MemoryProfiler/ViewModels/MemAllocGroupingByHeap.cpp b/Engine/Source/Developer/TraceInsights/Private/Insights/MemoryProfiler/ViewModels/MemAllocGroupingByHeap.cpp
--- a/Engine/Source/Developer/TraceInsights/Private/Insights/MemoryProfiler/ViewModels/MemAllocGroupingByHeap.cpp
+++ b/Engine/Source/Developer/TraceInsights/Private/Insights/MemoryProfiler/ViewModels/MemAllocGroupingByHeap.cpp
@@ -23,7 +23,7 @@ FMemAllocGroupingByHeap::FMemAllocGroupingByHeap(const TraceServices::IAllocatio
 	: FTreeNodeGrouping(
 		LOCTEXT("Grouping_ByHeap_ShortName", "Heap"),
 		LOCTEXT("Grouping_ByHeap_TitleName", "By Heap"),
-		LOCTEXT("Grouping_Heap_Desc", "Creates a tree based on heap."),
+		LOCTEXT("Grouping_ByHeap_Desc", "Creates a tree based on heap."),
 		TEXT("Icons.Group.TreeItem"),
 		nullptr)
 	, AllocProvider(InAllocProvider)
diff --git a/Engine/Source/Developer/TraceInsights/Private/Insights/MemoryProfiler/ViewModels/MemAllocGroupingBySwapPage.cpp b/Engine/Source/Developer/TraceInsights/Private/Insights/MemoryProfiler/ViewModels/MemAllocGroupingBySwapPage.cpp
new file mode 100644
--- /dev/null
+++ b/Engine/Source/Developer/TraceInsights/Private/Insights/MemoryProfiler/ViewModels/MemAllocGroupingBySwapPage.cpp
@@ -0,0 +1,169 @@
+﻿// Copyright Epic Games, Inc. All Rights Reserved.
+
+#include "MemAllocGroupingBySwapPage.h"
+
+#include "Common/ProviderLock.h" // TraceServices
+#include "Internationalization/Internationalization.h"
+
+// Insights
+#include "Insights/Common/AsyncOperationProgress.h"
+#include "Insights/MemoryProfiler/ViewModels/CallstackFormatting.h"
+#include "Insights/MemoryProfiler/ViewModels/MemAllocInSwapNode.h"
+
+#define LOCTEXT_NAMESPACE "Insights::FMemAllocGroupingBySwap"
+
+namespace Insights
+{
+
+INSIGHTS_IMPLEMENT_RTTI(FMemAllocGroupingBySwapPage)
+
+////////////////////////////////////////////////////////////////////////////////////////////////////
+
+FMemAllocGroupingBySwapPage::FMemAllocGroupingBySwapPage(const TraceServices::IAllocationsProvider& InAllocProvider)
+	: FTreeNodeGrouping(
+		LOCTEXT("Grouping_BySwap_ShortName", "Swap"),
+		LOCTEXT("Grouping_BySwap_TitleName", "By Swap"),
+		LOCTEXT("Grouping_BySwap_Desc", "Creates a tree based on swap state."),
+		TEXT("Icons.Group.TreeItem"),
+		nullptr)
+	, AllocProvider(InAllocProvider)
+{
+}
+
+////////////////////////////////////////////////////////////////////////////////////////////////////
+
+FMemAllocGroupingBySwapPage::~FMemAllocGroupingBySwapPage()
+{
+}
+
+////////////////////////////////////////////////////////////////////////////////////////////////////
+
+struct TMemoryPageMapKeyFuncs : BaseKeyFuncs<TPair<uint64, FTableTreeNode*>, uint64, false>
+{
+	static FORCEINLINE uint64 GetSetKey(const TPair<uint64, FTableTreeNode*>& Element)
+	{
+		return Element.Key;
+	}
+	static FORCEINLINE uint32 GetKeyHash(const uint64 Key)
+	{
+		// memory page is always at least 4k aligned, so skip lower bytes
+		return (uint32)(Key >> 12);
+	}
+	static FORCEINLINE bool Matches(const uint64 A, const uint64 B)
+	{
+		return A == B;
+	}
+};
+
+////////////////////////////////////////////////////////////////////////////////////////////////////
+
+void FMemAllocGroupingBySwapPage::GroupNodes(const TArray<FTableTreeNodePtr>& Nodes, FTableTreeNode& ParentGroup,
+	TWeakPtr<FTable> InParentTable, IAsyncOperationProgress& InAsyncOperationProgress) const
+{
+	ParentGroup.ClearChildren();
+
+	const FSlateBrush* IconBrush = FBaseTreeNode::GetDefaultIcon(true);
+	const FLinearColor Color(1.0f, 0.7f, 0.3f, 1.0f);
+
+	auto InSwapGroup = MakeShared<FCustomTableTreeNode>(FName(TEXT("In Swap")), InParentTable, IconBrush, Color);
+	auto InRamGroup = MakeShared<FCustomTableTreeNode>(FName(TEXT("In RAM")), InParentTable, IconBrush, Color);
+	ParentGroup.AddChildAndSetParent(InSwapGroup);
+	ParentGroup.AddChildAndSetParent(InRamGroup);
+
+	TMap<uint64, FTableTreeNode*, FDefaultSetAllocator, TMemoryPageMapKeyFuncs> SwapNodes;
+
+	for (FTableTreeNodePtr NodePtr : Nodes)
+	{
+		if (InAsyncOperationProgress.ShouldCancelAsyncOp())
+		{
+			return;
+		}
+
+		const FMemAllocNode& MemAllocNode = static_cast<const FMemAllocNode&>(*NodePtr);
+		const FMemoryAlloc* Alloc = MemAllocNode.GetMemAlloc();
+
+		if (Alloc && Alloc->IsSwap())
+		{
+			auto SwapEntry = MakeShared<FCustomTableTreeNode>(FName(FString::Printf(TEXT("0x%016llx"), Alloc->GetAddress())), InParentTable, IconBrush, Color);
+			InSwapGroup->AddChildAndSetParent(SwapEntry);
+			SwapEntry->AddChildAndSetParent(NodePtr);
+
+			ensure(SwapNodes.Find(Alloc->GetAddress()) == nullptr);
+
+			SwapNodes.Add(Alloc->GetAddress(), &SwapEntry.Get());
+		}
+	}
+
+	for (FTableTreeNodePtr NodePtr : Nodes)
+	{
+		if (InAsyncOperationProgress.ShouldCancelAsyncOp())
+		{
+			return;
+		}
+
+		if (NodePtr->IsGroup())
+		{
+			ParentGroup.AddChildAndSetParent(NodePtr);
+			continue;
+		}
+
+		const FMemAllocNode& MemAllocNode = static_cast<const FMemAllocNode&>(*NodePtr);
+		const FMemoryAlloc* Alloc = MemAllocNode.GetMemAlloc();
+
+		if (!Alloc)
+		{
+			InRamGroup->AddChildAndSetParent(NodePtr);
+			continue;
+		}
+
+		if (!Alloc->IsSwap())
+		{
+			bool bAnyInSwap = false;
+			const uint64 PageSize = AllocProvider.GetPlatformPageSize();
+			const uint64 PageMask = ~(PageSize - 1);
+
+			const uint64 AllocStart = Alloc->GetAddress();
+			const uint64 AllocEnd = AllocStart + FMath::Max(Alloc->GetSize(), 1); // fake alloc to have at least 1 byte size to make page walking logic work
+			const uint64 PageRangeStart = AllocStart & PageMask;
+			const uint64 PageRangeEnd = (AllocEnd + PageSize - 1) & PageMask;
+
+			// for every alloc go through every memory page and check if it's in swap
+			for (uint64 PageAddress = PageRangeStart; PageAddress < PageRangeEnd; PageAddress += PageSize)
+			{
+				FTableTreeNode** GroupNodePtr = SwapNodes.Find(PageAddress);
+				if (!GroupNodePtr)
+				{
+					continue;
+				}
+
+				FTableTreeNode* GroupNode = *GroupNodePtr;
+
+				uint64 AllocSizeInPage = PageSize;
+				if (PageAddress < AllocStart) // adjust size if allocation starts after page start 
+				{
+					AllocSizeInPage -= AllocStart - PageAddress;
+				}
+				if (PageAddress + PageSize > AllocEnd) // adjust size if allocation ends before page end
+				{
+					AllocSizeInPage -= PageAddress + PageSize - AllocEnd;
+				}
+
+				auto SwapAllocNode = MakeShared<FMemAllocInSwapNode>(MemAllocNode.GetName(), StaticCastWeakPtr<FMemAllocTable>(MemAllocNode.GetParentTable()), MemAllocNode.GetRowIndex(), AllocSizeInPage);
+				GroupNode->AddChildAndSetParent(SwapAllocNode);
+
+				bAnyInSwap = true;
+			}
+
+			if (!bAnyInSwap)
+			{
+				InRamGroup->AddChildAndSetParent(NodePtr);
+			}
+		}
+	}
+}
+
+////////////////////////////////////////////////////////////////////////////////////////////////////
+
+} // namespace Insights
+
+#undef LOCTEXT_NAMESPACE
diff --git a/Engine/Source/Developer/TraceInsights/Private/Insights/MemoryProfiler/ViewModels/MemAllocGroupingBySwapPage.h b/Engine/Source/Developer/TraceInsights/Private/Insights/MemoryProfiler/ViewModels/MemAllocGroupingBySwapPage.h
new file mode 100644
--- /dev/null
+++ b/Engine/Source/Developer/TraceInsights/Private/Insights/MemoryProfiler/ViewModels/MemAllocGroupingBySwapPage.h
@@ -0,0 +1,38 @@
+﻿// Copyright Epic Games, Inc. All Rights Reserved.
+
+#pragma once
+
+#include "CoreMinimal.h"
+
+// Insights
+#include "Insights/Table/ViewModels/BaseTreeNode.h"
+#include "Insights/Table/ViewModels/TreeNodeGrouping.h"
+#include "TraceServices/Model/AllocationsProvider.h"
+
+namespace TraceServices
+{
+	struct FStackFrame;
+}
+
+namespace Insights
+{
+class IAsyncOperationProgress;
+
+////////////////////////////////////////////////////////////////////////////////////////////////////
+
+class FMemAllocGroupingBySwapPage : public FTreeNodeGrouping
+{
+	INSIGHTS_DECLARE_RTTI(FMemAllocGroupingBySwapPage, FTreeNodeGrouping)
+
+public:
+	FMemAllocGroupingBySwapPage(const TraceServices::IAllocationsProvider& AllocProvider);
+	virtual ~FMemAllocGroupingBySwapPage() override;
+
+	virtual void GroupNodes(const TArray<FTableTreeNodePtr>& Nodes, FTableTreeNode& ParentGroup, TWeakPtr<FTable> InParentTable, IAsyncOperationProgress& InAsyncOperationProgress) const override;
+private:
+	const TraceServices::IAllocationsProvider& AllocProvider;
+};
+
+////////////////////////////////////////////////////////////////////////////////////////////////////
+
+} // namespace Insights
diff --git a/Engine/Source/Developer/TraceInsights/Private/Insights/MemoryProfiler/ViewModels/MemAllocInSwapNode.cpp b/Engine/Source/Developer/TraceInsights/Private/Insights/MemoryProfiler/ViewModels/MemAllocInSwapNode.cpp
new file mode 100644
--- /dev/null
+++ b/Engine/Source/Developer/TraceInsights/Private/Insights/MemoryProfiler/ViewModels/MemAllocInSwapNode.cpp
@@ -0,0 +1,23 @@
+﻿// Copyright Epic Games, Inc. All Rights Reserved.
+
+#include "MemAllocInSwapNode.h"
+
+#define LOCTEXT_NAMESPACE "Insights::FMemAllocSwapNode"
+
+namespace Insights
+{
+
+INSIGHTS_IMPLEMENT_RTTI(FMemAllocInSwapNode)
+
+////////////////////////////////////////////////////////////////////////////////////////////////////
+
+uint64 FMemAllocInSwapNode::GetBytesInSwapPage() const
+{
+	return BytesInSwapPage;
+}
+
+////////////////////////////////////////////////////////////////////////////////////////////////////
+
+} // namespace Insights
+
+#undef LOCTEXT_NAMESPACE
diff --git a/Engine/Source/Developer/TraceInsights/Private/Insights/MemoryProfiler/ViewModels/MemAllocInSwapNode.h b/Engine/Source/Developer/TraceInsights/Private/Insights/MemoryProfiler/ViewModels/MemAllocInSwapNode.h
new file mode 100644
--- /dev/null
+++ b/Engine/Source/Developer/TraceInsights/Private/Insights/MemoryProfiler/ViewModels/MemAllocInSwapNode.h
@@ -0,0 +1,36 @@
+﻿// Copyright Epic Games, Inc. All Rights Reserved.
+
+#pragma once
+
+#include "Insights/MemoryProfiler/ViewModels/MemAllocNode.h"
+
+////////////////////////////////////////////////////////////////////////////////////////////////////
+
+namespace Insights
+{
+
+////////////////////////////////////////////////////////////////////////////////////////////////////
+/**
+* Class used to store information about an allocation node in a particular swap page.
+*/
+class FMemAllocInSwapNode : public FMemAllocNode
+{
+	INSIGHTS_DECLARE_RTTI(FMemAllocInSwapNode, FMemAllocNode)
+
+public:
+	/** Initialization constructor for the MemAlloc node. */
+	explicit FMemAllocInSwapNode(const FName InName, TWeakPtr<FMemAllocTable> InParentTable, int32 InRowIndex, uint64 InBytesInSwapPage)
+		: FMemAllocNode(InName, InParentTable, InRowIndex)
+		, BytesInSwapPage(InBytesInSwapPage)
+	{
+	}
+
+	uint64 GetBytesInSwapPage() const;
+
+private:
+	const uint64 BytesInSwapPage;
+};
+
+////////////////////////////////////////////////////////////////////////////////////////////////////
+
+} // namespace Insights
diff --git a/Engine/Source/Developer/TraceInsights/Private/Insights/MemoryProfiler/ViewModels/MemAllocNode.h b/Engine/Source/Developer/TraceInsights/Private/Insights/MemoryProfiler/ViewModels/MemAllocNode.h
--- a/Engine/Source/Developer/TraceInsights/Private/Insights/MemoryProfiler/ViewModels/MemAllocNode.h
+++ b/Engine/Source/Developer/TraceInsights/Private/Insights/MemoryProfiler/ViewModels/MemAllocNode.h
@@ -62,6 +62,7 @@ public:
 	}
 
 	FMemAllocTable& GetMemTableChecked() const { return *MemAllocTable; }
+	FMemAllocTable* GetMemTable() const { return MemAllocTable; }
 
 	bool IsValidMemAlloc() const { return GetMemTableChecked().IsValidRowIndex(GetRowIndex()); }
 	const FMemoryAlloc* GetMemAlloc() const { return GetMemTableChecked().GetMemAlloc(GetRowIndex()); }
diff --git a/Engine/Source/Developer/TraceInsights/Private/Insights/MemoryProfiler/ViewModels/MemAllocTable.cpp b/Engine/Source/Developer/TraceInsights/Private/Insights/MemoryProfiler/ViewModels/MemAllocTable.cpp
--- a/Engine/Source/Developer/TraceInsights/Private/Insights/MemoryProfiler/ViewModels/MemAllocTable.cpp
+++ b/Engine/Source/Developer/TraceInsights/Private/Insights/MemoryProfiler/ViewModels/MemAllocTable.cpp
@@ -12,6 +12,7 @@
 #include "Insights/MemoryProfiler/ViewModels/MemAllocFilterValueConverter.h"
 #include "Insights/MemoryProfiler/ViewModels/MemAllocGroupingByTag.h"
 #include "Insights/MemoryProfiler/ViewModels/MemAllocNode.h"
+#include "Insights/MemoryProfiler/ViewModels/MemAllocInSwapNode.h"
 #include "Insights/MemoryProfiler/ViewModels/MemAllocTable.h"
 #include "Insights/Table/ViewModels/TableCellValueFormatter.h"
 #include "Insights/Table/ViewModels/TableCellValueGetter.h"
@@ -38,6 +39,8 @@ const FName FMemAllocTableColumns::FreeThreadColumnId(TEXT("FreeThread"));
 const FName FMemAllocTableColumns::AddressColumnId(TEXT("Address"));
 const FName FMemAllocTableColumns::MemoryPageColumnId(TEXT("MemoryPage"));
 const FName FMemAllocTableColumns::CountColumnId(TEXT("Count"));
+const FName FMemAllocTableColumns::SwapCompressedSizeColumnId(TEXT("SwapCompressedSize"));
+const FName FMemAllocTableColumns::SizeInSwapColumnId(TEXT("SizeInSwap"));
 const FName FMemAllocTableColumns::SizeColumnId(TEXT("Size"));
 const FName FMemAllocTableColumns::LLMSizeColumnId(TEXT("LLMSize"));
 const FName FMemAllocTableColumns::LLMDeltaSizeColumnId(TEXT("LLMDeltaSize"));
@@ -657,10 +660,11 @@ void FMemAllocTable::AddDefaultColumns()
 				else
 				{
 					const FMemAllocNode& MemAllocNode = static_cast<const FMemAllocNode&>(Node);
+					const FMemAllocTable& MemAllocTable = MemAllocNode.GetMemTableChecked();
 					const FMemoryAlloc* Alloc = MemAllocNode.GetMemAlloc();
 					if (Alloc)
 					{
-						return FTableCellValue(static_cast<int64>(Alloc->GetPage()));
+						return FTableCellValue(static_cast<int64>(MemAllocTable.GetAddressPage(Alloc->GetAddress())));
 					}
 				}
 
@@ -738,6 +742,128 @@ void FMemAllocTable::AddDefaultColumns()
 		AddColumn(ColumnRef);
 	}
 	//////////////////////////////////////////////////
+	// Swap Compressed Size
+	{
+		TSharedRef<FTableColumn> ColumnRef = MakeShared<FTableColumn>(FMemAllocTableColumns::SwapCompressedSizeColumnId);
+		FTableColumn& Column = *ColumnRef;
+
+		Column.SetIndex(ColumnIndex++);
+
+		Column.SetShortName(LOCTEXT("SwapCompressedSizeColumnName", "Compressed Size In Swap"));
+		Column.SetTitleName(LOCTEXT("SwapCompressedSizeColumnTitle", "Compressed Size In Swap"));
+		Column.SetDescription(LOCTEXT("SwapCompressedSizeColumnDesc", "Shows compressed size of swap page."));
+
+		Column.SetFlags(ETableColumnFlags::CanBeHidden | ETableColumnFlags::CanBeFiltered);
+
+		Column.SetHorizontalAlignment(HAlign_Left);
+		Column.SetInitialWidth(100.0f);
+
+		Column.SetDataType(ETableCellDataType::Int64);
+
+		class FSwapCompressedSizeColumnValueGetter : public FTableCellValueGetter
+		{
+		public:
+			virtual const TOptional<FTableCellValue> GetValue(const FTableColumn& Column, const FBaseTreeNode& Node) const override
+			{
+				if (Node.IsGroup())
+				{
+					const FTableTreeNode& NodePtr = static_cast<const FTableTreeNode&>(Node);
+					if (NodePtr.HasAggregatedValue(Column.GetId()))
+					{
+						return NodePtr.GetAggregatedValue(Column.GetId());
+					}
+				}
+				else if (Node.Is<FMemAllocNode>())
+				{
+					const FMemAllocNode& MemAllocNode = static_cast<const FMemAllocNode&>(Node);
+					const FMemoryAlloc* Alloc = MemAllocNode.GetMemAlloc();
+					if (Alloc && Alloc->IsSwap())
+					{
+						return FTableCellValue(static_cast<int64>(Alloc->GetSize()));
+					}
+				}
+
+				return TOptional<FTableCellValue>();
+			}
+		};
+
+		TSharedRef<ITableCellValueGetter> Getter = MakeShared<FSwapCompressedSizeColumnValueGetter>();
+		Column.SetValueGetter(Getter);
+
+		TSharedRef<ITableCellValueFormatter> Formatter = MakeShared<FInt64ValueFormatterAsMemory>();
+		Column.SetValueFormatter(Formatter);
+
+		TSharedRef<ITableCellValueSorter> Sorter = MakeShared<FSorterByInt64Value>(ColumnRef);
+		Column.SetValueSorter(Sorter);
+		Column.SetInitialSortMode(EColumnSortMode::Descending);
+
+		TSharedRef<IFilterValueConverter> Converter = MakeShared<FMemoryFilterValueConverter>();
+		Column.SetValueConverter(Converter);
+
+		Column.SetAggregation(ETableColumnAggregation::Sum);
+
+		AddColumn(ColumnRef);
+	}
+	//////////////////////////////////////////////////
+	// Page Size in Swap Column
+	{
+		TSharedRef<FTableColumn> ColumnRef = MakeShared<FTableColumn>(FMemAllocTableColumns::SizeInSwapColumnId);
+		FTableColumn& Column = *ColumnRef;
+
+		Column.SetIndex(ColumnIndex++);
+
+		Column.SetShortName(LOCTEXT("SizeInSwapColumnName", "Size in Swap"));
+		Column.SetTitleName(LOCTEXT("SizeInSwapColumnTitle", "Size in Swap"));
+		Column.SetDescription(LOCTEXT("SizeInSwapColumnDesc", "Shows (partial) size of allocation which is stored in the correponding swap page."));
+
+		Column.SetFlags(ETableColumnFlags::CanBeHidden | ETableColumnFlags::CanBeFiltered);
+
+		Column.SetHorizontalAlignment(HAlign_Left);
+		Column.SetInitialWidth(100.0f);
+
+		Column.SetDataType(ETableCellDataType::Int64);
+
+		class FSizeInSwapColumnValueGetter : public FTableCellValueGetter
+		{
+		public:
+			virtual const TOptional<FTableCellValue> GetValue(const FTableColumn& Column, const FBaseTreeNode& Node) const override
+			{
+				if (Node.IsGroup())
+				{
+					const FTableTreeNode& NodePtr = static_cast<const FTableTreeNode&>(Node);
+					if (NodePtr.HasAggregatedValue(Column.GetId()))
+					{
+						return NodePtr.GetAggregatedValue(Column.GetId());
+					}
+				}
+				else if (Node.Is<FMemAllocInSwapNode>())
+				{
+					const FMemAllocInSwapNode& MemAllocInSwapNode = static_cast<const FMemAllocInSwapNode&>(Node);
+					return FTableCellValue(static_cast<int64>(MemAllocInSwapNode.GetBytesInSwapPage()));
+				}
+
+				return TOptional<FTableCellValue>();
+			}
+		};
+
+		TSharedRef<ITableCellValueGetter> Getter = MakeShared<FSizeInSwapColumnValueGetter>();
+		Column.SetValueGetter(Getter);
+
+		TSharedRef<ITableCellValueFormatter> Formatter = MakeShared<FInt64ValueFormatterAsMemory>();
+		Column.SetValueFormatter(Formatter);
+
+		TSharedRef<ITableCellValueSorter> Sorter = MakeShared<FSorterByInt64Value>(ColumnRef);
+		Column.SetValueSorter(Sorter);
+		Column.SetInitialSortMode(EColumnSortMode::Descending);
+
+		TSharedRef<IFilterValueConverter> Converter = MakeShared<FMemoryFilterValueConverter>();
+		Column.SetValueConverter(Converter);
+
+		Column.SetAggregation(ETableColumnAggregation::Sum);
+
+		AddColumn(ColumnRef);
+	}
+	//////////////////////////////////////////////////
 	// Size Column
 	{
 		TSharedRef<FTableColumn> ColumnRef = MakeShared<FTableColumn>(FMemAllocTableColumns::SizeColumnId);
@@ -775,7 +901,14 @@ void FMemAllocTable::AddDefaultColumns()
 					const FMemoryAlloc* Alloc = MemAllocNode.GetMemAlloc();
 					if (Alloc)
 					{
-						return FTableCellValue(static_cast<int64>(Alloc->GetSize()));
+						if (Alloc->IsSwap())
+						{
+							return FTableCellValue(static_cast<int64>(MemAllocNode.GetMemTableChecked().GetPlatformPageSize()));
+						}
+						else
+						{
+							return FTableCellValue(static_cast<int64>(Alloc->GetSize()));
+						}
 					}
 				}
 
diff --git a/Engine/Source/Developer/TraceInsights/Private/Insights/MemoryProfiler/ViewModels/MemAllocTable.h b/Engine/Source/Developer/TraceInsights/Private/Insights/MemoryProfiler/ViewModels/MemAllocTable.h
--- a/Engine/Source/Developer/TraceInsights/Private/Insights/MemoryProfiler/ViewModels/MemAllocTable.h
+++ b/Engine/Source/Developer/TraceInsights/Private/Insights/MemoryProfiler/ViewModels/MemAllocTable.h
@@ -29,6 +29,8 @@ struct FMemAllocTableColumns
 	static const FName AddressColumnId;
 	static const FName MemoryPageColumnId;
 	static const FName CountColumnId;
+	static const FName SwapCompressedSizeColumnId;
+	static const FName SizeInSwapColumnId;
 	static const FName SizeColumnId;
 	static const FName LLMSizeColumnId;
 	static const FName LLMDeltaSizeColumnId;
@@ -66,12 +68,18 @@ public:
 	double GetTimeMarkerA() const { return TimeA; }
 	void SetTimeMarkerA(double InTime) { TimeA = InTime; }
 
+	uint64 GetAddressPage(uint64 Address) const { return Address & ~(PlatformPageSize - 1); }
+
+	void SetPlatformPageSize(uint64 InPlatformPageSize) { PlatformPageSize = InPlatformPageSize; }
+	uint64 GetPlatformPageSize() const { return PlatformPageSize; }
+
 private:
 	void AddDefaultColumns();
 
 private:
 	TArray<FMemoryAlloc> Allocs;
 	double TimeA = 0.0;
+	uint64 PlatformPageSize = 0;
 };
 
 ////////////////////////////////////////////////////////////////////////////////////////////////////
diff --git a/Engine/Source/Developer/TraceInsights/Private/Insights/MemoryProfiler/ViewModels/MemoryAlloc.h b/Engine/Source/Developer/TraceInsights/Private/Insights/MemoryProfiler/ViewModels/MemoryAlloc.h
--- a/Engine/Source/Developer/TraceInsights/Private/Insights/MemoryProfiler/ViewModels/MemoryAlloc.h
+++ b/Engine/Source/Developer/TraceInsights/Private/Insights/MemoryProfiler/ViewModels/MemoryAlloc.h
@@ -29,7 +29,6 @@ public:
 	double GetDuration() const { return EndTime - StartTime; }
 
 	uint64 GetAddress() const { return Address; }
-	uint64 GetPage() const { return Address & ~(4llu*1024-1); }
 
 	int64 GetSize() const { return Size; }
 
@@ -50,6 +49,7 @@ public:
 
 	HeapId GetRootHeap() const { return RootHeap; }
 	bool IsHeap() const { return bIsHeap; }
+	bool IsSwap() const { return bIsSwap; }
 
 	bool IsDecline() const { return bIsDecline; }
 
@@ -73,6 +73,7 @@ private:
 	const TraceServices::FCallstack* FreeCallstack = nullptr;
 	HeapId RootHeap = 0;
 	bool bIsHeap = false;
+	bool bIsSwap = false;
 	bool bIsDecline = false;
 };
 
diff --git a/Engine/Source/Developer/TraceInsights/Private/Insights/MemoryProfiler/ViewModels/MemoryGraphTrack.cpp b/Engine/Source/Developer/TraceInsights/Private/Insights/MemoryProfiler/ViewModels/MemoryGraphTrack.cpp
--- a/Engine/Source/Developer/TraceInsights/Private/Insights/MemoryProfiler/ViewModels/MemoryGraphTrack.cpp
+++ b/Engine/Source/Developer/TraceInsights/Private/Insights/MemoryProfiler/ViewModels/MemoryGraphTrack.cpp
@@ -437,6 +437,46 @@ TSharedPtr<FMemoryGraphSeries> FMemoryGraphTrack::AddTimelineSeries(FMemoryGraph
 				break;
 			}
 
+			case FMemoryGraphSeries::ETimelineType::MinSwapMem:
+			{
+				Series->SetName(TEXT("Total Swap Memory (Min)"));
+				Series->SetDescription(TEXT("Minimum value per sample for the Total Swap Memory"));
+				const FLinearColor Color = FLinearColor(0.0f, 0.5f, 1.0f, 1.0f);
+				const FLinearColor BorderColor(FMath::Min(Color.R + 0.4f, 1.0f), FMath::Min(Color.G + 0.4f, 1.0f), FMath::Min(Color.B + 0.4f, 1.0f), 1.0f);
+				Series->SetColor(Color, BorderColor, Color.CopyWithNewOpacity(0.1f));
+				break;
+			}
+
+			case FMemoryGraphSeries::ETimelineType::MaxSwapMem:
+			{
+				Series->SetName(TEXT("Total Swap Memory (Max)"));
+				Series->SetDescription(TEXT("Maximum value per sample for the Total Swap Memory"));
+				const FLinearColor Color = FLinearColor(1.0f, 0.25f, 1.0f, 1.0f);
+				const FLinearColor BorderColor(FMath::Min(Color.R + 0.4f, 1.0f), FMath::Min(Color.G + 0.4f, 1.0f), FMath::Min(Color.B + 0.4f, 1.0f), 1.0f);
+				Series->SetColor(Color, BorderColor, Color.CopyWithNewOpacity(0.1f));
+				break;
+			}
+
+			case FMemoryGraphSeries::ETimelineType::MinCompressedSwapMem:
+			{
+				Series->SetName(TEXT("Total Compressed Swap Memory (Min)"));
+				Series->SetDescription(TEXT("Minimum value per sample for the Total Compressed Swap Memory"));
+				const FLinearColor Color = FLinearColor(1.0f, 1.0f, 0.25f, 1.0f);
+				const FLinearColor BorderColor(FMath::Min(Color.R + 0.4f, 1.0f), FMath::Min(Color.G + 0.4f, 1.0f), FMath::Min(Color.B + 0.4f, 1.0f), 1.0f);
+				Series->SetColor(Color, BorderColor, Color.CopyWithNewOpacity(0.1f));
+				break;
+			}
+
+			case FMemoryGraphSeries::ETimelineType::MaxCompressedSwapMem:
+			{
+				Series->SetName(TEXT("Total Compressed Swap Memory (Max)"));
+				Series->SetDescription(TEXT("Maximum value per sample for the Total Compressed Swap Memory"));
+				const FLinearColor Color = FLinearColor(1.0f, 0.25f, 1.0f, 1.0f);
+				const FLinearColor BorderColor(FMath::Min(Color.R + 0.4f, 1.0f), FMath::Min(Color.G + 0.4f, 1.0f), FMath::Min(Color.B + 0.4f, 1.0f), 1.0f);
+				Series->SetColor(Color, BorderColor, Color.CopyWithNewOpacity(0.1f));
+				break;
+			}
+
 			case FMemoryGraphSeries::ETimelineType::AllocEvents:
 			{
 				Series->SetName(TEXT("Alloc Event Count"));
@@ -456,6 +496,36 @@ TSharedPtr<FMemoryGraphSeries> FMemoryGraphTrack::AddTimelineSeries(FMemoryGraph
 				Series->SetColor(Color, BorderColor, Color.CopyWithNewOpacity(0.1f));
 				break;
 			}
+
+			case FMemoryGraphSeries::ETimelineType::PageInEvents:
+			{
+				Series->SetName(TEXT("Page In Event Count"));
+				Series->SetDescription(TEXT("Number of page in events per sample"));
+				const FLinearColor Color = FLinearColor(0.0f, 1.0f, 0.5f, 1.0f);
+				const FLinearColor BorderColor(FMath::Min(Color.R + 0.4f, 1.0f), FMath::Min(Color.G + 0.4f, 1.0f), FMath::Min(Color.B + 0.4f, 1.0f), 1.0f);
+				Series->SetColor(Color, BorderColor, Color.CopyWithNewOpacity(0.1f));
+				break;
+			}
+
+			case FMemoryGraphSeries::ETimelineType::PageOutEvents:
+			{
+				Series->SetName(TEXT("Page Out Event Count"));
+				Series->SetDescription(TEXT("Number of page out events per sample"));
+				const FLinearColor Color = FLinearColor(1.0f, 0.5f, 0.25f, 1.0f);
+				const FLinearColor BorderColor(FMath::Min(Color.R + 0.4f, 1.0f), FMath::Min(Color.G + 0.4f, 1.0f), FMath::Min(Color.B + 0.4f, 1.0f), 1.0f);
+				Series->SetColor(Color, BorderColor, Color.CopyWithNewOpacity(0.1f));
+				break;
+			}
+
+			case FMemoryGraphSeries::ETimelineType::SwapFreeEvents:
+			{
+				Series->SetName(TEXT("Swap Free Event Count"));
+				Series->SetDescription(TEXT("Number of swap free events per sample"));
+				const FLinearColor Color = FLinearColor(0.25f, 0.5f, 1.0f, 1.0f);
+				const FLinearColor BorderColor(FMath::Min(Color.R + 0.4f, 1.0f), FMath::Min(Color.G + 0.4f, 1.0f), FMath::Min(Color.B + 0.4f, 1.0f), 1.0f);
+				Series->SetColor(Color, BorderColor, Color.CopyWithNewOpacity(0.1f));
+				break;
+			}
 		}
 
 		Series->SetTimelineType(InTimelineType);
@@ -533,12 +603,33 @@ void FMemoryGraphTrack::PreUpdateAllocationsTimelineSeries(FMemoryGraphSeries& S
 			case FMemoryGraphSeries::ETimelineType::MaxLiveAllocs:
 				AllocationsProvider->EnumerateMaxLiveAllocationsTimeline(StartIndex, EndIndex, Callback32);
 				break;
+			case FMemoryGraphSeries::ETimelineType::MinSwapMem:
+				AllocationsProvider->EnumerateMinTotalSwapMemoryTimeline(StartIndex, EndIndex, Callback64);
+				break;
+			case FMemoryGraphSeries::ETimelineType::MaxSwapMem:
+				AllocationsProvider->EnumerateMaxTotalSwapMemoryTimeline(StartIndex, EndIndex, Callback64);
+				break;
+			case FMemoryGraphSeries::ETimelineType::MinCompressedSwapMem:
+				AllocationsProvider->EnumerateMinTotalCompressedSwapMemoryTimeline(StartIndex, EndIndex, Callback64);
+				break;
+			case FMemoryGraphSeries::ETimelineType::MaxCompressedSwapMem:
+				AllocationsProvider->EnumerateMaxTotalCompressedSwapMemoryTimeline(StartIndex, EndIndex, Callback64);
+				break;
 			case FMemoryGraphSeries::ETimelineType::AllocEvents:
 				AllocationsProvider->EnumerateAllocEventsTimeline(StartIndex, EndIndex, Callback32);
 				break;
 			case FMemoryGraphSeries::ETimelineType::FreeEvents:
 				AllocationsProvider->EnumerateFreeEventsTimeline(StartIndex, EndIndex, Callback32);
 				break;
+			case FMemoryGraphSeries::ETimelineType::PageInEvents:
+				AllocationsProvider->EnumeratePageInEventsTimeline(StartIndex, EndIndex, Callback32);
+				break;
+			case FMemoryGraphSeries::ETimelineType::PageOutEvents:
+				AllocationsProvider->EnumeratePageOutEventsTimeline(StartIndex, EndIndex, Callback32);
+				break;
+			case FMemoryGraphSeries::ETimelineType::SwapFreeEvents:
+				AllocationsProvider->EnumerateSwapFreeEventsTimeline(StartIndex, EndIndex, Callback32);
+				break;
 			}
 
 			if (Series.GetTimelineType() == FMemoryGraphSeries::ETimelineType::FreeEvents)
@@ -632,12 +723,33 @@ void FMemoryGraphTrack::UpdateAllocationsTimelineSeries(FMemoryGraphSeries& Seri
 				case FMemoryGraphSeries::ETimelineType::MaxLiveAllocs:
 					AllocationsProvider->EnumerateMaxLiveAllocationsTimeline(StartIndex, EndIndex, Callback32);
 					break;
+				case FMemoryGraphSeries::ETimelineType::MinSwapMem:
+					AllocationsProvider->EnumerateMinTotalSwapMemoryTimeline(StartIndex, EndIndex, Callback64);
+					break;
+				case FMemoryGraphSeries::ETimelineType::MaxSwapMem:
+					AllocationsProvider->EnumerateMaxTotalSwapMemoryTimeline(StartIndex, EndIndex, Callback64);
+					break;
+				case FMemoryGraphSeries::ETimelineType::MinCompressedSwapMem:
+					AllocationsProvider->EnumerateMinTotalCompressedSwapMemoryTimeline(StartIndex, EndIndex, Callback64);
+					break;
+				case FMemoryGraphSeries::ETimelineType::MaxCompressedSwapMem:
+					AllocationsProvider->EnumerateMaxTotalCompressedSwapMemoryTimeline(StartIndex, EndIndex, Callback64);
+					break;
 				case FMemoryGraphSeries::ETimelineType::AllocEvents:
 					AllocationsProvider->EnumerateAllocEventsTimeline(StartIndex, EndIndex, Callback32);
 					break;
 				case FMemoryGraphSeries::ETimelineType::FreeEvents:
 					AllocationsProvider->EnumerateFreeEventsTimeline(StartIndex, EndIndex, Callback32Negative);
 					break;
+				case FMemoryGraphSeries::ETimelineType::PageInEvents:
+					AllocationsProvider->EnumeratePageInEventsTimeline(StartIndex, EndIndex, Callback32);
+					break;
+				case FMemoryGraphSeries::ETimelineType::PageOutEvents:
+					AllocationsProvider->EnumeratePageOutEventsTimeline(StartIndex, EndIndex, Callback32Negative);
+					break;
+				case FMemoryGraphSeries::ETimelineType::SwapFreeEvents:
+					AllocationsProvider->EnumerateSwapFreeEventsTimeline(StartIndex, EndIndex, Callback32);
+					break;
 				}
 			}
 
diff --git a/Engine/Source/Developer/TraceInsights/Private/Insights/MemoryProfiler/ViewModels/MemoryGraphTrack.h b/Engine/Source/Developer/TraceInsights/Private/Insights/MemoryProfiler/ViewModels/MemoryGraphTrack.h
--- a/Engine/Source/Developer/TraceInsights/Private/Insights/MemoryProfiler/ViewModels/MemoryGraphTrack.h
+++ b/Engine/Source/Developer/TraceInsights/Private/Insights/MemoryProfiler/ViewModels/MemoryGraphTrack.h
@@ -44,8 +44,15 @@ public:
 		MaxTotalMem,
 		MinLiveAllocs,
 		MaxLiveAllocs,
+		MinSwapMem,
+		MaxSwapMem,
+		MinCompressedSwapMem,
+		MaxCompressedSwapMem,
 		AllocEvents,
-		FreeEvents
+		FreeEvents,
+		PageInEvents,
+		PageOutEvents,
+		SwapFreeEvents
 	};
 
 public:
diff --git a/Engine/Source/Developer/TraceInsights/Private/Insights/MemoryProfiler/ViewModels/MemorySharedState.cpp b/Engine/Source/Developer/TraceInsights/Private/Insights/MemoryProfiler/ViewModels/MemorySharedState.cpp
--- a/Engine/Source/Developer/TraceInsights/Private/Insights/MemoryProfiler/ViewModels/MemorySharedState.cpp
+++ b/Engine/Source/Developer/TraceInsights/Private/Insights/MemoryProfiler/ViewModels/MemorySharedState.cpp
@@ -78,6 +78,8 @@ FMemorySharedState::FMemorySharedState()
 	, MainGraphTrack(nullptr)
 	, LiveAllocsGraphTrack(nullptr)
 	, AllocFreeGraphTrack(nullptr)
+	, SwapMemoryGraphTrack(nullptr)
+	, PageSwapGraphTrack(nullptr)
 	, TrackHeightMode(EMemoryTrackHeightMode::Medium)
 	, bShowHideAllMemoryTracks(false)
 	, CreatedDefaultTracks()
@@ -113,6 +115,8 @@ void FMemorySharedState::OnBeginSession(Insights::ITimingViewSession& InSession)
 	MainGraphTrack = nullptr;
 	LiveAllocsGraphTrack = nullptr;
 	AllocFreeGraphTrack = nullptr;
+	SwapMemoryGraphTrack = nullptr;
+	PageSwapGraphTrack = nullptr;
 	AllTracks.Reset();
 
 	bShowHideAllMemoryTracks = true;
@@ -138,6 +142,8 @@ void FMemorySharedState::OnEndSession(Insights::ITimingViewSession& InSession)
 	MainGraphTrack = nullptr;
 	LiveAllocsGraphTrack = nullptr;
 	AllocFreeGraphTrack = nullptr;
+	SwapMemoryGraphTrack = nullptr;
+	PageSwapGraphTrack = nullptr;
 	AllTracks.Reset();
 
 	bShowHideAllMemoryTracks = false;
@@ -219,6 +225,53 @@ void FMemorySharedState::Tick(Insights::ITimingViewSession& InSession, const Tra
 		TimingView->InvalidateScrollableTracksOrder();
 	}
 
+	if (!SwapMemoryGraphTrack.IsValid())
+	{
+		SwapMemoryGraphTrack = CreateMemoryGraphTrack();
+		check(SwapMemoryGraphTrack);
+
+		SwapMemoryGraphTrack->SetOrder(FTimingTrackOrder::First + 3);
+		SwapMemoryGraphTrack->SetName(TEXT("Swap Memory Graph"));
+
+		SwapMemoryGraphTrack->AddTimelineSeries(FMemoryGraphSeries::ETimelineType::MaxSwapMem);
+		SwapMemoryGraphTrack->AddTimelineSeries(FMemoryGraphSeries::ETimelineType::MinSwapMem);
+		SwapMemoryGraphTrack->AddTimelineSeries(FMemoryGraphSeries::ETimelineType::MaxCompressedSwapMem);
+		SwapMemoryGraphTrack->AddTimelineSeries(FMemoryGraphSeries::ETimelineType::MinCompressedSwapMem);
+		SwapMemoryGraphTrack->SetLabelUnit(EGraphTrackLabelUnit::MiB, 1);
+
+		SwapMemoryGraphTrack->SetVisibilityFlag(false); // don't show swaps track by default
+
+		SwapMemoryGraphTrack->SetAvailableTrackHeight(EMemoryTrackHeightMode::Small, 50.0f);
+		SwapMemoryGraphTrack->SetAvailableTrackHeight(EMemoryTrackHeightMode::Medium, 100.0f);
+		SwapMemoryGraphTrack->SetAvailableTrackHeight(EMemoryTrackHeightMode::Large, 200.0f);
+		SwapMemoryGraphTrack->SetCurrentTrackHeight(TrackHeightMode);
+
+		TimingView->InvalidateScrollableTracksOrder();
+	}
+
+	if (!PageSwapGraphTrack.IsValid())
+	{
+		PageSwapGraphTrack = CreateMemoryGraphTrack();
+		check(PageSwapGraphTrack);
+
+		PageSwapGraphTrack->SetOrder(FTimingTrackOrder::First + 4);
+		PageSwapGraphTrack->SetName(TEXT("Page In/Out Event Count"));
+		PageSwapGraphTrack->SetLabelUnit(EGraphTrackLabelUnit::Count, 0);
+
+		PageSwapGraphTrack->AddTimelineSeries(FMemoryGraphSeries::ETimelineType::PageInEvents);
+		PageSwapGraphTrack->AddTimelineSeries(FMemoryGraphSeries::ETimelineType::PageOutEvents);
+		PageSwapGraphTrack->AddTimelineSeries(FMemoryGraphSeries::ETimelineType::SwapFreeEvents);
+
+		PageSwapGraphTrack->SetVisibilityFlag(false); // don't show swap tracks by default
+
+		PageSwapGraphTrack->SetAvailableTrackHeight(EMemoryTrackHeightMode::Small, 50.0f);
+		PageSwapGraphTrack->SetAvailableTrackHeight(EMemoryTrackHeightMode::Medium, 100.0f);
+		PageSwapGraphTrack->SetAvailableTrackHeight(EMemoryTrackHeightMode::Large, 200.0f);
+		PageSwapGraphTrack->SetCurrentTrackHeight(TrackHeightMode);
+
+		TimingView->InvalidateScrollableTracksOrder();
+	}
+
 	const int32 PrevTagCount = TagList.GetTags().Num();
 
 	TagList.Update();
@@ -1028,6 +1081,18 @@ void FMemorySharedState::InitMemoryRules()
 		LOCTEXT("MemRule_AaBCfD_Verbose", "Specific Lifetime"),
 		LOCTEXT("MemRule_AaBCfD_Desc", "Identifies allocations allocated between time A and time B and freed between time C and time D.\n(A ≤ a ≤ B ≤ C ≤ f ≤ D)")));
 
+	MemoryRules.Add(MakeShared<Insights::FMemoryRuleSpec>(
+		ERule::AoB, 2,
+		LOCTEXT("MemRule_AoB_Short", "A ↓ B"),
+		LOCTEXT("MemRule_AoB_Verbose", "Paged-Out Allocs"),
+		LOCTEXT("MemRule_AoB_Desc", "Identifies allocations paged-out (swapped-out) between time A and time B.\n(A ≤ page-out ≤ B)")));
+
+	MemoryRules.Add(MakeShared<Insights::FMemoryRuleSpec>(
+		ERule::AiB, 2,
+		LOCTEXT("MemRule_AiB_Short", "A ↑ B"),
+		LOCTEXT("MemRule_AiB_Verbose", "Paged-In Allocs"),
+		LOCTEXT("MemRule_AiB_Desc", "Identifies allocations paged-in (swapped-in) between time A and time B.\n(A ≤ page-in ≤ B)")));
+
 	//TODO
 	//MemoryRules.Add(MakeShared<Insights::FMemoryRuleSpec>(
 	//	ERule::A_vs_B, 2,
diff --git a/Engine/Source/Developer/TraceInsights/Private/Insights/MemoryProfiler/ViewModels/MemorySharedState.h b/Engine/Source/Developer/TraceInsights/Private/Insights/MemoryProfiler/ViewModels/MemorySharedState.h
--- a/Engine/Source/Developer/TraceInsights/Private/Insights/MemoryProfiler/ViewModels/MemorySharedState.h
+++ b/Engine/Source/Developer/TraceInsights/Private/Insights/MemoryProfiler/ViewModels/MemorySharedState.h
@@ -180,6 +180,8 @@ private:
 	TSharedPtr<FMemoryGraphTrack> MainGraphTrack; // the Main Memory Graph track; also hosts the Total Allocated Memory series
 	TSharedPtr<FMemoryGraphTrack> LiveAllocsGraphTrack; // the graph track for the Live Allocation Count series
 	TSharedPtr<FMemoryGraphTrack> AllocFreeGraphTrack; // the graph track for the Alloc Event Count and the Free Event Count series
+	TSharedPtr<FMemoryGraphTrack> SwapMemoryGraphTrack; // the swap memory graph for Total Swap Memory and Total Compressed Swap Memory series
+	TSharedPtr<FMemoryGraphTrack> PageSwapGraphTrack; // the graph track for the Page In Event Count and the Page Out Event Count series
 	TSet<TSharedPtr<FMemoryGraphTrack>> AllTracks;
 
 	EMemoryTrackHeightMode TrackHeightMode;
diff --git a/Engine/Source/Developer/TraceInsights/Private/Insights/MemoryProfiler/Widgets/SMemAllocTableTreeView.cpp b/Engine/Source/Developer/TraceInsights/Private/Insights/MemoryProfiler/Widgets/SMemAllocTableTreeView.cpp
--- a/Engine/Source/Developer/TraceInsights/Private/Insights/MemoryProfiler/Widgets/SMemAllocTableTreeView.cpp
+++ b/Engine/Source/Developer/TraceInsights/Private/Insights/MemoryProfiler/Widgets/SMemAllocTableTreeView.cpp
@@ -34,6 +34,7 @@
 #include "Insights/MemoryProfiler/ViewModels/CallstackFormatting.h"
 #include "Insights/MemoryProfiler/ViewModels/MemAllocGroupingByCallstack.h"
 #include "Insights/MemoryProfiler/ViewModels/MemAllocGroupingByHeap.h"
+#include "Insights/MemoryProfiler/ViewModels/MemAllocGroupingBySwapPage.h"
 #include "Insights/MemoryProfiler/ViewModels/MemAllocGroupingBySize.h"
 #include "Insights/MemoryProfiler/ViewModels/MemAllocGroupingByTag.h"
 #include "Insights/MemoryProfiler/ViewModels/MemAllocNode.h"
@@ -144,8 +145,10 @@ void SMemAllocTableTreeView::RebuildTree(bool bResync)
 				TableRowNodes.Reserve(TotalAllocCount);
 
 				uint32 HeapAllocCount = 0;
+				uint32 SwapAllocCount = 0;
 				const FName BaseNodeName(TEXT("alloc"));
 				const FName BaseHeapName(TEXT("heap"));
+				const FName BaseSwapName(TEXT("swap"));
 				for (int32 AllocIndex = TableRowNodes.Num(); AllocIndex < TotalAllocCount; ++AllocIndex)
 				{
 					const FMemoryAlloc* Alloc = MemAllocTable->GetMemAlloc(AllocIndex);
@@ -159,11 +162,20 @@ void SMemAllocTableTreeView::RebuildTree(bool bResync)
 						}
 					}
 
-					FName NodeName(Alloc->bIsHeap ? BaseHeapName : BaseNodeName, static_cast<int32>(Alloc->GetStartEventIndex() + 1));
+					if (Alloc->bIsSwap)
+					{
+						++SwapAllocCount;
+						if (!bIncludeSwapAllocs)
+						{
+							continue;
+						}
+					}
+
+					FName NodeName(Alloc->bIsHeap ? BaseHeapName : (Alloc->bIsSwap ? BaseSwapName : BaseNodeName), static_cast<int32>(Alloc->GetStartEventIndex() + 1));
 					FMemAllocNodePtr NodePtr = MakeShared<FMemAllocNode>(NodeName, MemAllocTable, AllocIndex);
 					TableRowNodes.Add(NodePtr);
 				}
-				ensure(TableRowNodes.Num() == (bIncludeHeapAllocs ? TotalAllocCount : TotalAllocCount - HeapAllocCount));
+				ensure(TableRowNodes.Num() == TotalAllocCount - (bIncludeHeapAllocs ? 0 : HeapAllocCount) - (bIncludeSwapAllocs ? 0 : SwapAllocCount));
 			}
 		}
 	}
@@ -273,6 +285,13 @@ void SMemAllocTableTreeView::StartQuery()
 	{
 		const TraceServices::IAllocationsProvider& Provider = *AllocationsProvider;
 		TraceServices::FProviderReadScopeLock _(Provider);
+
+		TSharedPtr<Insights::FMemAllocTable> MemAllocTable = GetMemAllocTable();
+		if (MemAllocTable)
+		{
+			MemAllocTable->SetPlatformPageSize(Provider.GetPlatformPageSize());
+		}
+
 		TraceServices::IAllocationsProvider::FQueryParams Params = { Rule->GetValue(), TimeMarkers[0], TimeMarkers[1], TimeMarkers[2], TimeMarkers[3] };
 		Query = Provider.StartQuery(Params);
 	}
@@ -468,6 +487,7 @@ void SMemAllocTableTreeView::UpdateQuery(TraceServices::IAllocationsProvider::EQ
 
 					Alloc.RootHeap = Allocation->GetRootHeap();
 					Alloc.bIsHeap = Allocation->IsHeap();
+					Alloc.bIsSwap = Allocation->IsSwap();
 
 					Alloc.bIsDecline = false;
 					if (Rule->GetValue() == TraceServices::IAllocationsProvider::EQueryRule::aAfaBf)
@@ -1076,18 +1096,18 @@ void SMemAllocTableTreeView::InitAvailableViewPresets()
 	AvailableViewPresets.Add(MakeShared<FCallstackViewPreset>(true, false));
 
 	//////////////////////////////////////////////////
-	// Address (4K Page) Breakdown View
+	// Address (Platform Page) Breakdown View
 
 	class FPageViewPreset : public ITableTreeViewPreset
 	{
 	public:
 		virtual FText GetName() const override
 		{
-			return LOCTEXT("Page_PresetName", "Address (4K Page)");
+			return LOCTEXT("Page_PresetName", "Address (Platform Page)");
 		}
 		virtual FText GetToolTip() const override
 		{
-			return LOCTEXT("Page_PresetToolTip", "4K Page Breakdown View\nConfigure the tree view to show a breakdown of allocations by their address.\nIt groups allocs into 4K aligned memory pages.");
+			return LOCTEXT("Page_PresetToolTip", "Platform Page Breakdown View\nConfigure the tree view to show a breakdown of allocations by their address.\nIt groups allocs into platform page size aligned memory pages.");
 		}
 		virtual FName GetSortColumn() const override
 		{
@@ -1127,6 +1147,58 @@ void SMemAllocTableTreeView::InitAvailableViewPresets()
 	};
 	AvailableViewPresets.Add(MakeShared<FPageViewPreset>());
 
+	//////////////////////////////////////////////////
+	// Swap Breakdown View
+
+	class FSwapViewPreset : public ITableTreeViewPreset
+	{
+	public:
+		virtual FText GetName() const override
+		{
+			return LOCTEXT("Swap_PresetName", "Swap");
+		}
+		virtual FText GetToolTip() const override
+		{
+			return LOCTEXT("Swap_PresetToolTip", "Swap Usage Breakdown View\nConfigure the tree view to show a breakdown of allocations by their swap page.\nIt groups allocs into corresponding swap pages.");
+		}
+		virtual FName GetSortColumn() const override
+		{
+			return FTable::GetHierarchyColumnId();
+		}
+		virtual EColumnSortMode::Type GetSortMode() const override
+		{
+			return EColumnSortMode::Type::Ascending;
+		}
+		virtual void SetCurrentGroupings(const TArray<TSharedPtr<FTreeNodeGrouping>>& InAvailableGroupings, TArray<TSharedPtr<FTreeNodeGrouping>>& InOutCurrentGroupings) const override
+		{
+			InOutCurrentGroupings.Reset();
+
+			const TSharedPtr<FTreeNodeGrouping>* HeapGrouping = InAvailableGroupings.FindByPredicate(
+				[](TSharedPtr<FTreeNodeGrouping>& Grouping)
+				{
+					return Grouping->Is<FMemAllocGroupingBySwapPage>();
+				});
+			if (HeapGrouping)
+			{
+				InOutCurrentGroupings.Add(*HeapGrouping);
+			}
+		}
+		virtual void GetColumnConfigSet(TArray<FTableColumnConfig>& InOutConfigSet) const override
+		{
+			InOutConfigSet.Add({ FTable::GetHierarchyColumnId(),                    true, 200.0f });
+			InOutConfigSet.Add({ FMemAllocTableColumns::StartTimeColumnId,          true, 100.0f });
+			InOutConfigSet.Add({ FMemAllocTableColumns::EndTimeColumnId,            true, 100.0f });
+			InOutConfigSet.Add({ FMemAllocTableColumns::AddressColumnId,            true, 120.0f });
+			InOutConfigSet.Add({ FMemAllocTableColumns::CountColumnId,              true, 100.0f });
+			InOutConfigSet.Add({ FMemAllocTableColumns::SwapCompressedSizeColumnId, true, 100.0f });
+			InOutConfigSet.Add({ FMemAllocTableColumns::SizeInSwapColumnId,         true, 100.0f });
+			InOutConfigSet.Add({ FMemAllocTableColumns::SizeColumnId,               true, 100.0f });
+			InOutConfigSet.Add({ FMemAllocTableColumns::TagColumnId,                true, 120.0f });
+			InOutConfigSet.Add({ FMemAllocTableColumns::AllocFunctionColumnId,      true, 400.0f });
+		}
+	};
+	AvailableViewPresets.Add(MakeShared<FSwapViewPreset>());
+
 	//////////////////////////////////////////////////
 
 	SelectedViewPreset = AvailableViewPresets[0];
@@ -1370,6 +1442,7 @@ void SMemAllocTableTreeView::InternalCreateGroupings()
 	if (AllocationsProvider)
 	{
 		AvailableGroupings.Insert(MakeShared<FMemAllocGroupingByHeap>(*AllocationsProvider), Index++);
+		AvailableGroupings.Insert(MakeShared<FMemAllocGroupingBySwapPage>(*AllocationsProvider), Index++);
 	}
 }
 
diff --git a/Engine/Source/Developer/TraceInsights/Private/Insights/MemoryProfiler/Widgets/SMemAllocTableTreeView.h b/Engine/Source/Developer/TraceInsights/Private/Insights/MemoryProfiler/Widgets/SMemAllocTableTreeView.h
--- a/Engine/Source/Developer/TraceInsights/Private/Insights/MemoryProfiler/Widgets/SMemAllocTableTreeView.h
+++ b/Engine/Source/Developer/TraceInsights/Private/Insights/MemoryProfiler/Widgets/SMemAllocTableTreeView.h
@@ -67,6 +67,7 @@ public:
 		TSharedPtr<FMemoryRuleSpec> Rule;
 		double TimeMarkers[4] = { 0.0, 0.0, 0.0, 0.0 };
 		bool bIncludeHeapAllocs = false;
+		bool bIncludeSwapAllocs = false;
 	};
 
 	void SetQueryParams(const FQueryParams& InQueryParams)
@@ -77,6 +78,7 @@ public:
 		TimeMarkers[2] = InQueryParams.TimeMarkers[2];
 		TimeMarkers[3] = InQueryParams.TimeMarkers[3];
 		bIncludeHeapAllocs = InQueryParams.bIncludeHeapAllocs;
+		bIncludeSwapAllocs = InQueryParams.bIncludeSwapAllocs;
 		OnQueryInvalidated();
 	}
 
@@ -148,6 +150,7 @@ private:
 	TSharedPtr<FMemoryRuleSpec> Rule = nullptr;
 	double TimeMarkers[4];
 	bool bIncludeHeapAllocs = false;
+	bool bIncludeSwapAllocs = false;
 	TraceServices::IAllocationsProvider::FQueryHandle Query = 0;
 	FText QueryInfo;
 	FText QueryInfoTooltip;
diff --git a/Engine/Source/Developer/TraceInsights/Private/Insights/MemoryProfiler/Widgets/SMemInvestigationView.cpp b/Engine/Source/Developer/TraceInsights/Private/Insights/MemoryProfiler/Widgets/SMemInvestigationView.cpp
--- a/Engine/Source/Developer/TraceInsights/Private/Insights/MemoryProfiler/Widgets/SMemInvestigationView.cpp
+++ b/Engine/Source/Developer/TraceInsights/Private/Insights/MemoryProfiler/Widgets/SMemInvestigationView.cpp
@@ -28,6 +28,7 @@
 SMemInvestigationView::SMemInvestigationView()
 	: ProfilerWindowWeakPtr()
 	, bIncludeHeapAllocs(false)
+	, bIncludeSwapAllocs(false)
 {
 }
 
@@ -163,6 +164,27 @@ TSharedRef<SWidget> SMemInvestigationView::ConstructInvestigationWidgetArea()
 			.ToolTipText(LOCTEXT("IncludeHeapAllocsToolTipText", "Include heap allocs."))
 		]
 
+		+ SVerticalBox::Slot()
+		.AutoHeight()
+		.Padding(0.0f, 4.0f, 0.0f, 0.0f)
+		[
+			SNew(SCheckBox)
+			.IsChecked_Lambda([this]()
+			{
+				return bIncludeSwapAllocs ? ECheckBoxState::Checked : ECheckBoxState::Unchecked;
+			})
+			.OnCheckStateChanged_Lambda([this](ECheckBoxState InCheckBoxState)
+			{
+				bIncludeSwapAllocs = (InCheckBoxState == ECheckBoxState::Checked);
+			})
+			.Content()
+			[
+				SNew(STextBlock)
+				.Text(LOCTEXT("IncludeSwapAllocsText", "Include Swap Entries"))
+			]
+			.ToolTipText(LOCTEXT("IncludeSwapAllocsToolTipText", "Include swap entries."))
+		]
+
 		+ SVerticalBox::Slot()
 		.AutoHeight()
 		.Padding(0.0f, 4.0f, 0.0f, 0.0f)
@@ -570,6 +592,7 @@ FReply SMemInvestigationView::RunQuery()
 		QueryParams.TimeMarkers[2] = (RuleNumTimeMarkers > 2) ? ProfilerWindow->GetCustomTimeMarker(2)->GetTime() : 0.0;
 		QueryParams.TimeMarkers[3] = (RuleNumTimeMarkers > 3) ? ProfilerWindow->GetCustomTimeMarker(3)->GetTime() : 0.0;
 		QueryParams.bIncludeHeapAllocs = bIncludeHeapAllocs;
+		QueryParams.bIncludeSwapAllocs = bIncludeSwapAllocs;
 		MemAllocTableTreeView->SetQueryParams(QueryParams);
 	}
 
diff --git a/Engine/Source/Developer/TraceInsights/Private/Insights/MemoryProfiler/Widgets/SMemInvestigationView.h b/Engine/Source/Developer/TraceInsights/Private/Insights/MemoryProfiler/Widgets/SMemInvestigationView.h
--- a/Engine/Source/Developer/TraceInsights/Private/Insights/MemoryProfiler/Widgets/SMemInvestigationView.h
+++ b/Engine/Source/Developer/TraceInsights/Private/Insights/MemoryProfiler/Widgets/SMemInvestigationView.h
@@ -99,6 +99,7 @@ private:
 	TSharedPtr<SComboBox<TSharedPtr<Insights::FMemoryRuleSpec>>> QueryRuleComboBox;
 
 	bool bIncludeHeapAllocs;
+	bool bIncludeSwapAllocs;
 
 	TSharedPtr<SComboBox<TSharedPtr<Insights::FQueryTargetWindowSpec>>> QueryTargetComboBox;
 	
diff --git a/Engine/Source/Developer/TraceServices/Private/Analyzers/AllocationsAnalysis.cpp b/Engine/Source/Developer/TraceServices/Private/Analyzers/AllocationsAnalysis.cpp
--- a/Engine/Source/Developer/TraceServices/Private/Analyzers/AllocationsAnalysis.cpp
+++ b/Engine/Source/Developer/TraceServices/Private/Analyzers/AllocationsAnalysis.cpp
@@ -65,6 +65,7 @@ void FAllocationsAnalyzer::OnAnalysisBegin(const FOnAnalysisContext& Context)
 	Builder.RouteEvent(RouteId_HeapUnmarkAlloc,     "Memory", "HeapUnmarkAlloc");
 	Builder.RouteEvent(RouteId_MemScopeTag,         "Memory", "MemoryScope", true);
 	Builder.RouteEvent(RouteId_MemScopePtr,         "Memory", "MemoryScopePtr", true);
+	Builder.RouteEvent(RouteId_MemSwapOp,           "Memory", "MemorySwapOp");
 
 #if INSIGHTS_MEM_TRACE_METADATA_TEST
 	{
@@ -121,8 +122,14 @@ bool FAllocationsAnalyzer::OnEvent(uint16 RouteId, EStyle Style, const FOnEventC
 			const uint8 MinAlignment = EventData.GetValue<uint8>("MinAlignment");
 			SizeShift = EventData.GetValue<uint8>("SizeShift");
 
+			uint64 PlatformPageSize = EventData.GetValue<uint64>("PageSize");
+			if (PlatformPageSize < 4096) // PageSize is missing from older traces
+			{
+				PlatformPageSize = 4096;
+			}
+
 			FProviderEditScopeLock _(AllocationsProvider);
-			AllocationsProvider.EditInit(Time, MinAlignment);
+			AllocationsProvider.EditInit(Time, MinAlignment, PlatformPageSize);
 			break;
 		}
 
@@ -361,6 +368,21 @@ bool FAllocationsAnalyzer::OnEvent(uint16 RouteId, EStyle Style, const FOnEventC
 			}
 			break;
 		}
+
+		case RouteId_MemSwapOp:
+		{
+			const uint32 ThreadId = Context.ThreadInfo.GetId();
+			const double Time = GetCurrentTime();
+
+			const uint64 Address = EventData.GetValue<uint64>("Address");
+			const uint32 CallstackId = EventData.GetValue<uint32>("CallstackId");
+			const uint32 CompressedSize = EventData.GetValue<uint32>("CompressedSize");
+			const EMemoryTraceSwapOperation SwapOp = (EMemoryTraceSwapOperation)EventData.GetValue<uint8>("SwapOp");
+
+			FProviderEditScopeLock _(AllocationsProvider);
+			AllocationsProvider.EditSwapOp(ThreadId, Time, Address, SwapOp, CompressedSize, CallstackId);
+			break;
+		}
 	}
 
 	return true;
diff --git a/Engine/Source/Developer/TraceServices/Private/Analyzers/AllocationsAnalysis.h b/Engine/Source/Developer/TraceServices/Private/Analyzers/AllocationsAnalysis.h
--- a/Engine/Source/Developer/TraceServices/Private/Analyzers/AllocationsAnalysis.h
+++ b/Engine/Source/Developer/TraceServices/Private/Analyzers/AllocationsAnalysis.h
@@ -35,6 +35,7 @@ private:
 		RouteId_HeapUnmarkAlloc,
 		RouteId_MemScopeTag,
 		RouteId_MemScopePtr,
+		RouteId_MemSwapOp,
 	};
 
 public:
diff --git a/Engine/Source/Developer/TraceServices/Private/Model/AllocationItem.h b/Engine/Source/Developer/TraceServices/Private/Model/AllocationItem.h
--- a/Engine/Source/Developer/TraceServices/Private/Model/AllocationItem.h
+++ b/Engine/Source/Developer/TraceServices/Private/Model/AllocationItem.h
@@ -25,6 +25,7 @@ struct FAllocationItem
 	FORCEINLINE uint64 GetSize() const { return UnpackSize(SizeAndAlignment); }
 	FORCEINLINE uint32 GetAlignment() const { return UnpackAlignment(SizeAndAlignment); }
 	FORCEINLINE bool IsHeap() const { return EnumHasAnyFlags(Flags, EMemoryTraceHeapAllocationFlags::Heap); }
+	FORCEINLINE bool IsSwap() const { return EnumHasAnyFlags(Flags, EMemoryTraceHeapAllocationFlags::Swap); }
 
 	uint64 Address;
 	uint64 SizeAndAlignment; // (Alignment << AlignmentShift) | Size
diff --git a/Engine/Source/Developer/TraceServices/Private/Model/AllocationsProvider.cpp b/Engine/Source/Developer/TraceServices/Private/Model/AllocationsProvider.cpp
--- a/Engine/Source/Developer/TraceServices/Private/Model/AllocationsProvider.cpp
+++ b/Engine/Source/Developer/TraceServices/Private/Model/AllocationsProvider.cpp
@@ -527,6 +527,14 @@ bool IAllocationsProvider::FAllocation::IsHeap() const
 	return Inner->IsHeap();
 }
 
+////////////////////////////////////////////////////////////////////////////////////////////////////
+
+bool IAllocationsProvider::FAllocation::IsSwap() const
+{
+	const auto* Inner = (const FAllocationItem*)this;
+	return Inner->IsSwap();
+}
+
 ////////////////////////////////////////////////////////////////////////////////////////////////////
 // IAllocationsProvider::FAllocations
 ////////////////////////////////////////////////////////////////////////////////////////////////////
@@ -995,6 +1003,7 @@ FLiveAllocCollection::FLiveAllocCollection()
 FLiveAllocCollection::~FLiveAllocCollection()
 {
 	HeapAllocs.Reset();
+	SwapAllocs.Reset();
 	LongLivingAllocs.Reset();
 	ShortLivingAllocs.Reset();
 
@@ -1053,6 +1062,20 @@ FAllocationItem* FLiveAllocCollection::FindHeapRef(uint64 Address) const
 
 ////////////////////////////////////////////////////////////////////////////////////////////////////
 
+FAllocationItem* FLiveAllocCollection::FindSwapRef(uint64 Address) const
+{
+	FAllocationItem* FoundSwapAlloc = SwapAllocs.FindRef(Address);
+	if (FoundSwapAlloc)
+	{
+		INSIGHTS_SLOW_CHECK(FoundSwapAlloc->Address == Address);
+		return FoundSwapAlloc;
+	}
+
+	return nullptr;
+}
+
+////////////////////////////////////////////////////////////////////////////////////////////////////
+
 FAllocationItem* FLiveAllocCollection::FindByAddressRange(uint64 Address) const
 {
 #if INSIGHTS_USE_LAST_ALLOC
@@ -1100,6 +1123,7 @@ FAllocationItem* FLiveAllocCollection::FindHeapByAddressRange(uint64 Address) co
 void FLiveAllocCollection::Enumerate(TFunctionRef<void(const FAllocationItem& Alloc)> Callback) const
 {
 	HeapAllocs.Enumerate(Callback);
+	SwapAllocs.Enumerate(Callback);
 
 #if INSIGHTS_USE_LAST_ALLOC
 	if (LastAlloc)
@@ -1120,6 +1144,7 @@ void FLiveAllocCollection::Enumerate(TFunctionRef<void(const FAllocationItem& Al
 void FLiveAllocCollection::Enumerate(uint64 StartAddress, uint64 EndAddress, TFunctionRef<void(const FAllocationItem& Alloc)> Callback) const
 {
 	HeapAllocs.Enumerate(StartAddress, EndAddress, Callback);
+	SwapAllocs.Enumerate(StartAddress, EndAddress, Callback);
 
 #if INSIGHTS_USE_LAST_ALLOC
 	if (LastAlloc && LastAlloc->Address >= StartAddress && LastAlloc->Address < EndAddress)
@@ -1261,6 +1286,46 @@ FAllocationItem* FLiveAllocCollection::RemoveHeap(uint64 Address)
 	return nullptr;
 }
 
+////////////////////////////////////////////////////////////////////////////////////////////////////
+
+FAllocationItem* FLiveAllocCollection::AddNewSwap(uint64 Address)
+{
+	FAllocationItem* NewAlloc = new FAllocationItem();
+	NewAlloc->Address = Address;
+
+	AddSwap(NewAlloc);
+	return NewAlloc;
+}
+
+////////////////////////////////////////////////////////////////////////////////////////////////////
+
+void FLiveAllocCollection::AddSwap(FAllocationItem* SwapAlloc)
+{
+	++TotalAllocCount;
+	if (TotalAllocCount > MaxAllocCount)
+	{
+		MaxAllocCount = TotalAllocCount;
+	}
+
+	SwapAllocs.Add(SwapAlloc);
+}
+
+////////////////////////////////////////////////////////////////////////////////////////////////////
+
+FAllocationItem* FLiveAllocCollection::RemoveSwap(uint64 Address)
+{
+	FAllocationItem* RemovedSwapAlloc = SwapAllocs.Remove(Address);
+	if (RemovedSwapAlloc)
+	{
+		INSIGHTS_SLOW_CHECK(RemovedSwapAlloc->Address == Address);
+		INSIGHTS_SLOW_CHECK(TotalAllocCount > 0);
+		--TotalAllocCount;
+		return RemovedSwapAlloc;
+	}
+
+	return nullptr;
+}
+
 ////////////////////////////////////////////////////////////////////////////////////////////////////
 // FAllocationsProvider
 ////////////////////////////////////////////////////////////////////////////////////////////////////
@@ -1274,8 +1339,15 @@ FAllocationsProvider::FAllocationsProvider(IAnalysisSession& InSession, FMetadat
 	, MaxTotalAllocatedMemoryTimeline(Session.GetLinearAllocator(), 1024)
 	, MinLiveAllocationsTimeline(Session.GetLinearAllocator(), 1024)
 	, MaxLiveAllocationsTimeline(Session.GetLinearAllocator(), 1024)
+	, MinTotalSwapMemoryTimeline(Session.GetLinearAllocator(), 1024)
+	, MaxTotalSwapMemoryTimeline(Session.GetLinearAllocator(), 1024)
+	, MinTotalCompressedSwapMemoryTimeline(Session.GetLinearAllocator(), 1024)
+	, MaxTotalCompressedSwapMemoryTimeline(Session.GetLinearAllocator(), 1024)
 	, AllocEventsTimeline(Session.GetLinearAllocator(), 1024)
 	, FreeEventsTimeline(Session.GetLinearAllocator(), 1024)
+	, PageInEventsTimeline(Session.GetLinearAllocator(), 1024)
+	, PageOutEventsTimeline(Session.GetLinearAllocator(), 1024)
+	, SwapFreeEventsTimeline(Session.GetLinearAllocator(), 1024)
 {
 	// Initial number of heap spec locations. Actual number of heap specs is unlimited (uint32).
 	HeapSpecs.AddDefaulted(FMath::Max(1024u, MaxRootHeaps));
@@ -1322,7 +1394,7 @@ FAllocationsProvider::~FAllocationsProvider()
 
 ////////////////////////////////////////////////////////////////////////////////////////////////////
 
-void FAllocationsProvider::EditInit(double InTime, uint8 InMinAlignment)
+void FAllocationsProvider::EditInit(double InTime, uint8 InMinAlignment, uint64 InPlatformPageSize)
 {
 	EditAccessCheck();
 
@@ -1335,6 +1407,7 @@ void FAllocationsProvider::EditInit(double InTime, uint8 InMinAlignment)
 	INSIGHTS_API_LOGF(0u, InTime, 0ull, TEXT("Init : MinAlignment=%u)"), uint32(InMinAlignment));
 
 	InitTime = InTime;
+	PlatformPageSize = InPlatformPageSize;
 	MinAlignment = InMinAlignment;
 
 	// Create system root heap structures for backwards compatibility (before heap description events)
@@ -1621,6 +1694,153 @@ void FAllocationsProvider::EditFree(uint32 ThreadId, double Time, uint32 Callsta
 
 ////////////////////////////////////////////////////////////////////////////////////////////////////
 
+void FAllocationsProvider::EditSwapOp(uint32 ThreadId, double Time, uint64 UnmaskedPageAddress, EMemoryTraceSwapOperation SwapOp, uint64 CompressedPageSize, uint32 CallstackId)
+{
+	EditAccessCheck();
+
+	if (!bInitialized)
+	{
+		return;
+	}
+
+	HeapId RootHeapId = EMemoryTraceRootHeap::SystemMemory;
+
+	if (!IsValidRootHeap(RootHeapId))
+	{
+		return;
+	}
+
+	FRootHeap& RootHeap = *RootHeaps[RootHeapId];
+
+	const uint64 MaskedPagedAddress = UnmaskedPageAddress & (~(PlatformPageSize - 1)); 
+
+	// emit fake page in if there was previous page out with the same address
+	if (SwapOp == EMemoryTraceSwapOperation::PageOut && RootHeap.LiveAllocs->FindSwapRef(MaskedPagedAddress) != nullptr)
+	{
+		++SwapErrors;
+		if (SwapErrors <= MaxLogMessagesPerErrorType)
+		{
+			UE_LOG(LogTraceServices, Error, TEXT("[MemAlloc] EditSwapOp: Trying to swap out page 0x%llX that is currently already in swap (Time=%f, CallstackId=%u). A fake swap page in will be created."), MaskedPagedAddress, Time, CallstackId);
+		}
+
+		EditSwapOp(ThreadId, Time, UnmaskedPageAddress, EMemoryTraceSwapOperation::PageIn, CompressedPageSize, CallstackId);
+	}
+
+	RootHeap.SbTree->SetTimeForEvent(RootHeap.EventIndex, Time);
+
+	AdvanceTimelines(Time);
+
+	switch(SwapOp)
+	{
+	case EMemoryTraceSwapOperation::PageOut:
+	{
+		++SamplePageOutEvents;
+		break;
+	}
+	case EMemoryTraceSwapOperation::PageIn:
+	{
+		++SamplePageInEvents;
+		break;
+	}
+	case EMemoryTraceSwapOperation::FreeInSwap:
+	{
+		++SampleSwapFreeEvents;
+		break;
+	}
+	default:
+		break;
+	}
+
+	if (MaskedPagedAddress != 0)
+	{
+		switch(SwapOp)
+		{
+		case EMemoryTraceSwapOperation::PageOut:
+		{
+			TotalSwapMemory += PlatformPageSize;
+			TotalCompressedSwapMemory += CompressedPageSize;
+			SampleMaxSwapMemory = FMath::Max(SampleMaxSwapMemory, TotalSwapMemory);
+			SampleMaxCompressedSwapMemory = FMath::Max(SampleMaxCompressedSwapMemory, TotalCompressedSwapMemory);
+
+			FAllocationItem* AllocationPtr = RootHeap.LiveAllocs->AddNewSwap(MaskedPagedAddress);
+			FAllocationItem& Allocation = *AllocationPtr;
+
+			INSIGHTS_SLOW_CHECK(Allocation.Address == Address);
+			Allocation.SizeAndAlignment = FAllocationItem::PackSizeAndAlignment(CompressedPageSize, 0);
+			Allocation.StartEventIndex = RootHeap.EventIndex;
+			Allocation.EndEventIndex = (uint32)-1;
+			Allocation.StartTime = Time;
+			Allocation.EndTime = std::numeric_limits<double>::infinity();
+			Allocation.AllocThreadId = static_cast<uint16>(ThreadId);
+			check(uint32(AllocationPtr->AllocThreadId) == ThreadId);
+			Allocation.FreeThreadId = 0;
+			Allocation.AllocCallstackId = CallstackId;
+			Allocation.FreeCallstackId = 0; // no callstack yet
+			Allocation.MetadataId = MetadataProvider.InvalidMetadataId;
+			Allocation.Tag = 0;
+			Allocation.RootHeap = static_cast<uint8>(RootHeapId);
+			Allocation.Flags = EMemoryTraceHeapAllocationFlags::Swap;
+			break;
+		}
+		case EMemoryTraceSwapOperation::PageIn:
+		case EMemoryTraceSwapOperation::FreeInSwap:
+		{
+			if (TotalSwapMemory >= PlatformPageSize) // in case if we have imbalanced events
+			{
+				TotalSwapMemory -= PlatformPageSize;
+			}
+			SampleMinSwapMemory = FMath::Min(SampleMinSwapMemory, TotalSwapMemory);
+
+			FAllocationItem* AllocationPtr = RootHeap.LiveAllocs->RemoveSwap(MaskedPagedAddress);
+		
+			if (!AllocationPtr)
+			{
+				++SwapErrors;
+				if (SwapErrors <= MaxLogMessagesPerErrorType)
+				{
+					UE_LOG(LogTraceServices, Error, TEXT("[MemAlloc] EditSwapOp: Trying to page in swap page 0x%llX that is currently not in swap (Time=%f, CallstackId=%u). Ignoring page in. Swap total compressed memory value will be invalid."), MaskedPagedAddress, Time, CallstackId);
+				}
+			}
+
+			if (AllocationPtr)
+			{
+				CompressedPageSize = AllocationPtr->GetSize();
+				if (TotalCompressedSwapMemory >= CompressedPageSize)
+				{
+					TotalCompressedSwapMemory -= CompressedPageSize;
+				}
+				SampleMinCompressedSwapMemory = FMath::Min(SampleMinCompressedSwapMemory, TotalCompressedSwapMemory);
+
+				check(RootHeap.EventIndex > AllocationPtr->StartEventIndex);
+				AllocationPtr->EndEventIndex = RootHeap.EventIndex;
+				AllocationPtr->EndTime = Time;
+				AllocationPtr->FreeThreadId = static_cast<uint16>(ThreadId);
+				check(uint32(AllocationPtr->FreeThreadId) == ThreadId);
+				AllocationPtr->FreeCallstackId = CallstackId;
+
+				RootHeap.SbTree->AddAlloc(AllocationPtr); // SbTree takes ownership of AllocationPtr
+			}
+
+			break;
+		}
+		default:
+			break;
+		}
+	}
+
+	if (RootHeap.EventIndex != ~0u)
+	{
+		RootHeap.EventIndex++;
+	}
+	else
+	{
+		UE_LOG(LogTraceServices, Error, TEXT("[MemAlloc] Too many events!"));
+		bInitialized = false; // ignore further events
+	}
+}
+
+////////////////////////////////////////////////////////////////////////////////////////////////////
+
 void FAllocationsProvider::EditHeapSpec(HeapId Id, HeapId ParentId, const FStringView& Name, EMemoryTraceHeapFlags Flags)
 {
 	EditAccessCheck();
@@ -2126,8 +2346,15 @@ void FAllocationsProvider::AdvanceTimelines(double Time)
 		MaxTotalAllocatedMemoryTimeline.EmplaceBack(SampleMaxTotalAllocatedMemory);
 		MinLiveAllocationsTimeline.EmplaceBack(SampleMinLiveAllocations);
 		MaxLiveAllocationsTimeline.EmplaceBack(SampleMaxLiveAllocations);
+		MinTotalSwapMemoryTimeline.EmplaceBack(SampleMinSwapMemory);
+		MaxTotalSwapMemoryTimeline.EmplaceBack(SampleMaxSwapMemory);
+		MinTotalCompressedSwapMemoryTimeline.EmplaceBack(SampleMinCompressedSwapMemory);
+		MaxTotalCompressedSwapMemoryTimeline.EmplaceBack(SampleMaxCompressedSwapMemory);
 		AllocEventsTimeline.EmplaceBack(SampleAllocEvents);
 		FreeEventsTimeline.EmplaceBack(SampleFreeEvents);
+		PageInEventsTimeline.EmplaceBack(SamplePageInEvents);
+		PageOutEventsTimeline.EmplaceBack(SamplePageOutEvents);
+		SwapFreeEventsTimeline.EmplaceBack(SampleSwapFreeEvents);
 
 		// Start a new sample.
 		SampleStartTimestamp = Time;
@@ -2135,8 +2362,15 @@ void FAllocationsProvider::AdvanceTimelines(double Time)
 		SampleMaxTotalAllocatedMemory = TotalAllocatedMemory;
 		SampleMinLiveAllocations = TotalLiveAllocations;
 		SampleMaxLiveAllocations = TotalLiveAllocations;
+		SampleMinSwapMemory = TotalSwapMemory;
+		SampleMaxSwapMemory = TotalSwapMemory;
+		SampleMinCompressedSwapMemory = TotalCompressedSwapMemory;
+		SampleMaxCompressedSwapMemory = TotalCompressedSwapMemory;
 		SampleAllocEvents = 0;
 		SampleFreeEvents = 0;
+		SamplePageInEvents = 0;
+		SamplePageOutEvents = 0;
+		SampleSwapFreeEvents = 0;
 
 		// If the previous sample is well distanced in time...
 		if (Time - SampleEndTimestamp > DefaultTimelineSampleGranularity)
@@ -2147,8 +2381,15 @@ void FAllocationsProvider::AdvanceTimelines(double Time)
 			MaxTotalAllocatedMemoryTimeline.EmplaceBack(TotalAllocatedMemory);
 			MinLiveAllocationsTimeline.EmplaceBack(TotalLiveAllocations);
 			MaxLiveAllocationsTimeline.EmplaceBack(TotalLiveAllocations);
+			MinTotalSwapMemoryTimeline.EmplaceBack(TotalSwapMemory);
+			MaxTotalSwapMemoryTimeline.EmplaceBack(TotalSwapMemory);
+			MinTotalCompressedSwapMemoryTimeline.EmplaceBack(TotalCompressedSwapMemory);
+			MaxTotalCompressedSwapMemoryTimeline.EmplaceBack(TotalCompressedSwapMemory);
 			AllocEventsTimeline.EmplaceBack(0);
 			FreeEventsTimeline.EmplaceBack(0);
+			PageInEventsTimeline.EmplaceBack(0);
+			PageOutEventsTimeline.EmplaceBack(0);
+			SwapFreeEventsTimeline.EmplaceBack(0);
 		}
 	}
 
@@ -2272,8 +2513,11 @@ void FAllocationsProvider::EditOnAnalysisCompleted(double Time)
 
 		// Update stats for the last timeline sample (reset to zero).
 		TotalAllocatedMemory = 0;
+		TotalSwapMemory = 0;
 		SampleMinTotalAllocatedMemory = 0;
 		SampleMinLiveAllocations = 0;
+		SampleMinSwapMemory = 0;
+		SampleMinCompressedSwapMemory = 0;
 		SampleFreeEvents += LiveAllocsTotalCount;
 	}
 #endif
@@ -2296,11 +2540,11 @@ void FAllocationsProvider::EditOnAnalysisCompleted(double Time)
 
 	//TODO: shrink live allocs buffers
 
-	if (AllocWarnings > 0 || FreeWarnings > 0 || HeapWarnings > 0 || MiscWarnings > 0)
+	if (AllocWarnings > 0 || FreeWarnings > 0 || HeapWarnings > 0 || SwapWarnings > 0 || MiscWarnings > 0)
 	{
-		UE_LOG(LogTraceServices, Warning, TEXT("[MemAlloc] %llu warnings (%llu ALLOC + %llu FREE + %llu HEAP + %llu other)"),
-			AllocWarnings + FreeWarnings + HeapWarnings + MiscWarnings,
-			AllocWarnings, FreeWarnings, HeapWarnings, MiscWarnings);
+		UE_LOG(LogTraceServices, Warning, TEXT("[MemAlloc] %llu warnings (%llu ALLOC + %llu FREE + %llu HEAP + %llu SWAP + %llu other)"),
+			AllocWarnings + FreeWarnings + HeapWarnings + SwapWarnings + MiscWarnings,
+			AllocWarnings, FreeWarnings, HeapWarnings, SwapWarnings, MiscWarnings);
 	}
 
 	if (AllocErrors > 0)
@@ -2315,6 +2559,10 @@ void FAllocationsProvider::EditOnAnalysisCompleted(double Time)
 	{
 		UE_LOG(LogTraceServices, Error, TEXT("[MemAlloc] HEAP event errors: %llu"), HeapErrors);
 	}
+	if (SwapErrors > 0)
+	{
+		UE_LOG(LogTraceServices, Error, TEXT("[MemAlloc] SWAP event errors: %llu"), SwapErrors);
+	}
 	if (MiscErrors > 0)
 	{
 		UE_LOG(LogTraceServices, Error, TEXT("[MemAlloc] Other errors: %llu"), MiscErrors);
@@ -2585,6 +2833,150 @@ void FAllocationsProvider::EnumerateMaxLiveAllocationsTimeline(int32 StartIndex,
 
 ////////////////////////////////////////////////////////////////////////////////////////////////////
 
+void FAllocationsProvider::EnumerateMinTotalSwapMemoryTimeline(int32 StartIndex, int32 EndIndex, TFunctionRef<void(double Time, double Duration, uint64 Value)> Callback) const
+{
+	ReadAccessCheck();
+
+	const int32 NumPoints = static_cast<int32>(Timeline.Num());
+	StartIndex = FMath::Max(StartIndex, 0);
+	EndIndex = FMath::Min(EndIndex + 1, NumPoints); // make it exclusive
+	if (StartIndex < EndIndex)
+	{
+		auto TimeIt = Timeline.GetIteratorFromItem(StartIndex);
+		auto ValueIt = MinTotalSwapMemoryTimeline.GetIteratorFromItem(StartIndex);
+		double PrevTime = *TimeIt;
+		uint64 PrevValue = *ValueIt;
+		++TimeIt;
+		++ValueIt;
+		for (int32 Index = StartIndex + 1; Index < EndIndex; ++Index, ++TimeIt, ++ValueIt)
+		{
+			const double Time = *TimeIt;
+			Callback(PrevTime, Time - PrevTime, PrevValue);
+			PrevTime = *TimeIt;
+			PrevValue = *ValueIt;
+		}
+		if (EndIndex < NumPoints)
+		{
+			const double Time = *TimeIt;
+			Callback(PrevTime, Time - PrevTime, PrevValue);
+		}
+		else
+		{
+			Callback(PrevTime, std::numeric_limits<double>::infinity(), PrevValue);
+		}
+	}
+}
+
+////////////////////////////////////////////////////////////////////////////////////////////////////
+
+void FAllocationsProvider::EnumerateMaxTotalSwapMemoryTimeline(int32 StartIndex, int32 EndIndex, TFunctionRef<void(double Time, double Duration, uint64 Value)> Callback) const
+{
+	ReadAccessCheck();
+
+	const int32 NumPoints = static_cast<int32>(Timeline.Num());
+	StartIndex = FMath::Max(StartIndex, 0);
+	EndIndex = FMath::Min(EndIndex + 1, NumPoints); // make it exclusive
+	if (StartIndex < EndIndex)
+	{
+		auto TimeIt = Timeline.GetIteratorFromItem(StartIndex);
+		auto ValueIt = MaxTotalSwapMemoryTimeline.GetIteratorFromItem(StartIndex);
+		double PrevTime = *TimeIt;
+		uint64 PrevValue = *ValueIt;
+		++TimeIt;
+		++ValueIt;
+		for (int32 Index = StartIndex + 1; Index < EndIndex; ++Index, ++TimeIt, ++ValueIt)
+		{
+			const double Time = *TimeIt;
+			Callback(PrevTime, Time - PrevTime, PrevValue);
+			PrevTime = *TimeIt;
+			PrevValue = *ValueIt;
+		}
+		if (EndIndex < NumPoints)
+		{
+			const double Time = *TimeIt;
+			Callback(PrevTime, Time - PrevTime, PrevValue);
+		}
+		else
+		{
+			Callback(PrevTime, std::numeric_limits<double>::infinity(), PrevValue);
+		}
+	}
+}
+
+////////////////////////////////////////////////////////////////////////////////////////////////////
+
+void FAllocationsProvider::EnumerateMinTotalCompressedSwapMemoryTimeline(int32 StartIndex, int32 EndIndex, TFunctionRef<void(double Time, double Duration, uint64 Value)> Callback) const
+{
+	ReadAccessCheck();
+
+	const int32 NumPoints = static_cast<int32>(Timeline.Num());
+	StartIndex = FMath::Max(StartIndex, 0);
+	EndIndex = FMath::Min(EndIndex + 1, NumPoints); // make it exclusive
+	if (StartIndex < EndIndex)
+	{
+		auto TimeIt = Timeline.GetIteratorFromItem(StartIndex);
+		auto ValueIt = MinTotalCompressedSwapMemoryTimeline.GetIteratorFromItem(StartIndex);
+		double PrevTime = *TimeIt;
+		uint64 PrevValue = *ValueIt;
+		++TimeIt;
+		++ValueIt;
+		for (int32 Index = StartIndex + 1; Index < EndIndex; ++Index, ++TimeIt, ++ValueIt)
+		{
+			const double Time = *TimeIt;
+			Callback(PrevTime, Time - PrevTime, PrevValue);
+			PrevTime = *TimeIt;
+			PrevValue = *ValueIt;
+		}
+		if (EndIndex < NumPoints)
+		{
+			const double Time = *TimeIt;
+			Callback(PrevTime, Time - PrevTime, PrevValue);
+		}
+		else
+		{
+			Callback(PrevTime, std::numeric_limits<double>::infinity(), PrevValue);
+		}
+	}
+}
+
+////////////////////////////////////////////////////////////////////////////////////////////////////
+
+void FAllocationsProvider::EnumerateMaxTotalCompressedSwapMemoryTimeline(int32 StartIndex, int32 EndIndex, TFunctionRef<void(double Time, double Duration, uint64 Value)> Callback) const
+{
+	ReadAccessCheck();
+
+	const int32 NumPoints = static_cast<int32>(Timeline.Num());
+	StartIndex = FMath::Max(StartIndex, 0);
+	EndIndex = FMath::Min(EndIndex + 1, NumPoints); // make it exclusive
+	if (StartIndex < EndIndex)
+	{
+		auto TimeIt = Timeline.GetIteratorFromItem(StartIndex);
+		auto ValueIt = MaxTotalCompressedSwapMemoryTimeline.GetIteratorFromItem(StartIndex);
+		double PrevTime = *TimeIt;
+		uint64 PrevValue = *ValueIt;
+		++TimeIt;
+		++ValueIt;
+		for (int32 Index = StartIndex + 1; Index < EndIndex; ++Index, ++TimeIt, ++ValueIt)
+		{
+			const double Time = *TimeIt;
+			Callback(PrevTime, Time - PrevTime, PrevValue);
+			PrevTime = *TimeIt;
+			PrevValue = *ValueIt;
+		}
+		if (EndIndex < NumPoints)
+		{
+			const double Time = *TimeIt;
+			Callback(PrevTime, Time - PrevTime, PrevValue);
+		}
+		else
+		{
+			Callback(PrevTime, std::numeric_limits<double>::infinity(), PrevValue);
+		}
+	}
+}
+
+////////////////////////////////////////////////////////////////////////////////////////////////////
+
 void FAllocationsProvider::EnumerateAllocEventsTimeline(int32 StartIndex, int32 EndIndex, TFunctionRef<void(double Time, double Duration, uint32 Value)> Callback) const
 {
 	ReadAccessCheck();
@@ -2657,6 +3049,114 @@ void FAllocationsProvider::EnumerateFreeEventsTimeline(int32 StartIndex, int32 E
 
 ////////////////////////////////////////////////////////////////////////////////////////////////////
 
+void FAllocationsProvider::EnumeratePageInEventsTimeline(int32 StartIndex, int32 EndIndex, TFunctionRef<void(double Time, double Duration, uint32 Value)> Callback) const
+{
+	ReadAccessCheck();
+
+	const int32 NumPoints = static_cast<int32>(Timeline.Num());
+	StartIndex = FMath::Max(StartIndex, 0);
+	EndIndex = FMath::Min(EndIndex + 1, NumPoints); // make it exclusive
+	if (StartIndex < EndIndex)
+	{
+		auto TimeIt = Timeline.GetIteratorFromItem(StartIndex);
+		auto ValueIt = PageInEventsTimeline.GetIteratorFromItem(StartIndex);
+		double PrevTime = *TimeIt;
+		uint32 PrevValue = *ValueIt;
+		++TimeIt;
+		++ValueIt;
+		for (int32 Index = StartIndex + 1; Index < EndIndex; ++Index, ++TimeIt, ++ValueIt)
+		{
+			const double Time = *TimeIt;
+			Callback(PrevTime, Time - PrevTime, PrevValue);
+			PrevTime = *TimeIt;
+			PrevValue = *ValueIt;
+		}
+		if (EndIndex < NumPoints)
+		{
+			const double Time = *TimeIt;
+			Callback(PrevTime, Time - PrevTime, PrevValue);
+		}
+		else
+		{
+			Callback(PrevTime, std::numeric_limits<double>::infinity(), PrevValue);
+		}
+	}
+}
+
+////////////////////////////////////////////////////////////////////////////////////////////////////
+
+void FAllocationsProvider::EnumeratePageOutEventsTimeline(int32 StartIndex, int32 EndIndex, TFunctionRef<void(double Time, double Duration, uint32 Value)> Callback) const
+{
+	ReadAccessCheck();
+
+	const int32 NumPoints = static_cast<int32>(Timeline.Num());
+	StartIndex = FMath::Max(StartIndex, 0);
+	EndIndex = FMath::Min(EndIndex + 1, NumPoints); // make it exclusive
+	if (StartIndex < EndIndex)
+	{
+		auto TimeIt = Timeline.GetIteratorFromItem(StartIndex);
+		auto ValueIt = PageOutEventsTimeline.GetIteratorFromItem(StartIndex);
+		double PrevTime = *TimeIt;
+		uint32 PrevValue = *ValueIt;
+		++TimeIt;
+		++ValueIt;
+		for (int32 Index = StartIndex + 1; Index < EndIndex; ++Index, ++TimeIt, ++ValueIt)
+		{
+			const double Time = *TimeIt;
+			Callback(PrevTime, Time - PrevTime, PrevValue);
+			PrevTime = *TimeIt;
+			PrevValue = *ValueIt;
+		}
+		if (EndIndex < NumPoints)
+		{
+			const double Time = *TimeIt;
+			Callback(PrevTime, Time - PrevTime, PrevValue);
+		}
+		else
+		{
+			Callback(PrevTime, std::numeric_limits<double>::infinity(), PrevValue);
+		}
+	}
+}
+
+////////////////////////////////////////////////////////////////////////////////////////////////////
+
+void FAllocationsProvider::EnumerateSwapFreeEventsTimeline(int32 StartIndex, int32 EndIndex, TFunctionRef<void(double Time, double Duration, uint32 Value)> Callback) const
+{
+	ReadAccessCheck();
+
+	const int32 NumPoints = static_cast<int32>(Timeline.Num());
+	StartIndex = FMath::Max(StartIndex, 0);
+	EndIndex = FMath::Min(EndIndex + 1, NumPoints); // make it exclusive
+	if (StartIndex < EndIndex)
+	{
+		auto TimeIt = Timeline.GetIteratorFromItem(StartIndex);
+		auto ValueIt = SwapFreeEventsTimeline.GetIteratorFromItem(StartIndex);
+		double PrevTime = *TimeIt;
+		uint32 PrevValue = *ValueIt;
+		++TimeIt;
+		++ValueIt;
+		for (int32 Index = StartIndex + 1; Index < EndIndex; ++Index, ++TimeIt, ++ValueIt)
+		{
+			const double Time = *TimeIt;
+			Callback(PrevTime, Time - PrevTime, PrevValue);
+			PrevTime = *TimeIt;
+			PrevValue = *ValueIt;
+		}
+		if (EndIndex < NumPoints)
+		{
+			const double Time = *TimeIt;
+			Callback(PrevTime, Time - PrevTime, PrevValue);
+		}
+		else
+		{
+			Callback(PrevTime, std::numeric_limits<double>::infinity(), PrevValue);
+		}
+	}
+}
+
+////////////////////////////////////////////////////////////////////////////////////////////////////
+
 void FAllocationsProvider::EnumerateTags(TFunctionRef<void(const TCHAR*, const TCHAR*, TagIdType, TagIdType)> Callback) const
 {
 	ReadAccessCheck();
@@ -2706,6 +3206,13 @@ const IAllocationsProvider::FQueryStatus FAllocationsProvider::PollQuery(FQueryH
 
 ////////////////////////////////////////////////////////////////////////////////////////////////////
 
+uint64 FAllocationsProvider::GetPlatformPageSize() const
+{
+	return bInitialized ? PlatformPageSize : 4096;
+}
+
+////////////////////////////////////////////////////////////////////////////////////////////////////
+
 void FAllocationsProvider::EnumerateLiveAllocs(TFunctionRef<void(const FAllocationItem& Alloc)> Callback) const
 {
 	ReadAccessCheck();
diff --git a/Engine/Source/Developer/TraceServices/Private/Model/AllocationsProvider.h b/Engine/Source/Developer/TraceServices/Private/Model/AllocationsProvider.h
--- a/Engine/Source/Developer/TraceServices/Private/Model/AllocationsProvider.h
+++ b/Engine/Source/Developer/TraceServices/Private/Model/AllocationsProvider.h
@@ -206,6 +206,10 @@ public:
 	// Returns the found allocation or nullptr if not found.
 	FORCEINLINE FAllocationItem* FindHeapRef(uint64 Address) const;
 
+	// Finds the swap allocation with specified address.
+	// Returns the found allocation or nullptr if not found.
+	FORCEINLINE FAllocationItem* FindSwapRef(uint64 Address) const;
+
 	// Finds an allocation containing the address.
 	// Returns the found allocation or nullptr if not found.
 	FORCEINLINE FAllocationItem* FindByAddressRange(uint64 Address) const;
@@ -248,11 +252,26 @@ public:
 	// Returns the removed heap allocation or nullptr if not found.
 	FORCEINLINE FAllocationItem* RemoveHeap(uint64 Address);
 
+	// Adds a new swap allocation with specified address (a memory page that was moved to swap).
+	// The collection keeps ownership of FAllocationItem* until RemoveSwap is called.
+	// Returns the new added allocation.
+	FORCEINLINE FAllocationItem* AddNewSwap(uint64 Address);
+
+	// Adds a new swap allocation.
+	// The collection keeps ownership of FAllocationItem* until RemoveSwap is called.
+	FORCEINLINE void AddSwap(FAllocationItem* SwapAlloc);
+
+	// Removes a swap allocation with specified address.
+	// The caller takes ownership of FAllocationItem*.
+	// Returns the removed swap allocation or nullptr if not found.
+	FORCEINLINE FAllocationItem* RemoveSwap(uint64 Address);
+
 private:
 	FAllocationItem* LastAlloc = nullptr; // last allocation
 	FShortLivingAllocs ShortLivingAllocs; // short living allocations
 	FAllocMap LongLivingAllocs; // long living allocations
 	FHeapAllocs HeapAllocs; // heap allocations
+	FHeapAllocs SwapAllocs; // swap allocations (memory pages stored in swap)
 
 	uint32 TotalAllocCount = 0;
 	uint32 MaxAllocCount = 0; // debug stats
@@ -318,13 +337,22 @@ public:
 	virtual void EnumerateMaxTotalAllocatedMemoryTimeline(int32 StartIndex, int32 EndIndex, TFunctionRef<void(double Time, double Duration, uint64 Value)> Callback) const override;
 	virtual void EnumerateMinLiveAllocationsTimeline(int32 StartIndex, int32 EndIndex, TFunctionRef<void(double Time, double Duration, uint32 Value)> Callback) const override;
 	virtual void EnumerateMaxLiveAllocationsTimeline(int32 StartIndex, int32 EndIndex, TFunctionRef<void(double Time, double Duration, uint32 Value)> Callback) const override;
+	virtual void EnumerateMinTotalSwapMemoryTimeline(int32 StartIndex, int32 EndIndex, TFunctionRef<void(double Time, double Duration, uint64 Value)> Callback) const override;
+	virtual void EnumerateMaxTotalSwapMemoryTimeline(int32 StartIndex, int32 EndIndex, TFunctionRef<void(double Time, double Duration, uint64 Value)> Callback) const override;
+	virtual void EnumerateMinTotalCompressedSwapMemoryTimeline(int32 StartIndex, int32 EndIndex, TFunctionRef<void(double Time, double Duration, uint64 Value)> Callback) const override;
+	virtual void EnumerateMaxTotalCompressedSwapMemoryTimeline(int32 StartIndex, int32 EndIndex, TFunctionRef<void(double Time, double Duration, uint64 Value)> Callback) const override;
 	virtual void EnumerateAllocEventsTimeline(int32 StartIndex, int32 EndIndex, TFunctionRef<void(double Time, double Duration, uint32 Value)> Callback) const override;
 	virtual void EnumerateFreeEventsTimeline(int32 StartIndex, int32 EndIndex, TFunctionRef<void(double Time, double Duration, uint32 Value)> Callback) const override;
+	virtual void EnumeratePageInEventsTimeline(int32 StartIndex, int32 EndIndex, TFunctionRef<void(double Time, double Duration, uint32 Value)> Callback) const override;
+	virtual void EnumeratePageOutEventsTimeline(int32 StartIndex, int32 EndIndex, TFunctionRef<void(double Time, double Duration, uint32 Value)> Callback) const override;
+	virtual void EnumerateSwapFreeEventsTimeline(int32 StartIndex, int32 EndIndex, TFunctionRef<void(double Time, double Duration, uint32 Value)> Callback) const override;
 
 	virtual FQueryHandle StartQuery(const FQueryParams& Params) const override;
 	virtual void CancelQuery(FQueryHandle Query) const override;
 	virtual const FQueryStatus PollQuery(FQueryHandle Query) const override;
 
+	virtual uint64 GetPlatformPageSize() const override;
+
 	const FSbTree* GetSbTreeUnchecked(HeapId Heap) const { ReadAccessCheck(); return RootHeaps[Heap]->SbTree; }
 
 	void EnumerateLiveAllocs(TFunctionRef<void(const FAllocationItem& Alloc)> Callback) const;
@@ -339,10 +367,11 @@ public:
 	virtual void EndEdit() const override         { Lock.EndWrite(GAllocationsProviderLockState); }
 	virtual void EditAccessCheck() const override { Lock.WriteAccessCheck(GAllocationsProviderLockState); }
 
-	void EditInit(double Time, uint8 MinAlignment);
+	void EditInit(double Time, uint8 MinAlignment, uint64 PlatformPageSize);
 
 	void EditAlloc(uint32 ThreadId, double Time, uint32 CallstackId, uint64 Address, uint64 Size, uint32 Alignment, HeapId RootHeap);
 	void EditFree(uint32 ThreadId, double Time, uint32 CallstackId, uint64 Address, HeapId RootHeap);
+	void EditSwapOp(uint32 ThreadId, double Time, uint64 UnmaskedPageAddress, EMemoryTraceSwapOperation SwapOp, uint64 CompressedPageSize, uint32 CallstackId);
 
 	void EditHeapSpec(HeapId Id, HeapId ParentId, const FStringView& Name, EMemoryTraceHeapFlags Flags);
 	void EditMarkAllocationAsHeap(uint32 ThreadId, double Time, uint32 CallstackId, uint64 Address, HeapId Heap, EMemoryTraceHeapAllocationFlags Flags);
@@ -395,6 +424,7 @@ private:
 	FMetadataProvider& MetadataProvider;
 
 	double InitTime = 0;
+	uint64 PlatformPageSize = 0;
 	uint8 MinAlignment = 0;
 	uint8 SizeShift = 0;
 	uint8 SummarySizeShift = 0;
@@ -417,9 +447,13 @@ private:
 	uint64 AllocErrors = 0;
 	uint64 FreeWarnings = 0;
 	uint64 FreeErrors = 0;
+	uint64 SwapWarnings = 0;
+	uint64 SwapErrors = 0;
 
 	uint64 TotalAllocatedMemory = 0;
 	uint32 TotalLiveAllocations = 0;
+	uint64 TotalSwapMemory = 0;
+	uint64 TotalCompressedSwapMemory = 0;
 
 	double SampleStartTimestamp = 0.0;
 	double SampleEndTimestamp = 0.0;
@@ -427,16 +461,30 @@ private:
 	uint64 SampleMaxTotalAllocatedMemory = 0;
 	uint32 SampleMinLiveAllocations = 0;
 	uint32 SampleMaxLiveAllocations = 0;
+	uint64 SampleMinSwapMemory = 0;
+	uint64 SampleMaxSwapMemory = 0;
+	uint64 SampleMinCompressedSwapMemory = 0;
+	uint64 SampleMaxCompressedSwapMemory = 0;
 	uint32 SampleAllocEvents = 0;
 	uint32 SampleFreeEvents = 0;
+	uint32 SamplePageInEvents = 0;
+	uint32 SamplePageOutEvents = 0;
+	uint32 SampleSwapFreeEvents = 0;
 
 	TPagedArray<double> Timeline;
 	TPagedArray<uint64> MinTotalAllocatedMemoryTimeline;
 	TPagedArray<uint64> MaxTotalAllocatedMemoryTimeline;
 	TPagedArray<uint32> MinLiveAllocationsTimeline;
 	TPagedArray<uint32> MaxLiveAllocationsTimeline;
+	TPagedArray<uint64> MinTotalSwapMemoryTimeline;
+	TPagedArray<uint64> MaxTotalSwapMemoryTimeline;
+	TPagedArray<uint64> MinTotalCompressedSwapMemoryTimeline;
+	TPagedArray<uint64> MaxTotalCompressedSwapMemoryTimeline;
 	TPagedArray<uint32> AllocEventsTimeline;
 	TPagedArray<uint32> FreeEventsTimeline;
+	TPagedArray<uint32> PageInEventsTimeline;
+	TPagedArray<uint32> PageOutEventsTimeline;
+	TPagedArray<uint32> SwapFreeEventsTimeline;
 };
 
 ////////////////////////////////////////////////////////////////////////////////////////////////////
diff --git a/Engine/Source/Developer/TraceServices/Private/Model/AllocationsQuery.cpp b/Engine/Source/Developer/TraceServices/Private/Model/AllocationsQuery.cpp
--- a/Engine/Source/Developer/TraceServices/Private/Model/AllocationsQuery.cpp
+++ b/Engine/Source/Developer/TraceServices/Private/Model/AllocationsQuery.cpp
@@ -396,6 +396,36 @@ void FAllocationsQuery::QueryLiveAllocs(TArray<const FAllocationItem*>& OutAlloc
 		});
 	}
 	break;
+
+	case IAllocationsProvider::EQueryRule::AoB:
+	{
+		const double TimeA = Params.TimeA;
+		const double TimeB = Params.TimeB;
+		AllocationsProvider.EnumerateLiveAllocs([TimeA, TimeB, &OutAllocs](const FAllocationItem& Alloc)
+		{
+			if ((Alloc.IsSwap() && Alloc.StartTime >= TimeA && Alloc.StartTime <= TimeB) ||
+				(!Alloc.IsSwap() && Alloc.StartTime <= TimeB && Alloc.EndTime >= TimeA)) // include allocs that were freed between A and B
+			{
+				OutAllocs.Add(&Alloc);
+			}
+		});
+	}
+	break;
+
+	case IAllocationsProvider::EQueryRule::AiB:
+	{
+		const double TimeA = Params.TimeA;
+		const double TimeB = Params.TimeB;
+		AllocationsProvider.EnumerateLiveAllocs([TimeA, TimeB, &OutAllocs](const FAllocationItem& Alloc)
+		{
+			if ((Alloc.IsSwap() && Alloc.EndTime >= TimeA && Alloc.EndTime <= TimeB) ||
+				(!Alloc.IsSwap() && Alloc.StartTime <= TimeB && Alloc.EndTime >= TimeA)) // include allocs that were freed between A and B
+			{
+				OutAllocs.Add(&Alloc);
+			}
+		});
+	}
+	break;
 	}
 }
 
diff --git a/Engine/Source/Developer/TraceServices/Private/Model/SbTree.cpp b/Engine/Source/Developer/TraceServices/Private/Model/SbTree.cpp
--- a/Engine/Source/Developer/TraceServices/Private/Model/SbTree.cpp
+++ b/Engine/Source/Developer/TraceServices/Private/Model/SbTree.cpp
@@ -300,6 +300,37 @@ void FSbTreeCell::Query(TArray<const FAllocationItem*>& OutAllocs, const IAlloca
 		}
 	}
 	break;
+
+	case IAllocationsProvider::EQueryRule::AoB: // paged out allocs
+	{
+		const double TimeA = Params.TimeA;
+		const double TimeB = Params.TimeB;
+		for (const FAllocationItem* Alloc : Allocs)
+		{
+			if ((Alloc->IsSwap() && Alloc->StartTime >= TimeA && Alloc->StartTime <= TimeB) ||
+				(!Alloc->IsSwap() && Alloc->StartTime <= TimeB && Alloc->EndTime >= TimeA)) // include allocs that were freed between A and B
+			{
+				OutAllocs.Add(Alloc);
+			}
+		}
+	}
+	break;
+
+	case IAllocationsProvider::EQueryRule::AiB: // paged in allocs
+	{
+		const double TimeA = Params.TimeA;
+		const double TimeB = Params.TimeB;
+		for (const FAllocationItem* Alloc : Allocs)
+		{
+			if ((Alloc->IsSwap() && Alloc->EndTime >= TimeA && Alloc->EndTime <= TimeB) ||
+				(!Alloc->IsSwap() && Alloc->StartTime <= TimeB && Alloc->EndTime >= TimeA)) // include allocs that were freed between A and B
+			{
+				OutAllocs.Add(Alloc);
+			}
+		}
+	}
+	break;
+
 	}
 }
 
@@ -624,6 +655,22 @@ void FSbTree::Query(TArray<const FSbTreeCell*>& OutCells, const IAllocationsProv
 			IterateCells(OutCells, Column1, Column2);
 		}
 		break;
+		
+		case IAllocationsProvider::EQueryRule::AoB:
+		{
+			int32 Column1;
+			GetColumnsAtTime(Params.TimeB, nullptr, &Column1); // we care about all allocs up to Time B
+			IterateCells(OutCells, 0, Column1);
+		}
+		break;
+
+		case IAllocationsProvider::EQueryRule::AiB:
+		{
+			int32 Column1;
+			GetColumnsAtTime(Params.TimeB, nullptr, &Column1); // we care about all allocs up to Time B
+			IterateCells(OutCells, 0, Column1);
+		}
+		break;
 	}
 }
 
diff --git a/Engine/Source/Developer/TraceServices/Public/TraceServices/Model/AllocationsProvider.h b/Engine/Source/Developer/TraceServices/Public/TraceServices/Model/AllocationsProvider.h
--- a/Engine/Source/Developer/TraceServices/Public/TraceServices/Model/AllocationsProvider.h
+++ b/Engine/Source/Developer/TraceServices/Public/TraceServices/Model/AllocationsProvider.h
@@ -42,6 +42,8 @@ public:
 		AaBfC,   // limited lifetime
 		aABfC,   // decline of long living allocs
 		AaBCfD,  // specific lifetime
+		AoB,     // all paged out allocs between A and B
+		AiB,     // all paged in allocs between A and B
 		//A_vs_B,  // compare A vs. B; {aAf} vs. {aBf}
 		//A_or_B,  // live at A or at B; {aAf} U {aBf}
 		//A_xor_B, // live either at A or at B; ({aAf} U {aBf}) \ {aABf}
@@ -73,6 +75,7 @@ public:
 		TagIdType GetTag() const;
 		HeapId GetRootHeap() const;
 		bool IsHeap() const;
+		bool IsSwap() const;
 	};
 
 	class TRACESERVICES_API FAllocations
@@ -151,15 +154,38 @@ public:
 	// Enumerates the Max Live Allocations timeline points in the inclusive index interval [StartIndex, EndIndex].
 	virtual void EnumerateMaxLiveAllocationsTimeline(int32 StartIndex, int32 EndIndex, TFunctionRef<void(double Time, double Duration, uint32 Value)> Callback) const = 0;
 
+	// Enumerates the Min Total Swap Memory timeline points in the inclusive index interval [StartIndex, EndIndex].
+	virtual void EnumerateMinTotalSwapMemoryTimeline(int32 StartIndex, int32 EndIndex, TFunctionRef<void(double Time, double Duration, uint64 Value)> Callback) const = 0;
+
+	// Enumerates the Max Total Swap Memory timeline points in the inclusive index interval [StartIndex, EndIndex].
+	virtual void EnumerateMaxTotalSwapMemoryTimeline(int32 StartIndex, int32 EndIndex, TFunctionRef<void(double Time, double Duration, uint64 Value)> Callback) const = 0;
+
+	// Enumerates the Min Total Compressed Swap Memory timeline points in the inclusive index interval [StartIndex, EndIndex].
+	virtual void EnumerateMinTotalCompressedSwapMemoryTimeline(int32 StartIndex, int32 EndIndex, TFunctionRef<void(double Time, double Duration, uint64 Value)> Callback) const = 0;
+
+	// Enumerates the Max Total Compressed Swap Memory timeline points in the inclusive index interval [StartIndex, EndIndex].
+	virtual void EnumerateMaxTotalCompressedSwapMemoryTimeline(int32 StartIndex, int32 EndIndex, TFunctionRef<void(double Time, double Duration, uint64 Value)> Callback) const = 0;
+
 	// Enumerates the Alloc Events timeline points in the inclusive index interval [StartIndex, EndIndex].
 	virtual void EnumerateAllocEventsTimeline(int32 StartIndex, int32 EndIndex, TFunctionRef<void(double Time, double Duration, uint32 Value)> Callback) const = 0;
 
 	// Enumerates the Free Events timeline points in the inclusive index interval [StartIndex, EndIndex].
 	virtual void EnumerateFreeEventsTimeline(int32 StartIndex, int32 EndIndex, TFunctionRef<void(double Time, double Duration, uint32 Value)> Callback) const = 0;
 
+	// Enumerates the Page In Events timeline points in the inclusive index interval [StartIndex, EndIndex].
+	virtual void EnumeratePageInEventsTimeline(int32 StartIndex, int32 EndIndex, TFunctionRef<void(double Time, double Duration, uint32 Value)> Callback) const = 0;
+
+	// Enumerates the Page Out Events timeline points in the inclusive index interval [StartIndex, EndIndex].
+	virtual void EnumeratePageOutEventsTimeline(int32 StartIndex, int32 EndIndex, TFunctionRef<void(double Time, double Duration, uint32 Value)> Callback) const = 0;
+
+	// Enumerates the Swap Free Events timeline points in the inclusive index interval [StartIndex, EndIndex].
+	virtual void EnumerateSwapFreeEventsTimeline(int32 StartIndex, int32 EndIndex, TFunctionRef<void(double Time, double Duration, uint32 Value)> Callback) const = 0;
+
 	virtual FQueryHandle StartQuery(const FQueryParams& Params) const = 0;
 	virtual void CancelQuery(FQueryHandle Query) const = 0;
 	virtual const FQueryStatus PollQuery(FQueryHandle Query) const = 0;
+
+	virtual uint64 GetPlatformPageSize() const = 0;
 };
 
 TRACESERVICES_API FName GetAllocationsProviderName();
diff --git a/Engine/Source/Runtime/Core/Private/ProfilingDebugging/MemoryTrace.cpp b/Engine/Source/Runtime/Core/Private/ProfilingDebugging/MemoryTrace.cpp
--- a/Engine/Source/Runtime/Core/Private/ProfilingDebugging/MemoryTrace.cpp
+++ b/Engine/Source/Runtime/Core/Private/ProfilingDebugging/MemoryTrace.cpp
@@ -57,11 +57,11 @@ namespace Trace {
 
 ////////////////////////////////////////////////////////////////////////////////
 UE_TRACE_EVENT_BEGIN(Memory, Init, NoSync|Important)
+	UE_TRACE_EVENT_FIELD(uint64, PageSize) // new in UE 5.5
 	UE_TRACE_EVENT_FIELD(uint32, MarkerPeriod)
 	UE_TRACE_EVENT_FIELD(uint8, Version)
 	UE_TRACE_EVENT_FIELD(uint8, MinAlignment)
 	UE_TRACE_EVENT_FIELD(uint8, SizeShift)
-	UE_TRACE_EVENT_FIELD(uint8, Mode) // unused
 UE_TRACE_EVENT_END()
 
 UE_TRACE_EVENT_BEGIN(Memory, Marker)
@@ -132,6 +132,13 @@ UE_TRACE_EVENT_BEGIN(Memory, ReallocFreeSystem)
 	UE_TRACE_EVENT_FIELD(uint32, CallstackId)
 UE_TRACE_EVENT_END()
 
+UE_TRACE_EVENT_BEGIN(Memory, MemorySwapOp)
+	UE_TRACE_EVENT_FIELD(uint64, Address) // page fault real address
+	UE_TRACE_EVENT_FIELD(uint32, CallstackId)
+	UE_TRACE_EVENT_FIELD(uint32, CompressedSize)
+	UE_TRACE_EVENT_FIELD(uint8,  SwapOp)
+UE_TRACE_EVENT_END()
+
 UE_TRACE_EVENT_BEGIN(Memory, HeapSpec, NoSync|Important)
 	UE_TRACE_EVENT_FIELD(HeapId, Id)
 	UE_TRACE_EVENT_FIELD(HeapId, ParentId)
@@ -380,7 +387,10 @@ void MemoryTrace_Initialize()
 	// At this point we initialized the system to allow tracing.
 	GTraceAllowed = true;
 
+	const FPlatformMemoryConstants& MemoryConstants = FPlatformMemory::GetConstants();
+
 	UE_TRACE_LOG(Memory, Init, MemAllocChannel)
+		<< Init.PageSize(MemoryConstants.PageSize)
 		<< Init.MarkerPeriod(MarkerSamplePeriod + 1)
 		<< Init.Version(MemoryTraceVersion)
 		<< Init.MinAlignment(uint8(MIN_ALIGNMENT))
@@ -597,6 +607,23 @@ void MemoryTrace_ReallocFree(uint64 Address, HeapId RootHeap, uint32 ExternalCal
 	MemoryTrace_UpdateInternal();
 }
 
+////////////////////////////////////////////////////////////////////////////////
+void MemoryTrace_SwapOp(uint64 PageAddress, EMemoryTraceSwapOperation SwapOperation, uint32 CompressedSize, uint32 CallstackId)
+{
+	if (!GTraceAllowed)
+	{
+		return;
+	}
+
+	UE_TRACE_LOG(Memory, MemorySwapOp, MemAllocChannel)
+		<< MemorySwapOp.Address(PageAddress)
+		<< MemorySwapOp.CallstackId(CallstackId)
+		<< MemorySwapOp.CompressedSize(CompressedSize)
+		<< MemorySwapOp.SwapOp((uint8)SwapOperation);
+
+	MemoryTrace_UpdateInternal();
+}
+
 ////////////////////////////////////////////////////////////////////////////////
 HeapId MemoryTrace_HeapSpec(HeapId ParentId, const TCHAR* Name, EMemoryTraceHeapFlags Flags)
 {
diff --git a/Engine/Source/Runtime/Core/Public/ProfilingDebugging/MemoryTrace.h b/Engine/Source/Runtime/Core/Public/ProfilingDebugging/MemoryTrace.h
--- a/Engine/Source/Runtime/Core/Public/ProfilingDebugging/MemoryTrace.h
+++ b/Engine/Source/Runtime/Core/Public/ProfilingDebugging/MemoryTrace.h
@@ -56,9 +56,18 @@ enum class EMemoryTraceHeapAllocationFlags : uint8
 {
 	None = 0,
 	Heap = 1 << 0, // Is a heap, can be used to unmark alloc as heap.
+	Swap = 2 << 0, // Is a swap page
 };
 ENUM_CLASS_FLAGS(EMemoryTraceHeapAllocationFlags);
 
+////////////////////////////////////////////////////////////////////////////////
+enum class EMemoryTraceSwapOperation : uint8
+{
+	PageOut    = 0, // Paged out to swap
+	PageIn     = 1, // Read from swap via page fault
+	FreeInSwap = 2, // Freed while being paged out in swap
+};
+
 ////////////////////////////////////////////////////////////////////////////////
 #if UE_MEMORY_TRACE_ENABLED
 
@@ -140,6 +149,14 @@ CORE_API void MemoryTrace_ReallocFree(uint64 Address, HeapId RootHeap = EMemoryT
  */
 CORE_API void MemoryTrace_ReallocAlloc(uint64 Address, uint64 NewSize, uint32 Alignment, HeapId RootHeap = EMemoryTraceRootHeap::SystemMemory, uint32 ExternalCallstackId = 0);
 
+/** Trace a swap operation. Only available for system memory root heap (EMemoryTraceRootHeap::SystemMemory).
+ * @param PageAddress Page address for operation, in case of PageIn can be address of the page fault (not aligned to page boundary).
+ * @param SwapOperation Which swap operation is happening to the address.
+ * @param CompressedSize Compressed size of the page for page out operation.
+ * @param CallstackId CallstackId to use, if 0 to ignore (will not use current callstack id).
+ */
+CORE_API void MemoryTrace_SwapOp(uint64 PageAddress, EMemoryTraceSwapOperation SwapOperation, uint32 CompressedSize = 0, uint32 CallstackId = 0);
+
 ////////////////////////////////////////////////////////////////////////////////
 #else // UE_MEMORY_TRACE_ENABLED
 
@@ -152,5 +169,6 @@ inline void MemoryTrace_Alloc(uint64 Address, uint64 Size, uint32 Alignment, Hea
 inline void MemoryTrace_Free(uint64 Address, HeapId RootHeap = EMemoryTraceRootHeap::SystemMemory, uint32 ExternalCallstackId = 0) {}
 inline void MemoryTrace_ReallocFree(uint64 Address, HeapId RootHeap = EMemoryTraceRootHeap::SystemMemory, uint32 ExternalCallstackId = 0) {}
 inline void MemoryTrace_ReallocAlloc(uint64 Address, uint64 NewSize, uint32 Alignment, HeapId RootHeap = EMemoryTraceRootHeap::SystemMemory, uint32 ExternalCallstackId = 0) {}
+inline void MemoryTrace_SwapOp(uint64 PageAddress, EMemoryTraceSwapOperation SwapOperation, uint32 CompressedSize = 0, uint32 CallstackId = 0) {}
 
 #endif // UE_MEMORY_TRACE_ENABLED
