[Insights] Memory Insights: Changed error handling for cases where trace contains duplicated LLM memory tag declarations. If declarations are identical, a warning message is emitted. If declarations are different (either different parent tag id or different display name) an error is emitted.

diff --git a/Engine/Source/Developer/TraceServices/Private/Model/AllocationsProvider.cpp b/Engine/Source/Developer/TraceServices/Private/Model/AllocationsProvider.cpp
--- a/Engine/Source/Developer/TraceServices/Private/Model/AllocationsProvider.cpp
+++ b/Engine/Source/Developer/TraceServices/Private/Model/AllocationsProvider.cpp
@@ -205,20 +205,35 @@ void FTagTracker::AddTagSpec(TagIdType InTag, TagIdType InParentTag, const TCHAR
 		return;
 	}
 
-	if (ensure(!TagMap.Contains(InTag)))
+	if (!InDisplay || *InDisplay == TEXT('\0'))
 	{
-		FStringView Display(InDisplay);
-		FString DisplayName;
-		TStringBuilder<128> FullName;
-		if (Display.Contains(TEXT("/")))
+		++NumErrors;
+		if (NumErrors <= MaxLogMessagesPerErrorType)
+		{
+			UE_LOG(LogTraceServices, Error, TEXT("[MemAlloc] Tag with id %u has invalid display name (ParentTag=%u)!"), InTag, InParentTag);
+		}
+		InDisplay = TEXT("Unknown");
+	}
+
+	const FTagEntry* TagEntry = TagMap.Find(InTag);
+	if (!TagEntry)
+	{
+		const TCHAR* TagDisplayName;
+		const TCHAR* TagFullPath;
+
+		FStringView DisplayName(InDisplay);
+		int32 OutIndex;
+		if (DisplayName.FindLastChar(TEXT('/'), OutIndex))
 		{
-			DisplayName = FPathViews::GetPathLeaf(Display);
-			FullName = Display;
+			DisplayName.RightChopInline(OutIndex + 1);
+			TagDisplayName = Session.StoreString(DisplayName);
+			TagFullPath = Session.StoreString(InDisplay);
+
 			// It is possible to define a child tag in runtime using only a string, even if the parent tag does not yet
 			// exist. We need to find the correct parent or store it to the side until the parent tag is announced.
 			if (InParentTag == InvalidTagId)
 			{
-				const FStringView Parent = FPathViews::GetPathLeaf(FPathViews::GetPath(Display));
+				const FStringView Parent = FPathViews::GetPathLeaf(FPathViews::GetPath(InDisplay));
 				for (const auto& EntryPair : TagMap)
 				{
 					const uint32 Id = EntryPair.Get<0>();
@@ -238,12 +253,12 @@ void FTagTracker::AddTagSpec(TagIdType InTag, TagIdType InParentTag, const TCHAR
 		}
 		else
 		{
-			DisplayName = Display;
-			BuildTagPath(FullName, Display, InParentTag);
+			TagDisplayName = Session.StoreString(DisplayName);
+			TStringBuilder<128> FullNameBuilder;
+			BuildTagPath(FullNameBuilder, DisplayName, InParentTag);
+			TagFullPath = Session.StoreString(FullNameBuilder);
 		}
 
-		const TCHAR* TagDisplayName = Session.StoreString(DisplayName);
-		const TCHAR* TagFullPath = Session.StoreString(FullName.ToString());
 		const FTagEntry& Entry = TagMap.Emplace(InTag, FTagEntry{ TagDisplayName, TagFullPath, InParentTag });
 
 		// Check if this new tag has been referenced before by a child tag
@@ -251,27 +266,38 @@ void FTagTracker::AddTagSpec(TagIdType InTag, TagIdType InParentTag, const TCHAR
 		{
 			const TagIdType ReferencingId = Pending.Get<0>();
 			const FString& Name = Pending.Get<1>();
-			if (Name.Equals(DisplayName))
+			if (DisplayName.Equals(Name))
 			{
 				TagMap[ReferencingId].ParentTag = InTag;
 			}
 		}
 
-		if (!InDisplay || *InDisplay == TEXT('\0'))
+		UE_LOG(LogTraceServices, Verbose, TEXT("[MemAlloc] Added Tag '%s' ('%s') with id %u (ParentTag=%u)."), Entry.Display, Entry.FullPath, InTag, InParentTag);
+	}
+	else
+	{
+		FStringView DisplayName(InDisplay);
+		int32 OutIndex;
+		if (DisplayName.FindLastChar(TEXT('/'), OutIndex))
 		{
-			UE_LOG(LogTraceServices, Warning, TEXT("[MemAlloc] Tag with id %u has invalid display name (ParentTag=%u)!"), InTag, InParentTag);
+			DisplayName.RightChopInline(OutIndex + 1);
 		}
-		else
+
+		if (InParentTag == TagEntry->ParentTag && DisplayName.Equals(TagEntry->Display))
 		{
-			UE_LOG(LogTraceServices, Verbose, TEXT("[MemAlloc] Added Tag '%s' ('%s') with id %u (ParentTag=%u)."), Entry.Display, Entry.FullPath, InTag, InParentTag);
+			++NumWarnings;
+			if (NumWarnings <= MaxLogMessagesPerWarningType)
+			{
+				UE_LOG(LogTraceServices, Warning, TEXT("[MemAlloc] Tag with id %u (ParentTag=%u, Display='%s') was already added!"), InTag, InParentTag, InDisplay);
+			}
 		}
-	}
-	else
-	{
-		++NumErrors;
-		if (NumErrors <= MaxLogMessagesPerErrorType)
+		else
 		{
-			UE_LOG(LogTraceServices, Error, TEXT("[MemAlloc] Tag with id %u (ParentTag=%u, Display='%s') already added!"), InTag, InParentTag, InDisplay);
+			++NumErrors;
+			if (NumErrors <= MaxLogMessagesPerErrorType)
+			{
+				UE_LOG(LogTraceServices, Error, TEXT("[MemAlloc] Tag with id %u (ParentTag=%u, Display='%s') was already added (ParentTag=%u, Display='%s')!"), InTag, InParentTag, InDisplay, TagEntry->ParentTag, TagEntry->Display);
+			}
 		}
 	}
 }
@@ -1712,7 +1738,7 @@ void FAllocationsProvider::EditSwapOp(uint32 ThreadId, double Time, uint64 Unmas
 
 	FRootHeap& RootHeap = *RootHeaps[RootHeapId];
 
-	const uint64 MaskedPagedAddress = UnmaskedPageAddress & (~(PlatformPageSize - 1)); 
+	const uint64 MaskedPagedAddress = UnmaskedPageAddress & (~(PlatformPageSize - 1));
 
 	// emit fake page in if there was previous page out with the same address
 	if (SwapOp == EMemoryTraceSwapOperation::PageOut && RootHeap.LiveAllocs->FindSwapRef(MaskedPagedAddress) != nullptr)
@@ -1792,7 +1818,7 @@ void FAllocationsProvider::EditSwapOp(uint32 ThreadId, double Time, uint64 Unmas
 			SampleMinSwapMemory = FMath::Min(SampleMinSwapMemory, TotalSwapMemory);
 
 			FAllocationItem* AllocationPtr = RootHeap.LiveAllocs->RemoveSwap(MaskedPagedAddress);
-		
+
 			if (!AllocationPtr)
 			{
 				++SwapErrors;
@@ -2546,6 +2572,10 @@ void FAllocationsProvider::EditOnAnalysisCompleted(double Time)
 			AllocWarnings + FreeWarnings + HeapWarnings + SwapWarnings + MiscWarnings,
 			AllocWarnings, FreeWarnings, HeapWarnings, SwapWarnings, MiscWarnings);
 	}
+	if (TagTracker.GetNumWarnings() > 0)
+	{
+		UE_LOG(LogTraceServices, Warning, TEXT("[MemAlloc] TagTracker warnings: %u"), TagTracker.GetNumWarnings());
+	}
 
 	if (AllocErrors > 0)
 	{
diff --git a/Engine/Source/Developer/TraceServices/Private/Model/AllocationsProvider.h b/Engine/Source/Developer/TraceServices/Private/Model/AllocationsProvider.h
--- a/Engine/Source/Developer/TraceServices/Private/Model/AllocationsProvider.h
+++ b/Engine/Source/Developer/TraceServices/Private/Model/AllocationsProvider.h
@@ -73,6 +73,7 @@ public:
 	bool HasTagFromPtrScope(uint32 ThreadId, uint8 Tracker) const;
 
 	uint32 GetNumErrors() const { return NumErrors; }
+	uint32 GetNumWarnings() const { return NumWarnings; }
 
 private:
 	void BuildTagPath(FStringBuilderBase& OutString, FStringView Name, TagIdType ParentTagId);
@@ -86,6 +87,7 @@ private:
 	TMap<TagIdType, FTagEntry> TagMap;
 	TArray<TTuple<TagIdType, FString>> PendingTags;
 	uint32 NumErrors = 0;
+	uint32 NumWarnings = 0;
 };
 
 ////////////////////////////////////////////////////////////////////////////////////////////////////
