[Insights] MemoryInsights: Increased the allowed limit for advancing the marker/timestamp for memory events from 1 minute to 5 minutes. This allows correct processing of memory trace events in cases where app is very slow to start or in cases when app is suspended at a breakpoint.

diff --git a/Engine/Source/Developer/TraceServices/Private/Analyzers/AllocationsAnalysis.cpp b/Engine/Source/Developer/TraceServices/Private/Analyzers/AllocationsAnalysis.cpp
--- a/Engine/Source/Developer/TraceServices/Private/Analyzers/AllocationsAnalysis.cpp
+++ b/Engine/Source/Developer/TraceServices/Private/Analyzers/AllocationsAnalysis.cpp
@@ -262,7 +262,7 @@ bool FAllocationsAnalyzer::OnEvent(uint16 RouteId, EStyle Style, const FOnEventC
 			{
 				const double Seconds = Context.EventTime.AsSeconds(Cycle);
 				check(Seconds >= LastMarkerSeconds);
-				if (ensure((Seconds - LastMarkerSeconds < 60.0) || LastMarkerSeconds == 0.0f))
+				if (ensure((Seconds - LastMarkerSeconds < 5 * 60.0) || LastMarkerSeconds == 0.0f))
 				{
 					LastMarkerCycle = Cycle;
 					LastMarkerSeconds = Seconds;
