[Insights] Memory Insights / Symbols Resolver:
  - Fixed concurency issues.
  - Updated log messages re "binary files" and "debug symbol files".
  - Fixed path order for searching symbol files(1. Any new path entered by the user this session; 2. Path of the executable (if available); 3. Paths from UE_INSIGHTS_SYMBOL_PATH; 4. Paths from _NT_SYMBOL_PATH; 5. Paths from the user configuration file).
  - Fixed re-load of a module when user specifies a pdb file. Now it will automatically detect and load the corresponding exe/dll binary file.
  - Fixed the succesfull re-load of a module not triggering the re-load of other modules.

diff --git a/Engine/Source/Developer/TraceServices/Private/Model/ModuleProvider.cpp b/Engine/Source/Developer/TraceServices/Private/Model/ModuleProvider.cpp
--- a/Engine/Source/Developer/TraceServices/Private/Model/ModuleProvider.cpp
+++ b/Engine/Source/Developer/TraceServices/Private/Model/ModuleProvider.cpp
@@ -205,7 +205,15 @@ template <typename SymbolResolverType>
 FGraphEventRef TModuleProvider<SymbolResolverType>::LoadSymbolsForModuleUsingPath(uint64 Base, const TCHAR* Path)
 {
 	FReadScopeLock _(ModulesLock);
-	const FModule* Module = Algo::FindBy(Modules, Base, &FModule::Base);
+	FModule* Module = nullptr;
+	for (uint32 ModuleIndex = 0; ModuleIndex < Modules.Num(); ++ModuleIndex)
+	{
+		if (Modules[ModuleIndex].Base == Base)
+		{
+			Module = &Modules[ModuleIndex];
+			break;
+		}
+	}
 	if (Module)
 	{
 		const FString FullPath = FPaths::ConvertRelativePathToFull(Path);
@@ -215,7 +223,7 @@ FGraphEventRef TModuleProvider<SymbolResolverType>::LoadSymbolsForModuleUsingPat
 			// re-resolve any cached symbols
 			LoadSymbolsTask = FFunctionGraphTask::CreateAndDispatchWhenReady([this, Module, FullPath]()
 			{
-				auto ReloadModuleFn = [this] (const FModule* InModule, const TCHAR* InPath)
+				auto ReloadModuleFn = [this] (FModule* InModule, const TCHAR* InPath)
 				{
 					const uint32 DiscoveredSymbols = InModule->Stats.Discovered.load();
 					const uint64 ModuleBegin = InModule->Base;
@@ -240,15 +248,16 @@ FGraphEventRef TModuleProvider<SymbolResolverType>::LoadSymbolsForModuleUsingPat
 					Resolver->QueueModuleReload(InModule, InPath, ReresolveOnSuccess);
 
 					// Wait for the resolver to do it's work
-					while (InModule->Status.load() == EModuleStatus::Pending)
+					do
 					{
 						FPlatformProcess::Sleep(0.1f);
 					}
+					while (InModule->Status.load() == EModuleStatus::Pending);
 
 					return InModule->Status.load();
 				};
 
-				UE_LOG(LogTraceServices, Display, TEXT("Queing symbol loading using path %s."), *FullPath);
+				UE_LOG(LogTraceServices, Display, TEXT("Queuing symbol loading using path %s."), *FullPath);
 
 				// Load the requested module
 				const EModuleStatus Result = ReloadModuleFn(Module, *FullPath);
@@ -258,21 +267,34 @@ FGraphEventRef TModuleProvider<SymbolResolverType>::LoadSymbolsForModuleUsingPat
 					// Queue up any other failed module using the directory.
 					IPlatformFile* PlatformFile = &FPlatformFileManager::Get().GetPlatformFile();
 					const FString Directory = PlatformFile->DirectoryExists(*FullPath) ? FullPath : FPaths::GetPath(FullPath);
-					for (auto& OtherModule : Modules)
+
+					TArray<FModule*> OtherModules;
 					{
-						if (LoadSymbolsAbort)
+						FReadScopeLock _(ModulesLock);
+						for (uint32 ModuleIndex = 0; ModuleIndex < Modules.Num(); ++ModuleIndex)
 						{
-							return;
+							FModule* OtherModule = &Modules[ModuleIndex];
+							if (OtherModule != Module)
+							{
+								const EModuleStatus ModuleStatus = OtherModule->Status.load();
+								if (ModuleStatus >= EModuleStatus::FailedStatusStart)
+								{
+									OtherModules.Add(OtherModule);
+								}
+							}
 						}
-						const EModuleStatus ModuleStatus = OtherModule.Status.load();
-						if (&OtherModule != Module && ModuleStatus >= EModuleStatus::FailedStatusStart)
+					}
+					for (FModule* OtherModule : OtherModules)
+					{
+						if (LoadSymbolsAbort)
 						{
-							ReloadModuleFn(&OtherModule, *Directory);
+							return;
 						}
+						ReloadModuleFn(OtherModule, *Directory);
 					}
 				}
 
-				UE_LOG(LogTraceServices, Display, TEXT("Loading symbols for path %s complete."), *FullPath);
+				UE_LOG(LogTraceServices, Display, TEXT("Completed loading symbols for path %s."), *FullPath);
 			});
 
 			return LoadSymbolsTask;
diff --git a/Engine/Source/Developer/TraceServices/Private/Model/SymslibResolver.cpp b/Engine/Source/Developer/TraceServices/Private/Model/SymslibResolver.cpp
--- a/Engine/Source/Developer/TraceServices/Private/Model/SymslibResolver.cpp
+++ b/Engine/Source/Developer/TraceServices/Private/Model/SymslibResolver.cpp
@@ -26,28 +26,26 @@
 #include "TraceServices/Model/Diagnostics.h"
 #include <atomic>
 
-/////////////////////////////////////////////////////////////////////
-DEFINE_LOG_CATEGORY_STATIC(LogSymslib, Log, All);
-
 #if PLATFORM_WINDOWS
 	#define USE_DBG_HELP_UNDECORATOR 1
 	#if USE_DBG_HELP_UNDECORATOR
 		#include <Microsoft/AllowMicrosoftPlatformTypes.h>
-		#include "DbgHelp.h"
+		#include <DbgHelp.h>
 		#include <Microsoft/HideMicrosoftPlatformTypes.h>
 	#endif
 #else
 	#define USE_DBG_HELP_UNDECORATOR 0
 #endif
 
-/////////////////////////////////////////////////////////////////////
-namespace TraceServices {
+////////////////////////////////////////////////////////////////////////////////////////////////////
 
-/////////////////////////////////////////////////////////////////////
+DEFINE_LOG_CATEGORY_STATIC(LogSymslib, Log, All);
+
+namespace TraceServices {
 
 static const TCHAR* GUnknownModuleTextSymsLib = TEXT("Unknown");
 
-/////////////////////////////////////////////////////////////////////
+////////////////////////////////////////////////////////////////////////////////////////////////////
 
 namespace
 {
@@ -79,7 +77,7 @@ namespace
 
 		// On Linux and Mac find the non-stripped binary if available
 		// todo: We cannot actually rely on the Platform string being set yet, we can enable that check when that's part of trace files metadata.
-		// 		 for now we always override the binary if a .debug file exists
+		//       for now we always override the binary if a .debug file exists
 		//if (Platform.Equals(TEXT("Linux")) || Platform.Equals(TEXT("Mac")))
 		{
 			FString NonStrippedFile = FPaths::SetExtension(File, TEXT("debug"));
@@ -90,27 +88,42 @@ namespace
 			}
 		}
 
+		FString FileName = FPaths::GetCleanFilename(File);
+
 		// If the search path is an absolute path to a file use this. Filenames does
-		// not need to match (e.g. eboot.bin <-> gamename.self)
-		if (PlatformFile->FileExists(*SearchPath))
+		// not need to match (e.g. eboot.bin <-> gamename.self);
+		// but except the *.pdb files (in which case, we need to load the exe or dll binary instead).
+		if (!SearchPath.EndsWith(TEXT(".pdb")))
 		{
-			return SearchPath;
+			if (PlatformFile->FileExists(*SearchPath))
+			{
+				return SearchPath;
+			}
 		}
 
-		// Extract only filename part in case Path is absolute path
-		FString FileName = FPaths::GetCleanFilename(File);
+		// Look for exact filename of the module in the provided search path.
+		{
+			FString Result = FPaths::Combine(SearchPath, FileName);
+			if (PlatformFile->FileExists(*Result))
+			{
+				return Result;
+			}
+		}
 
-		// Look in the search path
-		FString Result = FPaths::Combine(SearchPath, FileName);
-		if (PlatformFile->FileExists(*Result))
+		// Look for exact filename of the module in the provided search path
+		// (case where the search path includes a filename).
 		{
-			return Result;
+			FString Result = FPaths::Combine(FPaths::GetPath(SearchPath), FileName);
+			if (PlatformFile->FileExists(*Result))
+			{
+				return Result;
+			}
 		}
 
 		// In case File is relative
-		if (FPaths::IsRelative(SearchPath))
+		if (FPaths::IsRelative(File))
 		{
-			Result = FPaths::Combine(SearchPath, File);
+			FString Result = FPaths::Combine(SearchPath, File);
 			if (PlatformFile->FileExists(*Result))
 			{
 				return Result;
@@ -123,7 +136,6 @@ namespace
 			return File;
 		}
 
-
 		// File not found
 		return FString();
 	}
@@ -175,19 +187,14 @@ namespace
 	static FString FindFileInEngineFolder(const FString& FilePath, const FString& SearchPath, const FString& Platform, const FString& Project)
 	{
 		IPlatformFile* PlatformFile = &FPlatformFileManager::Get().GetPlatformFile();
-		const TCHAR* EngineFolders[] = {TEXT("Engine"), TEXT("Plugins")};
-		FString Result;
+
 		FString File = FPaths::GetCleanFilename(FilePath);
 
-		// Look in some well known directories
-		for (const TCHAR* Folder : EngineFolders)
+		// Look in engine directory
+		FString Result = FPaths::Combine(SearchPath, TEXT("Engine"), TEXT("Binaries"), Platform, File);
+		if (PlatformFile->FileExists(*Result))
 		{
-			Result = FPaths::Combine(SearchPath, Folder, TEXT("Binaries"), Platform, File);
-			if (PlatformFile->FileExists(*Result))
-			{
-				return Result;
-			}
-
+			return Result;
 		}
 
 		// Look in project directory
@@ -233,9 +240,9 @@ namespace
 		return FString();
 	}
 
-	static bool LoadBinary(const TCHAR* Path, SYMS_Arena* Arena, SYMS_ParseBundle& Bundle, TArray<FAutoMappedFile>& Files, const FString& SearchPath, const FString& Platform, const FString& AppName)
+	static bool LoadBinary(const TCHAR* ModuleFullName, SYMS_Arena* Arena, SYMS_ParseBundle& Bundle, TArray<FAutoMappedFile>& Files, const FString& SearchPath, const FString& Platform, const FString& AppName)
 	{
-		FString FilePath = Path;
+		FString FilePath = ModuleFullName;
 		bool bFileFound = false;
 
 		// Remap known renamed binaries
@@ -261,7 +268,7 @@ namespace
 		FAutoMappedFile& File = Files.AddDefaulted_GetRef();
 		if (!File.Load(*BinaryPath))
 		{
-			UE_LOG(LogSymslib, Warning, TEXT("Failed to load '%s' file"), *BinaryPath);
+			UE_LOG(LogSymslib, Warning, TEXT("Failed to load binary file '%s'"), *BinaryPath);
 			return false;
 		}
 
@@ -269,7 +276,7 @@ namespace
 		SYMS_BinAccel* BinAccel = syms_bin_accel_from_file(Arena, File.GetData(), Accel);
 		if (!syms_accel_is_good(BinAccel))
 		{
-			UE_LOG(LogSymslib, Warning, TEXT("Cannot parse '%s' binary file"), *BinaryPath);
+			UE_LOG(LogSymslib, Warning, TEXT("Cannot parse binary file '%s'"), *BinaryPath);
 			return false;
 		}
 
@@ -289,6 +296,7 @@ namespace
 			// binary has debug info built-in (like dwarf file)
 			Bundle.dbg = syms_dbg_accel_from_bin(Arena, Files[0].GetData(), Bundle.bin);
 			Bundle.dbg_data = Bundle.bin_data;
+			UE_LOG(LogSymslib, Verbose, TEXT("Binary file '%s' has debug info built-in"), *Files[0].FilePath);
 			return true;
 		}
 
@@ -304,16 +312,17 @@ namespace
 		// debug file path from metadata in executable
 		FString FilePath = ANSI_TO_TCHAR(reinterpret_cast<char*>(ExtFile.file_name.str));
 
-		bool bFileFound = false;
+		// try in the same path as the binary file
+		FString BinaryPath = FPaths::GetPath(Files[0].FilePath);
+		FString DebugPath = FindSymbolFileInPath(FilePath, BinaryPath);
 
-		// First try file in symbol path
-		FString DebugPath = FindSymbolFileInPath(FilePath, SearchPath);
-		if (!DebugPath.IsEmpty())
+		if (DebugPath.IsEmpty())
 		{
-			bFileFound = true;
+			// try in provided search path
+			DebugPath = FindSymbolFileInPath(FilePath, SearchPath);
 		}
 
-		if (!bFileFound)
+		if (DebugPath.IsEmpty())
 		{
 			// if executable is PE format, try Windows symbol path format with guid in path
 			if (Bundle.bin->format == SYMS_FileFormat_PE && SearchPath.Contains(TEXT("*")))
@@ -332,23 +341,19 @@ namespace
 					*FileName);
 
 				DebugPath = FindWindowsSymbolFile(GuidPath, SearchPath);
-				if (!DebugPath.IsEmpty())
-				{
-					bFileFound = true;
-				}
 			}
 		}
 
-		if (!bFileFound)
+		if (DebugPath.IsEmpty())
 		{
-			UE_LOG(LogSymslib, Verbose, TEXT("Debug file '%s' not found"), *FilePath);
+			UE_LOG(LogSymslib, Verbose, TEXT("Debug symbols file '%s' not found"), *FilePath);
 			return false;
 		}
 
 		FAutoMappedFile& File = Files.AddDefaulted_GetRef();
 		if (!File.Load(*DebugPath))
 		{
-			UE_LOG(LogSymslib, Warning, TEXT("Failed to load '%s' file"), *DebugPath);
+			UE_LOG(LogSymslib, Warning, TEXT("Failed to load debug symbols file '%s'"), *DebugPath);
 			return false;
 		}
 
@@ -356,7 +361,7 @@ namespace
 		SYMS_DbgAccel* DbgAccel = syms_dbg_accel_from_file(Arena, File.GetData(), Accel);
 		if (!syms_accel_is_good(DbgAccel))
 		{
-			UE_LOG(LogSymslib, Warning, TEXT("Cannot parse '%s' debug file"), *DebugPath);
+			UE_LOG(LogSymslib, Warning, TEXT("Cannot parse debug symbols file '%s'"), *DebugPath);
 			return false;
 		}
 
@@ -447,9 +452,10 @@ namespace
 		// either ID's are matching, or ID is not found in which case return "success" case
 		return true;
 	}
-}
 
-/////////////////////////////////////////////////////////////////////
+} // namespace
+
+////////////////////////////////////////////////////////////////////////////////////////////////////
 
 class FSymbolStringAllocator
 {
@@ -492,28 +498,21 @@ private:
 	uint32 BlockUsed = 0;
 };
 
-/////////////////////////////////////////////////////////////////////
+////////////////////////////////////////////////////////////////////////////////////////////////////
+
 FSymslibResolver::FSymslibResolver(IAnalysisSession& InSession, IResolvedSymbolFilter& InSymbolFilter)
 	: Modules(InSession.GetLinearAllocator(), 128)
 	, CancelTasks(false)
 	, Session(InSession)
 	, SymbolFilter(InSymbolFilter)
 {
-	// Setup search paths. The SearchPaths array is a priority stack, which
-	// means paths are searched in reversed order.
-	// 1. Any new paths entered by the user this session
+	// Setup search paths. Paths are searched in the following order:
+	// 1. Any new path entered by the user this session
 	// 2. Path of the executable (if available)
-	// 3. Paths from UE_INSIGHTS_SYMBOLPATH
+	// 3. Paths from UE_INSIGHTS_SYMBOL_PATH
 	// 4. Paths from _NT_SYMBOL_PATH
 	// 5. Paths from the user configuration file
 
-	// Paths from configuration
-	FString SettingsIni;
-	if (FConfigContext::ReadIntoGConfig().Load(TEXT("UnrealInsightsSettings"), SettingsIni))
-	{
-		GConfig->GetArray(TEXT("Insights.MemoryProfiler"), TEXT("SymbolSearchPaths"), SymbolSearchPaths, SettingsIni);
-	}
-
 	auto SplitEnvPaths = [](FStringView EnvVariable, TArray<FString>& OutList)
 	{
 		FString SymbolPathPart, SymbolPathRemainder(EnvVariable);
@@ -529,28 +528,42 @@ FSymslibResolver::FSymslibResolver(IAnalysisSession& InSession, IResolvedSymbolF
 
 	// Paths from environment
 	const FString InsightsSymbolPath =  FPlatformMisc::GetEnvironmentVariable(TEXT("UE_INSIGHTS_SYMBOL_PATH"));
-	UE_LOG(LogSymslib, Log, TEXT("UE_INSIGHTS_SYMBOL_PATH: %s"), InsightsSymbolPath.IsEmpty() ? TEXT("Not set") : *InsightsSymbolPath);
-	SplitEnvPaths(InsightsSymbolPath, SymbolSearchPaths);
+	UE_LOG(LogSymslib, Log, TEXT("UE_INSIGHTS_SYMBOL_PATH: %s"), InsightsSymbolPath.IsEmpty() ? TEXT("not set") : *InsightsSymbolPath);
+	SplitEnvPaths(InsightsSymbolPath, ConfigSymbolSearchPaths);
 
 #if PLATFORM_WINDOWS
 	const FString NTSymbolPath = FPlatformMisc::GetEnvironmentVariable(TEXT("_NT_SYMBOL_PATH"));
-	UE_LOG(LogSymslib, Log, TEXT("_NT_SYMBOL_PATH: %s"), NTSymbolPath.IsEmpty() ? TEXT("Not set") : *NTSymbolPath);
-	SplitEnvPaths(NTSymbolPath, SymbolSearchPaths);
+	UE_LOG(LogSymslib, Log, TEXT("_NT_SYMBOL_PATH: %s"), NTSymbolPath.IsEmpty() ? TEXT("not set") : *NTSymbolPath);
+	SplitEnvPaths(NTSymbolPath, ConfigSymbolSearchPaths);
 #endif
 
-	// Try to get session information
-	// Depending on how module tracing and session info tracing is implemented on different platforms this may not
-	// be available at this point.
-	FAnalysisSessionReadScope _(Session);
-	const IDiagnosticsProvider* DiagnosticsProvider = ReadDiagnosticsProvider(Session);
-	if (DiagnosticsProvider && DiagnosticsProvider->IsSessionInfoAvailable())
+	// Paths from configuration
+	FString SettingsIni;
+	if (FConfigContext::ReadIntoGConfig().Load(TEXT("UnrealInsightsSettings"), SettingsIni))
 	{
-		const FSessionInfo Info = DiagnosticsProvider->GetSessionInfo();
-		Platform = Info.Platform;
-		AppName = Info.AppName;
+		TArray<FString> SymbolSearchPaths;
+		GConfig->GetArray(TEXT("Insights.MemoryProfiler"), TEXT("SymbolSearchPaths"), SymbolSearchPaths, SettingsIni);
+#if !NO_LOGGING
+		if (SymbolSearchPaths.IsEmpty())
+		{
+			UE_LOG(LogSymslib, Log, TEXT("[Insights.MemoryProfiler] SymbolSearchPaths not set"));
+		}
+		else
+		{
+			for (const FString& Path : SymbolSearchPaths)
+			{
+				UE_LOG(LogSymslib, Log, TEXT("[Insights.MemoryProfiler] +SymbolSearchPaths=%s"), *Path);
+			}
+		}
+#endif
+		ConfigSymbolSearchPaths += SymbolSearchPaths;
 	}
+
+	UpdateSessionInfo();
 }
 
+////////////////////////////////////////////////////////////////////////////////////////////////////
+
 FSymslibResolver::~FSymslibResolver()
 {
 	CancelTasks = true;
@@ -566,20 +579,33 @@ FSymslibResolver::~FSymslibResolver()
 	}
 }
 
+////////////////////////////////////////////////////////////////////////////////////////////////////
+
 void FSymslibResolver::QueueModuleLoad(const uint8* ImageId, uint32 ImageIdSize, FModule* Module)
 {
+	check(Module != nullptr);
+	ensure(GetModuleEntry(Module) == nullptr);
+
 	FWriteScopeLock _(ModulesLock);
+
 	check(!CleanupTask); // No queued addresses should be queued after cleanup task is started
 
-	// Add module
+	// Add the new module entry.
 	FModuleEntry* Entry = &Modules.PushBack();
 	Entry->Module = Module;
 	Entry->ImageId = TArrayView<const uint8>(ImageId, ImageIdSize);
 
+	// Sort list according to base address.
+	SortedModules.Add(Entry);
+	Algo::Sort(SortedModules, [](const FModuleEntry* Lhs, const FModuleEntry* Rhs) { return Lhs->Module->Base < Rhs->Module->Base; });
+
 	++ModulesDiscovered;
 
-	// Queue up module to have symbols loaded. Run as background task as to not
-	// interfere with Slate.
+	// Set the Pending state before scheduling the background task (to allow calling code to wait, if needed).
+	Module->Status.store(EModuleStatus::Pending);
+
+	// Queue up module to have symbols loaded.
+	// Run as background task as to not interfere with Slate.
 	++TasksInFlight;
 	FFunctionGraphTask::CreateAndDispatchWhenReady(
 		[this, Entry]
@@ -587,117 +613,113 @@ void FSymslibResolver::QueueModuleLoad(const uint8* ImageId, uint32 ImageIdSize,
 			LoadModuleTracked(Entry, FStringView());
 			--TasksInFlight;
 		}, TStatId{}, nullptr, ENamedThreads::AnyBackgroundThreadNormalTask);
-
-	// Sort list according to base address
-	SortedModules.Add(Entry);
-	Algo::Sort(SortedModules, [](const FModuleEntry* Lhs, const FModuleEntry* Rhs) { return Lhs->Module->Base < Rhs->Module->Base; });
 }
 
-void FSymslibResolver::QueueModuleReload(const FModule* Module, const TCHAR* InPath, TFunction<void(SymbolArray&)> ResolveOnSuccess)
+////////////////////////////////////////////////////////////////////////////////////////////////////
+
+void FSymslibResolver::QueueModuleReload(FModule* Module, const TCHAR* Path, TFunction<void(SymbolArray&)> ResolveOnSuccess)
 {
-	FWriteScopeLock _(ModulesLock);
+	check(Module != nullptr);
 
 	// Find the entry
-	FModuleEntry* Entry = nullptr;
-	for (uint32 Idx = 0; Idx < Modules.Num(); ++Idx)
+	FModuleEntry* Entry = GetModuleEntry(Module);
+	if (!Entry)
 	{
-		if (Modules[Idx].Module == Module)
-		{
-			Entry = &Modules[Idx];
-			break;
-		}
+		return;
 	}
 
-	if (Entry)
+	// No use in trying reload already loaded modules
+	if (Module->Status == EModuleStatus::Loaded)
 	{
-		// No use in trying reload already loaded modules
-		if (Entry->Module->Status == EModuleStatus::Loaded)
-		{
-			return;
-		}
+		return;
+	}
 
-		FString PathStr(InPath);
-		FPaths::NormalizeDirectoryName(PathStr);
-		const TCHAR* Path = Session.StoreString(PathStr);
+	// Set the Pending state before scheduling the background task (to allow calling code to wait, if needed).
+	EModuleStatus PreviousStatus = Module->Status.exchange(EModuleStatus::Pending);
+	if (PreviousStatus >= EModuleStatus::FailedStatusStart)
+	{
+		--ModulesFailed;
+	}
 
-		// Can only launch one task at a time
-		if (ModuleReloadTask && !ModuleReloadTask->IsComplete())
-		{
-			ModuleReloadTask->Wait();
-		}
+	FString PathStr(Path);
+	FPaths::NormalizeDirectoryName(PathStr);
+	const TCHAR* OverrideSearchPath = Session.StoreString(PathStr);
 
-		FFunctionGraphTask::CreateAndDispatchWhenReady([this, Entry, Path, ResolveOnSuccess]
-		{
-			FScopeLock CleanupExclusive(&CleanupLock);
+	// Can only launch one task at a time
+	if (ModuleReloadTask && !ModuleReloadTask->IsComplete())
+	{
+		ModuleReloadTask->Wait();
+	}
 
-			EModuleStatus PreviousStatus = Entry->Module->Status.load();
+	FFunctionGraphTask::CreateAndDispatchWhenReady([this, Entry, OverrideSearchPath, ResolveOnSuccess]
+	{
+		FScopeLock CleanupExclusive(&CleanupLock);
 
-			LoadModuleTracked(Entry, Path);
+		LoadModuleTracked(Entry, OverrideSearchPath);
 
-			if (PreviousStatus >= EModuleStatus::FailedStatusStart)
+		if (Entry->Module->Status.load() == EModuleStatus::Loaded)
+		{
+			// If an additional search path was specified and successful, add to the permanent
+			// list of search paths.
 			{
-				--ModulesFailed;
+				FWriteScopeLock _(CustomSymbolSearchPathsLock);
+				// Get the base path if path points to a file
+				IPlatformFile* PlatformFile = &FPlatformFileManager::Get().GetPlatformFile();
+				const FString Directory = PlatformFile->DirectoryExists(OverrideSearchPath) ? OverrideSearchPath : FPaths::GetPath(OverrideSearchPath);
+				if (!Directory.IsEmpty())
+				{
+					CustomSymbolSearchPaths.AddUnique(Directory);
+				}
 			}
 
-			if (Entry->Module->Status.load() == EModuleStatus::Loaded)
+			// Reset stats for reloaded module.
+			// Note: If there are symbols pending to resolve (ex. busy spinning in ResolveSymbol()), the Stats.Discovered
+			//       might be already incremented (see ResolveSymbolTracked()). So, resetting the stats here might result
+			//       in slightly wrong stats where Discovered < Resolved + Failed.
+			FModule::SymbolStats& ModuleStats = Entry->Module->Stats;
+			ModuleStats.Discovered.store(0u);
+			ModuleStats.Resolved.store(0u);
+			ModuleStats.Failed.store(0u);
+
+			// Ask the caller for a list of symbols that should be resolved, now that module is properly loaded and
+			// resolve them immediately.
 			{
-				// If an additional search path was specified and successful, add to the permanent
-				// list of search paths.
+				SymbolArray SymbolsToResolve;
+				FSymbolStringAllocator StringAllocator(Session.GetLinearAllocator(), (2 << 12) / sizeof(TCHAR));
+				ResolveOnSuccess(SymbolsToResolve);
+				for (TTuple<uint64, FResolvedSymbol*> Pair : SymbolsToResolve)
 				{
-					FWriteScopeLock _(SymbolSearchPathsLock);
-					// Get the base path if path points to a file
-					IPlatformFile* PlatformFile = &FPlatformFileManager::Get().GetPlatformFile();
-					const FString Directory = PlatformFile->DirectoryExists(Path) ? Path : FPaths::GetPath(Path);
-					if (!Directory.IsEmpty())
-					{
-						SymbolSearchPaths.AddUnique(Directory);
-					}
-				}
-
-				// Reset stats for reloaded module.
-				// Note: If there are symbols pending to resolve (ex. busy spinning in ResolveSymbol()), the Stats.Discovered
-				//       might be already incremented (see ResolveSymbolTracked()). So, resetting the stats here might result
-				//       in slightly wrong stats where Discovered < Resolved + Failed.
-				Entry->Module->Stats.Discovered.store(0u);
-				Entry->Module->Stats.Resolved.store(0u);
-				Entry->Module->Stats.Failed.store(0u);
-
-				// Ask the caller for a list of symbols that should be resolved, now that module is properly loaded and
-				// resolve them immediately.
-				{
-					SymbolArray SymbolsToResolve;
-					FSymbolStringAllocator StringAllocator(Session.GetLinearAllocator(), (2 << 12) / sizeof(TCHAR));
-					ResolveOnSuccess(SymbolsToResolve);
-					for (TTuple<uint64, FResolvedSymbol*> Pair : SymbolsToResolve)
-					{
-						ResolveSymbolTracked(Pair.Get<0>(), *Pair.Get<1>(), StringAllocator);
-					}
+					ResolveSymbolTracked(Pair.Get<0>(), *Pair.Get<1>(), StringAllocator);
 				}
+			}
 
-				// Finally if the cleanup task has already been dispatched, no other resolve request will come in and we
-				// can safely release our resources. CleanupTask and a reload task cannot run at the same time due to CleanupLock.
-				if (CleanupTask)
-				{
-					// Need to wait for tasks, since there could be queued symbol resolves for this module waiting
-					// in the queue.
-					WaitForTasks();
+			// Finally if the cleanup task has already been dispatched, no other resolve request will come in and we
+			// can safely release our resources. CleanupTask and a reload task cannot run at the same time due to CleanupLock.
+			if (CleanupTask)
+			{
+				// Need to wait for tasks, since there could be queued symbol resolves for this module waiting
+				// in the queue.
+				WaitForTasks();
 
-					Algo::ForEach(Entry->Instance.Arenas, syms_arena_release);
-					Entry->Instance.Arenas.Empty();
-				}
+				Algo::ForEach(Entry->Instance.Arenas, syms_arena_release);
+				Entry->Instance.Arenas.Empty();
 			}
-		}, TStatId{}, nullptr, ENamedThreads::AnyBackgroundThreadNormalTask);
-	}
+		}
+	}, TStatId{}, nullptr, ENamedThreads::AnyBackgroundThreadNormalTask);
 }
 
+////////////////////////////////////////////////////////////////////////////////////////////////////
+
 void FSymslibResolver::QueueSymbolResolve(uint64 Address, FResolvedSymbol* Symbol)
 {
 	FScopeLock _(&SymbolsQueueLock);
 	check(!CleanupTask); // No queued addresses should be queued after cleanup task is started
 	MaybeDispatchQueuedAddresses();
-	ResolveQueue.Add(FQueuedAddress{Address, Symbol});
+	ResolveQueue.Add(FQueuedAddress{ Address, Symbol });
 }
 
+////////////////////////////////////////////////////////////////////////////////////////////////////
+
 void FSymslibResolver::OnAnalysisComplete()
 {
 	// Dispatch any remaining requests.
@@ -718,9 +740,9 @@ void FSymslibResolver::OnAnalysisComplete()
 			FReadScopeLock _(ModulesLock);
 			for (uint32 ModuleIndex = 0; ModuleIndex < Modules.Num(); ++ModuleIndex)
 			{
-				FModuleEntry& Module = Modules[ModuleIndex];
-				Algo::ForEach(Module.Instance.Arenas, syms_arena_release);
-				Module.Instance.Arenas.Empty();
+				FModuleEntry& Entry = Modules[ModuleIndex];
+				Algo::ForEach(Entry.Instance.Arenas, syms_arena_release);
+				Entry.Instance.Arenas.Empty();
 			}
 		}
 
@@ -729,27 +751,37 @@ void FSymslibResolver::OnAnalysisComplete()
 	});
 }
 
+////////////////////////////////////////////////////////////////////////////////////////////////////
+
 void FSymslibResolver::GetStats(IModuleProvider::FStats* OutStats) const
 {
 	FReadScopeLock _(ModulesLock);
 	FMemory::Memzero(*OutStats);
-	for(const FModuleEntry& Entry : Modules)
+	for (uint32 ModuleIndex = 0; ModuleIndex < Modules.Num(); ++ModuleIndex)
 	{
-		OutStats->SymbolsDiscovered += Entry.Module->Stats.Discovered.load();
-		OutStats->SymbolsResolved += Entry.Module->Stats.Resolved.load();
-		OutStats->SymbolsFailed += Entry.Module->Stats.Failed.load();
+		FModule::SymbolStats& ModuleStats = Modules[ModuleIndex].Module->Stats;
+		OutStats->SymbolsDiscovered += ModuleStats.Discovered.load();
+		OutStats->SymbolsResolved += ModuleStats.Resolved.load();
+		OutStats->SymbolsFailed += ModuleStats.Failed.load();
 	}
 	OutStats->ModulesDiscovered = ModulesDiscovered.load();
 	OutStats->ModulesFailed = ModulesFailed.load();
 	OutStats->ModulesLoaded = ModulesLoaded.load();
 }
 
+////////////////////////////////////////////////////////////////////////////////////////////////////
+
 void FSymslibResolver::EnumerateSymbolSearchPaths(TFunctionRef<void(FStringView Path)> Callback) const
 {
-	FReadScopeLock _(SymbolSearchPathsLock);
-	Algo::ForEach(SymbolSearchPaths, Callback);
+	{
+		FReadScopeLock _(CustomSymbolSearchPathsLock);
+		Algo::ForEach(CustomSymbolSearchPaths, Callback);
+	}
+
+	Algo::ForEach(ConfigSymbolSearchPaths, Callback);
 }
 
+////////////////////////////////////////////////////////////////////////////////////////////////////
 /**
  * Checks if there are no modules in flight and that the queue has reached
  * the threshold for dispatching. Note that is up to the caller to synchronize.
@@ -763,6 +795,7 @@ void FSymslibResolver::MaybeDispatchQueuedAddresses()
 	}
 }
 
+////////////////////////////////////////////////////////////////////////////////////////////////////
 /**
  * Dispatches the currently queued addresses to be resolved. Note that is up to the caller to synchronize.
  */
@@ -791,6 +824,8 @@ void FSymslibResolver::DispatchQueuedAddresses()
 	}
 }
 
+////////////////////////////////////////////////////////////////////////////////////////////////////
+
 void FSymslibResolver::ResolveSymbols(TArrayView<FQueuedAddress>& QueuedWork)
 {
 	// Create a local string allocator. We don't use the session string store due to contention when going wide. Since
@@ -812,7 +847,24 @@ void FSymslibResolver::ResolveSymbols(TArrayView<FQueuedAddress>& QueuedWork)
 	SymbolBytesWasted.fetch_add(StringAllocator.BlockRemaining * sizeof(TCHAR));
 }
 
-FSymslibResolver::FModuleEntry* FSymslibResolver::GetModuleForAddress(uint64 Address)
+////////////////////////////////////////////////////////////////////////////////////////////////////
+
+FSymslibResolver::FModuleEntry* FSymslibResolver::GetModuleEntry(FModule* Module) const
+{
+	FReadScopeLock _(ModulesLock);
+	for (FModuleEntry* Entry : SortedModules)
+	{
+		if (Entry->Module == Module)
+		{
+			return Entry;
+		}
+	}
+	return nullptr;
+}
+
+////////////////////////////////////////////////////////////////////////////////////////////////////
+
+FSymslibResolver::FModuleEntry* FSymslibResolver::GetModuleForAddress(uint64 Address) const
 {
 	FReadScopeLock _(ModulesLock);
 	const int32 EntryIdx = Algo::LowerBoundBy(SortedModules, Address, [](const FModuleEntry* Entry) { return Entry->Module->Base; }) - 1;
@@ -823,6 +875,8 @@ FSymslibResolver::FModuleEntry* FSymslibResolver::GetModuleForAddress(uint64 Add
 	return SortedModules[EntryIdx];
 }
 
+////////////////////////////////////////////////////////////////////////////////////////////////////
+
 void FSymslibResolver::UpdateResolvedSymbol(FResolvedSymbol& Symbol, ESymbolQueryResult Result, const TCHAR* Module, const TCHAR* Name, const TCHAR* File, uint16 Line)
 {
 	Symbol.Module = Module;
@@ -832,42 +886,64 @@ void FSymslibResolver::UpdateResolvedSymbol(FResolvedSymbol& Symbol, ESymbolQuer
 	Symbol.Result.store(Result, std::memory_order_release);
 }
 
+////////////////////////////////////////////////////////////////////////////////////////////////////
+
 void FSymslibResolver::LoadModuleTracked(FModuleEntry* Entry, FStringView OverrideSearchPath)
 {
-	EModuleStatus PreviousStatus = Entry->Module->Status.exchange(EModuleStatus::Pending);
+	FModule* Module = Entry->Module;
+
+	Module->Status.store(EModuleStatus::Pending);
 
 	// Build search paths
-	FString ModuleNamePath;
-	TArray<FStringView> SearchPaths;
+	TArray<FString> SearchPaths;
 	if (OverrideSearchPath.IsEmpty())
 	{
-		FReadScopeLock _(SymbolSearchPathsLock);
-		for (const FString& SearchPath : SymbolSearchPaths)
+		// 1. Any new path entered by the user this session
 		{
-			SearchPaths.AddUnique(SearchPath);
+			FReadScopeLock _(CustomSymbolSearchPathsLock);
+			for (const FString& SearchPath : CustomSymbolSearchPaths)
+			{
+				SearchPaths.AddUnique(SearchPath);
+			}
 		}
 
-		ModuleNamePath = FPaths::GetPath(Entry->Module->FullName);
+		// 2. Path of the executable (if available)
+		FString ModuleNamePath = FPaths::GetPath(Module->FullName);
 		FPaths::NormalizeDirectoryName(ModuleNamePath);
 		if (!ModuleNamePath.IsEmpty())
 		{
 			SearchPaths.AddUnique(ModuleNamePath);
 		}
+
+		// 3. Paths from UE_INSIGHTS_SYMBOL_PATH
+		// 4. Paths from _NT_SYMBOL_PATH
+		// 5. Paths from the user configuration file
+		for (const FString& SearchPath : ConfigSymbolSearchPaths)
+		{
+			SearchPaths.AddUnique(SearchPath);
+		}
 	}
 	else
 	{
-		SearchPaths.Add(OverrideSearchPath);
+		SearchPaths.Add(FString(OverrideSearchPath));
 	}
 
 	// Exhaust all the search paths
 	EModuleStatus Status = EModuleStatus::Failed;
 	TStringBuilder<128> StatusMessage;
-	while (Status != EModuleStatus::Loaded && !CancelTasks.load() && SearchPaths.Num() > 0)
+	for (const FString& SearchPath : SearchPaths)
 	{
+		if (CancelTasks.load())
+		{
+			break;
+		}
 		StatusMessage.Reset();
-		Status = LoadModule(Entry, SearchPaths.Top(), StatusMessage);
+		Status = LoadModule(Entry, SearchPath, StatusMessage);
 		UE_LOG(LogSymslib, Display, TEXT("%s"), StatusMessage.ToString());
-		SearchPaths.Pop();
+		if (Status == EModuleStatus::Loaded)
+		{
+			break;
+		}
 	}
 
 	if (Status == EModuleStatus::Loaded)
@@ -880,13 +956,16 @@ void FSymslibResolver::LoadModuleTracked(FModuleEntry* Entry, FStringView Overri
 	}
 
 	// Make the final status visible to the world
-	Entry->Module->StatusMessage = Session.StoreString(StatusMessage.ToView());
-	Entry->Module->Status.store(Status);
+	Module->StatusMessage = Session.StoreString(StatusMessage.ToView());
+	Module->Status.store(Status);
 }
 
+////////////////////////////////////////////////////////////////////////////////////////////////////
+
 EModuleStatus FSymslibResolver::LoadModule(FModuleEntry* Entry, FStringView SearchPathView, FStringBuilderBase& OutStatusMessage) const
 {
 	const FString SearchPath(SearchPathView);
+
 	// temporary memory used for loading
 	SYMS_Group* Group = syms_group_alloc();
 
@@ -896,17 +975,22 @@ EModuleStatus FSymslibResolver::LoadModule(FModuleEntry* Entry, FStringView Sear
 	// contents of binary & debug file
 	SYMS_ParseBundle Bundle;
 
+	if (Platform.IsEmpty())
+	{
+		UpdateSessionInfo();
+	}
+
 	if (!LoadBinary(Entry->Module->FullName, Group->arena, Bundle, Files, SearchPath, Platform, AppName))
 	{
 		syms_group_release(Group);
-		OutStatusMessage.Appendf(TEXT("Failed to load image for '%s' in '%s'."), Entry->Module->Name, *SearchPath);
+		OutStatusMessage.Appendf(TEXT("Failed to load binary for '%s' in '%s'."), Entry->Module->Name, *SearchPath);
 		return EModuleStatus::Failed;
 	}
 
 	if (!LoadDebug(Group->arena, Bundle, Files, SearchPath))
 	{
 		syms_group_release(Group);
-		OutStatusMessage.Appendf(TEXT("Failed to load symbols for '%s' in '%s'."), Entry->Module->Name, *SearchPath);
+		OutStatusMessage.Appendf(TEXT("Failed to load debug symbols for '%s' in '%s'."), Entry->Module->Name, *SearchPath);
 		return EModuleStatus::Failed;
 	}
 
@@ -1048,6 +1132,8 @@ EModuleStatus FSymslibResolver::LoadModule(FModuleEntry* Entry, FStringView Sear
 	return EModuleStatus::Loaded;
 }
 
+////////////////////////////////////////////////////////////////////////////////////////////////////
+
 void FSymslibResolver::ResolveSymbolTracked(uint64 Address, FResolvedSymbol& Target, FSymbolStringAllocator& StringAllocator)
 {
 	const ESymbolQueryResult PreviousResult = Target.Result.load();
@@ -1082,6 +1168,8 @@ void FSymslibResolver::ResolveSymbolTracked(uint64 Address, FResolvedSymbol& Tar
 	}
 }
 
+////////////////////////////////////////////////////////////////////////////////////////////////////
+
 #if USE_DBG_HELP_UNDECORATOR
 FORCEINLINE void UndecorateAndCopySymbolName(ANSICHAR* OutStr, const char* Name, uint32 MaxStringLength)
 {
@@ -1112,7 +1200,9 @@ FORCEINLINE void UndecorateAndCopySymbolName(ANSICHAR* OutStr, const char* Name,
 		FCStringAnsi::Strncpy(OutStr, Name, MaxStringLength);
 	}
 }
-#endif // PLATFORM_WINDOWS
+#endif // USE_DBG_HELP_UNDECORATOR
+
+////////////////////////////////////////////////////////////////////////////////////////////////////
 
 bool FSymslibResolver::ResolveSymbol(uint64 Address, FResolvedSymbol& Target, FSymbolStringAllocator& StringAllocator, FModuleEntry* Entry) const
 {
@@ -1236,18 +1326,37 @@ bool FSymslibResolver::ResolveSymbol(uint64 Address, FResolvedSymbol& Target, FS
 	return true;
 }
 
+////////////////////////////////////////////////////////////////////////////////////////////////////
+
 void FSymslibResolver::WaitForTasks()
 {
 	uint32 OutstandingTasks = TasksInFlight.load(std::memory_order_acquire);
 	do
 	{
-		 OutstandingTasks = TasksInFlight.load(std::memory_order_acquire);
-		 FPlatformProcess::Sleep(0.0);
+		OutstandingTasks = TasksInFlight.load(std::memory_order_acquire);
+		FPlatformProcess::Sleep(0.0);
 	}
 	while (OutstandingTasks > 0);
 }
 
-/////////////////////////////////////////////////////////////////////
+////////////////////////////////////////////////////////////////////////////////////////////////////
+
+void FSymslibResolver::UpdateSessionInfo() const
+{
+	// Try to get session information.
+	// Depending on how module tracing and session info tracing is implemented on different platforms this may not
+	// be available at this point.
+	FAnalysisSessionReadScope _(Session);
+	const IDiagnosticsProvider* DiagnosticsProvider = ReadDiagnosticsProvider(Session);
+	if (DiagnosticsProvider && DiagnosticsProvider->IsSessionInfoAvailable())
+	{
+		const FSessionInfo Info = DiagnosticsProvider->GetSessionInfo();
+		Platform = Info.Platform;
+		AppName = Info.AppName;
+	}
+}
+
+////////////////////////////////////////////////////////////////////////////////////////////////////
 
 } // namespace TraceServices
 
diff --git a/Engine/Source/Developer/TraceServices/Private/Model/SymslibResolver.h b/Engine/Source/Developer/TraceServices/Private/Model/SymslibResolver.h
--- a/Engine/Source/Developer/TraceServices/Private/Model/SymslibResolver.h
+++ b/Engine/Source/Developer/TraceServices/Private/Model/SymslibResolver.h
@@ -23,10 +23,8 @@
 
 #include "symslib.h"
 
-/////////////////////////////////////////////////////////////////////
 namespace TraceServices {
 
-/////////////////////////////////////////////////////////////////////
 class FSymslibResolver
 {
 public:
@@ -35,7 +33,7 @@ public:
 	FSymslibResolver(IAnalysisSession& InSession, IResolvedSymbolFilter& InSymbolFilter);
 	~FSymslibResolver();
 	void QueueModuleLoad(const uint8* ImageId, uint32 ImageIdSize, FModule* Module);
-	void QueueModuleReload(const FModule* Module, const TCHAR* Path, TFunction<void(SymbolArray&)> ResolveOnSuccess);
+	void QueueModuleReload(FModule* Module, const TCHAR* Path, TFunction<void(SymbolArray&)> ResolveOnSuccess);
 	void QueueSymbolResolve(uint64 Address, FResolvedSymbol* Symbol);
 	void OnAnalysisComplete();
 	void GetStats(IModuleProvider::FStats* OutStats) const;
@@ -67,9 +65,9 @@ private:
 
 	struct FModuleEntry
 	{
-		FSymsInstance Instance;
-		TArray<uint8> ImageId;
 		FModule* Module;
+		TArray<uint8> ImageId;
+		FSymsInstance Instance;
 	};
 
 	struct FQueuedAddress
@@ -78,13 +76,6 @@ private:
 		FResolvedSymbol* Target;
 	};
 
-	struct FQueuedModule
-	{
-		uint64 Base;
-		uint64 Size;
-		const TCHAR* ImagePath;
-	};
-
 	enum : uint32 {
 		MaxNameLen = 512,
 		QueuedAddressLength = 2048,
@@ -101,13 +92,15 @@ private:
 	 */
 	void DispatchQueuedAddresses();
 	void ResolveSymbols(TArrayView<FQueuedAddress>& QueuedWork);
-	FModuleEntry* GetModuleForAddress(uint64 Address);
+	FModuleEntry* GetModuleEntry(FModule* Module) const;
+	FModuleEntry* GetModuleForAddress(uint64 Address) const;
 	static void UpdateResolvedSymbol(FResolvedSymbol& Symbol, ESymbolQueryResult Result, const TCHAR* Module, const TCHAR* Name, const TCHAR* File, uint16 Line);
 	void LoadModuleTracked(FModuleEntry* Entry, FStringView OverrideSearchPath);
 	EModuleStatus LoadModule(FModuleEntry* Entry, FStringView SearchPathView, FStringBuilderBase& OutStatusMessage) const;
 	void ResolveSymbolTracked(uint64 Address, FResolvedSymbol& Target, class FSymbolStringAllocator& StringAllocator);
 	bool ResolveSymbol(uint64 Address, FResolvedSymbol& Target, FSymbolStringAllocator& StringAllocator, FModuleEntry* Entry) const;
 	void WaitForTasks();
+	void UpdateSessionInfo() const;
 
 	mutable FRWLock ModulesLock;
 	TPagedArray<FModuleEntry> Modules;
@@ -119,8 +112,13 @@ private:
 	FGraphEventRef CleanupTask;
 	FGraphEventRef ModuleReloadTask;
 	std::atomic<bool> CancelTasks;
-	mutable FRWLock SymbolSearchPathsLock;
-	TArray<FString> SymbolSearchPaths;
+
+	mutable FRWLock CustomSymbolSearchPathsLock;
+	TArray<FString> CustomSymbolSearchPaths; // search paths added by user from UI in the current session
+	TArray<FString> ConfigSymbolSearchPaths; // search paths specified in environment variables and config files
+
+	mutable FString Platform;
+	mutable FString AppName;
 
 	std::atomic<uint32> ModulesDiscovered;
 	std::atomic<uint32> ModulesFailed;
@@ -130,12 +128,8 @@ private:
 
 	IAnalysisSession& Session;
 	IResolvedSymbolFilter& SymbolFilter;
-	FString Platform;
-	FString AppName;
 };
 
-/////////////////////////////////////////////////////////////////////
-
 } // namespace TraceServices
 
 #endif //UE_SYMSLIB_AVAILABLE
