[Insights] Memory Insights / Symbols Resolver: Updated code for DbgHelp resolver (several issues + updated code to use wchar API).

diff --git a/Engine/Source/Developer/TraceServices/Private/Model/DbgHelpResolver.cpp b/Engine/Source/Developer/TraceServices/Private/Model/DbgHelpResolver.cpp
--- a/Engine/Source/Developer/TraceServices/Private/Model/DbgHelpResolver.cpp
+++ b/Engine/Source/Developer/TraceServices/Private/Model/DbgHelpResolver.cpp
@@ -4,8 +4,8 @@
 
 #if PLATFORM_WINDOWS
 
-#include "Algo/Sort.h"
 #include "Algo/ForEach.h"
+#include "Algo/Sort.h"
 #include "Containers/Queue.h"
 #include "Containers/StringView.h"
 #include "HAL/PlatformProcess.h"
@@ -16,61 +16,91 @@
 #include "Misc/Guid.h"
 #include "Misc/PathViews.h"
 #include "Misc/ScopeLock.h"
+#include "Misc/ScopeRWLock.h"
 #include "Misc/StringBuilder.h"
 #include "TraceServices/Model/AnalysisSession.h"
 #include <atomic>
 
 #include "Windows/AllowWindowsPlatformTypes.h"
 THIRD_PARTY_INCLUDES_START
-#include <dbghelp.h>
+#include <DbgHelp.h>
 THIRD_PARTY_INCLUDES_END
 
-/////////////////////////////////////////////////////////////////////
+////////////////////////////////////////////////////////////////////////////////////////////////////
+
 DEFINE_LOG_CATEGORY_STATIC(LogDbgHelp, Log, All);
 
-/////////////////////////////////////////////////////////////////////
 namespace TraceServices {
 
-/////////////////////////////////////////////////////////////////////
-
 static const TCHAR* GUnknownModuleTextDbgHelp = TEXT("Unknown");
 
+////////////////////////////////////////////////////////////////////////////////////////////////////
 
-/////////////////////////////////////////////////////////////////////
 FDbgHelpResolver::FDbgHelpResolver(IAnalysisSession& InSession, IResolvedSymbolFilter& InSymbolFilter)
-	: bRunWorkerThread(false)
-	, bDrainThenStop(false)
+	: Modules(InSession.GetLinearAllocator(), 128)
 	, Session(InSession)
 	, SymbolFilter(InSymbolFilter)
+	, bRunWorkerThread(false)
+	, bDrainThenStop(false)
 {
-	// Setup search paths. The SearchPaths array is a priority stack, which
-	// means paths are searched in reversed order.
-	// 1. Any new paths entered by the user this session
+	// Setup search paths. Paths are searched in the following order:
+	// 1. Any new path entered by the user this session
 	// 2. Path of the executable (if available)
-	// 3. Paths from UE_INSIGHTS_SYMBOLPATH
-	// 4. Paths from the user configuration file
+	// 3. Paths from UE_INSIGHTS_SYMBOL_PATH
+	// 4. Paths from _NT_SYMBOL_PATH
+	// 5. Paths from the user configuration file
+
+	auto SplitEnvPaths = [](FStringView EnvVariable, TArray<FString>& OutList)
+	{
+		FString SymbolPathPart, SymbolPathRemainder(EnvVariable);
+		while (SymbolPathRemainder.Split(TEXT(";"), &SymbolPathPart, &SymbolPathRemainder))
+		{
+			OutList.Emplace(SymbolPathPart);
+		}
+		if (!SymbolPathRemainder.IsEmpty())
+		{
+			OutList.Emplace(SymbolPathRemainder);
+		}
+	};
+
+	// Paths from environment
+	const FString InsightsSymbolPath =  FPlatformMisc::GetEnvironmentVariable(TEXT("UE_INSIGHTS_SYMBOL_PATH"));
+	UE_LOG(LogDbgHelp, Log, TEXT("UE_INSIGHTS_SYMBOL_PATH: %s"), InsightsSymbolPath.IsEmpty() ? TEXT("Not set") : *InsightsSymbolPath);
+	SplitEnvPaths(InsightsSymbolPath, ConfigSymbolSearchPaths);
+
+#if PLATFORM_WINDOWS
+	const FString NTSymbolPath = FPlatformMisc::GetEnvironmentVariable(TEXT("_NT_SYMBOL_PATH"));
+	UE_LOG(LogDbgHelp, Log, TEXT("_NT_SYMBOL_PATH: %s"), NTSymbolPath.IsEmpty() ? TEXT("Not set") : *NTSymbolPath);
+	SplitEnvPaths(NTSymbolPath, ConfigSymbolSearchPaths);
+#endif
 
 	// Paths from configuration
 	FString SettingsIni;
-
 	if (FConfigContext::ReadIntoGConfig().Load(TEXT("UnrealInsightsSettings"), SettingsIni))
 	{
+		TArray<FString> SymbolSearchPaths;
 		GConfig->GetArray(TEXT("Insights.MemoryProfiler"), TEXT("SymbolSearchPaths"), SymbolSearchPaths, SettingsIni);
-	}
-
-	// Paths from environment
-	FString SymbolPathEnvVar =  FPlatformMisc::GetEnvironmentVariable(TEXT("UE_INSIGHTS_SYMBOL_PATH"));
-	UE_LOG(LogDbgHelp, Log, TEXT("UE_INSIGHTS_SYMBOL_PATH: '%s'"), *SymbolPathEnvVar);
-	FString SymbolPathPart;
-	while (SymbolPathEnvVar.Split(TEXT(";"), &SymbolPathPart, &SymbolPathEnvVar))
-	{
-		SymbolSearchPaths.Emplace(SymbolPathPart);
+#if !NO_LOGGING
+		if (SymbolSearchPaths.IsEmpty())
+		{
+			UE_LOG(LogDbgHelp, Log, TEXT("[Insights.MemoryProfiler] SymbolSearchPaths not set"));
+		}
+		else
+		{
+			for (const FString& Path : SymbolSearchPaths)
+			{
+				UE_LOG(LogDbgHelp, Log, TEXT("[Insights.MemoryProfiler] +SymbolSearchPaths=%s"), *Path);
+			}
+		}
+#endif
+		ConfigSymbolSearchPaths += SymbolSearchPaths;
 	}
 
 	Start();
 }
 
-/////////////////////////////////////////////////////////////////////
+////////////////////////////////////////////////////////////////////////////////////////////////////
+
 FDbgHelpResolver::~FDbgHelpResolver()
 {
 	bRunWorkerThread = false;
@@ -80,42 +110,65 @@ FDbgHelpResolver::~FDbgHelpResolver()
 	}
 }
 
-/////////////////////////////////////////////////////////////////////
+////////////////////////////////////////////////////////////////////////////////////////////////////
+
 void FDbgHelpResolver::QueueModuleLoad(const uint8* ImageId, uint32 ImageIdSize, FModule* Module)
 {
 	check(Module != nullptr);
+	ensure(GetModuleEntry(Module) == nullptr);
 
-	FScopeLock _(&ModulesCs);
+	FWriteScopeLock _(ModulesLock);
 
-	const FStringView ModuleName = FPathViews::GetCleanFilename(Module->FullName);
+	// Add the new module entry.
+	FModuleEntry* Entry = &Modules.PushBack();
+	Entry->Module = Module;
+	Entry->ImageId = TArrayView<const uint8>(ImageId, ImageIdSize);
 
-	// Add module and sort list according to base address
-	const int32 Index = LoadedModules.Add(FModuleEntry{
-		Module->Base, Module->Size, Session.StoreString(ModuleName), Session.StoreString(Module->FullName),
-		Module, TArray(ImageId, ImageIdSize)
-	});
+	// Sort list according to base address.
+	SortedModules.Add(Entry);
+	Algo::Sort(SortedModules, [](const FModuleEntry* Lhs, const FModuleEntry* Rhs) { return Lhs->Module->Base < Rhs->Module->Base; });
 
-	// Queue up module to have symbols loaded
-	LoadSymbolsQueue.Enqueue(FQueuedModule{ Module, nullptr, LoadedModules[Index].ImageId});
+	++ModulesDiscovered;
 
-	// Sort list according to base address
-	Algo::Sort(LoadedModules, [](const FModuleEntry& Lhs, const FModuleEntry& Rhs) { return Lhs.Base < Rhs.Base; });
+	// Set the Pending state before scheduling the background task (to allow calling code to wait, if needed).
+	Module->Status.store(EModuleStatus::Pending);
 
-	++ModulesDiscovered;
+	// Queue up module to have symbols loaded.
+	LoadSymbolsQueue.Enqueue(FQueuedModule{ Module, nullptr, Entry->ImageId });
 }
 
-/////////////////////////////////////////////////////////////////////
-void FDbgHelpResolver::QueueModuleReload(const FModule* Module, const TCHAR* InPath, TFunction<void(SymbolArray&)> ResolveOnSuccess)
+////////////////////////////////////////////////////////////////////////////////////////////////////
+
+void FDbgHelpResolver::QueueModuleReload(FModule* Module, const TCHAR* Path, TFunction<void(SymbolArray&)> ResolveOnSuccess)
 {
-	FScopeLock _(&ModulesCs);
-	const uint64 ModuleBase = Module->Base;
-	const FModuleEntry* Entry = LoadedModules.FindByPredicate([ModuleBase](const FModuleEntry& Entry) { return Entry.Base == ModuleBase; });
-	if (Entry)
+	check(Module != nullptr);
+
+	// Find the entry
+	FModuleEntry* Entry = GetModuleEntry(Module);
+	if (!Entry)
+	{
+		return;
+	}
+
+	// No use in trying reload already loaded modules
+	if (Module->Status == EModuleStatus::Loaded)
+	{
+		return;
+	}
+
+	// Set the Pending state before scheduling the background task (to allow calling code to wait, if needed).
+	EModuleStatus PreviousStatus = Module->Status.exchange(EModuleStatus::Pending);
+	if (PreviousStatus >= EModuleStatus::FailedStatusStart)
 	{
-		Entry->Module->Status.store(EModuleStatus::Pending);
-		LoadSymbolsQueue.Enqueue(FQueuedModule{Module, Session.StoreString(InPath), TArrayView<const uint8>(Entry->ImageId)});
+		--ModulesFailed;
 	}
 
+	FString PathStr(Path);
+	FPaths::NormalizeDirectoryName(PathStr);
+	const TCHAR* OverrideSearchPath = Session.StoreString(PathStr);
+
+	LoadSymbolsQueue.Enqueue(FQueuedModule{ Module, OverrideSearchPath, TArrayView<const uint8>(Entry->ImageId) });
+
 	SymbolArray SymbolsToResolve;
 	ResolveOnSuccess(SymbolsToResolve);
 	for(TTuple<uint64, FResolvedSymbol*> Pair : SymbolsToResolve)
@@ -130,36 +183,45 @@ void FDbgHelpResolver::QueueModuleReload(const FModule* Module, const TCHAR* InP
 	}
 }
 
-/////////////////////////////////////////////////////////////////////
+////////////////////////////////////////////////////////////////////////////////////////////////////
+
 void FDbgHelpResolver::QueueSymbolResolve(uint64 Address, FResolvedSymbol* Symbol)
 {
-	ResolveQueue.Enqueue(FQueuedAddress{Address, Symbol});
+	ResolveQueue.Enqueue(FQueuedAddress{ Address, Symbol });
 }
 
-/////////////////////////////////////////////////////////////////////
+////////////////////////////////////////////////////////////////////////////////////////////////////
+
 void FDbgHelpResolver::GetStats(IModuleProvider::FStats* OutStats) const
 {
-	FScopeLock _(&ModulesCs);
+	FReadScopeLock _(ModulesLock);
 	FMemory::Memzero(*OutStats);
-	for(const FModuleEntry& Entry : LoadedModules)
+	for (uint32 ModuleIndex = 0; ModuleIndex < Modules.Num(); ++ModuleIndex)
 	{
-		OutStats->SymbolsDiscovered += Entry.Module->Stats.Discovered.load();
-		OutStats->SymbolsResolved += Entry.Module->Stats.Resolved.load();
-		OutStats->SymbolsFailed += Entry.Module->Stats.Failed.load();
+		FModule::SymbolStats& ModuleStats = Modules[ModuleIndex].Module->Stats;
+		OutStats->SymbolsDiscovered += ModuleStats.Discovered.load();
+		OutStats->SymbolsResolved += ModuleStats.Resolved.load();
+		OutStats->SymbolsFailed += ModuleStats.Failed.load();
 	}
 	OutStats->ModulesDiscovered = ModulesDiscovered.load();
 	OutStats->ModulesFailed = ModulesFailed.load();
 	OutStats->ModulesLoaded = ModulesLoaded.load();
 }
 
-/////////////////////////////////////////////////////////////////////
+////////////////////////////////////////////////////////////////////////////////////////////////////
+
 void FDbgHelpResolver::EnumerateSymbolSearchPaths(TFunctionRef<void(FStringView Path)> Callback) const
 {
-	FScopeLock _(&SymbolSearchPathsLock);
-	Algo::ForEach(SymbolSearchPaths, Callback);
+	{
+		FReadScopeLock _(CustomSymbolSearchPathsLock);
+		Algo::ForEach(CustomSymbolSearchPaths, Callback);
+	}
+
+	Algo::ForEach(ConfigSymbolSearchPaths, Callback);
 }
 
-/////////////////////////////////////////////////////////////////////
+////////////////////////////////////////////////////////////////////////////////////////////////////
+
 void FDbgHelpResolver::OnAnalysisComplete()
 {
 	// At this point no more module loads or symbol requests will be queued,
@@ -167,8 +229,8 @@ void FDbgHelpResolver::OnAnalysisComplete()
 	bDrainThenStop = true;
 }
 
+////////////////////////////////////////////////////////////////////////////////////////////////////
 
-/////////////////////////////////////////////////////////////////////
 bool FDbgHelpResolver::SetupSyms()
 {
 	// Create a unique handle
@@ -189,16 +251,16 @@ bool FDbgHelpResolver::SetupSyms()
 	return SymInitialize((HANDLE)Handle, NULL, FALSE);
 }
 
+////////////////////////////////////////////////////////////////////////////////////////////////////
 
-/////////////////////////////////////////////////////////////////////
 void FDbgHelpResolver::FreeSyms() const
 {
 	// This release file locks on debug files
 	SymCleanup((HANDLE)Handle);
 }
 
+////////////////////////////////////////////////////////////////////////////////////////////////////
 
-/////////////////////////////////////////////////////////////////////
 uint32 FDbgHelpResolver::Run()
 {
 	const bool bInitialized = SetupSyms();
@@ -240,7 +302,8 @@ uint32 FDbgHelpResolver::Run()
 	return 0;
 }
 
-/////////////////////////////////////////////////////////////////////
+////////////////////////////////////////////////////////////////////////////////////////////////////
+
 void FDbgHelpResolver::Start()
 {
 	// Start the worker thread
@@ -248,13 +311,43 @@ void FDbgHelpResolver::Start()
 	Thread = FRunnableThread::Create(this, TEXT("DbgHelpWorker"), 0, TPri_Normal);
 }
 
-/////////////////////////////////////////////////////////////////////
+////////////////////////////////////////////////////////////////////////////////////////////////////
+
 void FDbgHelpResolver::Stop()
 {
 	bRunWorkerThread = false;
 }
 
-/////////////////////////////////////////////////////////////////////
+////////////////////////////////////////////////////////////////////////////////////////////////////
+
+FDbgHelpResolver::FModuleEntry* FDbgHelpResolver::GetModuleEntry(FModule* Module) const
+{
+	FReadScopeLock _(ModulesLock);
+	for (FModuleEntry* Entry : SortedModules)
+	{
+		if (Entry->Module == Module)
+		{
+			return Entry;
+		}
+	}
+	return nullptr;
+}
+
+////////////////////////////////////////////////////////////////////////////////////////////////////
+
+FDbgHelpResolver::FModuleEntry* FDbgHelpResolver::GetModuleForAddress(uint64 Address) const
+{
+	FReadScopeLock _(ModulesLock);
+	const int32 EntryIdx = Algo::LowerBoundBy(SortedModules, Address, [](const FModuleEntry* Entry) { return Entry->Module->Base; }) - 1;
+	if (EntryIdx < 0 || EntryIdx >= SortedModules.Num())
+	{
+		return nullptr;
+	}
+	return SortedModules[EntryIdx];
+}
+
+////////////////////////////////////////////////////////////////////////////////////////////////////
+
 void FDbgHelpResolver::UpdateResolvedSymbol(FResolvedSymbol& Symbol, ESymbolQueryResult Result, const TCHAR* Module, const TCHAR* Name, const TCHAR* File, uint16 Line)
 {
 	Symbol.Module = Module;
@@ -264,7 +357,8 @@ void FDbgHelpResolver::UpdateResolvedSymbol(FResolvedSymbol& Symbol, ESymbolQuer
 	Symbol.Result.store(Result, std::memory_order_release);
 }
 
-/////////////////////////////////////////////////////////////////////
+////////////////////////////////////////////////////////////////////////////////////////////////////
+
 void FDbgHelpResolver::ResolveSymbol(uint64 Address, FResolvedSymbol& Target)
 {
 	if (Target.Result.load() == ESymbolQueryResult::OK)
@@ -272,9 +366,10 @@ void FDbgHelpResolver::ResolveSymbol(uint64 Address, FResolvedSymbol& Target)
 		return;
 	}
 
-	const FModuleEntry* Entry = GetModuleForAddress(Address);
+	FModuleEntry* Entry = GetModuleForAddress(Address);
 	if (!Entry)
 	{
+		UE_LOG(LogDbgHelp, Warning, TEXT("No module mapped to address 0x%016llx."), Address);
 		UpdateResolvedSymbol(Target,
 			ESymbolQueryResult::NotFound,
 			GUnknownModuleTextDbgHelp,
@@ -284,16 +379,16 @@ void FDbgHelpResolver::ResolveSymbol(uint64 Address, FResolvedSymbol& Target)
 		SymbolFilter.Update(Target);
 		return;
 	}
+	FModule* Module = Entry->Module;
+
+	++Module->Stats.Discovered;
 
-	const EModuleStatus ModuleStatus = Entry->Module->Status.load();
+	const EModuleStatus ModuleStatus = Module->Status.load();
 	if (ModuleStatus != EModuleStatus::Loaded)
 	{
-		const ESymbolQueryResult Result = ModuleStatus == EModuleStatus::VersionMismatch
-			                                  ? ESymbolQueryResult::Mismatch
-			                                  : ESymbolQueryResult::NotLoaded;
-		++Entry->Module->Stats.Failed;
+		++Module->Stats.Failed;
 		UpdateResolvedSymbol(Target,
-			Result,
+			ModuleStatus == EModuleStatus::VersionMismatch ? ESymbolQueryResult::Mismatch : ESymbolQueryResult::NotLoaded,
 			GUnknownModuleTextDbgHelp,
 			GUnknownModuleTextDbgHelp,
 			GUnknownModuleTextDbgHelp,
@@ -302,20 +397,18 @@ void FDbgHelpResolver::ResolveSymbol(uint64 Address, FResolvedSymbol& Target)
 		return;
 	}
 
-	++Entry->Module->Stats.Discovered;
-
-	uint8 InfoBuffer[sizeof(SYMBOL_INFO) + (MaxNameLen * sizeof(char) + 1)];
-	SYMBOL_INFO* Info = (SYMBOL_INFO*)InfoBuffer;
-	Info->SizeOfStruct = sizeof(SYMBOL_INFO);
+	uint8 InfoBuffer[sizeof(SYMBOL_INFOW) + (MaxNameLen + 1) * sizeof(TCHAR)];
+	SYMBOL_INFOW* Info = (SYMBOL_INFOW*)InfoBuffer;
+	Info->SizeOfStruct = sizeof(SYMBOL_INFOW);
 	Info->MaxNameLen = MaxNameLen;
 
 	// Find and build the symbol name
-	if (!SymFromAddr((HANDLE)Handle, Address, NULL, Info))
+	if (!SymFromAddrW((HANDLE)Handle, Address, NULL, Info))
 	{
-		++Entry->Module->Stats.Failed;
+		++Module->Stats.Failed;
 		UpdateResolvedSymbol(Target,
 			ESymbolQueryResult::NotFound,
-			Entry->Name,
+			Module->Name,
 			GUnknownModuleTextDbgHelp,
 			GUnknownModuleTextDbgHelp,
 			0);
@@ -323,19 +416,19 @@ void FDbgHelpResolver::ResolveSymbol(uint64 Address, FResolvedSymbol& Target)
 		return;
 	}
 
-	const TCHAR* SymbolNameStr = Session.StoreString(ANSI_TO_TCHAR(Info->Name));
+	const TCHAR* SymbolNameStr = Session.StoreString(Info->Name);
 
 	// Find the source file and line
 	DWORD  dwDisplacement;
-	IMAGEHLP_LINE Line;
-	Line.SizeOfStruct = sizeof(IMAGEHLP_LINE);
+	IMAGEHLP_LINEW64 Line;
+	Line.SizeOfStruct = sizeof(IMAGEHLP_LINEW64);
 
-	if (!SymGetLineFromAddr((HANDLE)Handle, Address, &dwDisplacement, &Line))
+	if (!SymGetLineFromAddrW64((HANDLE)Handle, Address, &dwDisplacement, &Line))
 	{
-		++Entry->Module->Stats.Failed;
+		++Module->Stats.Failed;
 		UpdateResolvedSymbol(Target,
 			ESymbolQueryResult::OK,
-			Entry->Name,
+			Module->Name,
 			SymbolNameStr,
 			GUnknownModuleTextDbgHelp,
 			0);
@@ -343,20 +436,21 @@ void FDbgHelpResolver::ResolveSymbol(uint64 Address, FResolvedSymbol& Target)
 		return;
 	}
 
-	const TCHAR* SymbolFileStr = Session.StoreString(ANSI_TO_TCHAR(Line.FileName));
+	const TCHAR* SymbolFileStr = Session.StoreString(Line.FileName);
 
-	++Entry->Module->Stats.Resolved;
+	++Module->Stats.Resolved;
 	UpdateResolvedSymbol(Target,
 		ESymbolQueryResult::OK,
-		Entry->Name,
+		Module->Name,
 		SymbolNameStr,
 		SymbolFileStr,
 		static_cast<uint16>(Line.LineNumber));
 	SymbolFilter.Update(Target);
 }
 
-/////////////////////////////////////////////////////////////////////
-void FDbgHelpResolver::LoadModuleSymbols(const FModule* Module, const TCHAR* Path, const TArrayView<const uint8> ImageId)
+////////////////////////////////////////////////////////////////////////////////////////////////////
+
+void FDbgHelpResolver::LoadModuleSymbols(FModule* Module, const TCHAR* OverrideSearchPath, const TArrayView<const uint8> ImageId)
 {
 	check(Module);
 
@@ -365,13 +459,55 @@ void FDbgHelpResolver::LoadModuleSymbols(const FModule* Module, const TCHAR* Pat
 
 	// Setup symbol search path
 	{
-		TAnsiStringBuilder<1024> UserSearchPath;
-		UserSearchPath << Path << ";";
-		Algo::ForEach(SymbolSearchPaths, [&UserSearchPath] (const FString& Path){ UserSearchPath.Appendf("%s;", TCHAR_TO_ANSI(*Path));});
+		TStringBuilder<1024> UserSearchPath;
+
+		if (OverrideSearchPath && *OverrideSearchPath != TEXT('\0'))
+		{
+			UserSearchPath.Append(OverrideSearchPath);
+			UserSearchPath.AppendChar(TEXT(';'));
+		}
+		else
+		{
+			// 1. Any new path entered by the user this session
+			{
+				FReadScopeLock _(CustomSymbolSearchPathsLock);
+				Algo::ForEach(CustomSymbolSearchPaths,
+					[&UserSearchPath](const FString& Path)
+					{
+						if (!Path.IsEmpty())
+						{
+							UserSearchPath.Append(Path);
+							UserSearchPath.AppendChar(TEXT(';'));
+						}
+					});
+			}
+
+			// 2. Path of the executable (if available)
+			FString ModuleNamePath = FPaths::GetPath(Module->FullName);
+			FPaths::NormalizeDirectoryName(ModuleNamePath);
+			if (!ModuleNamePath.IsEmpty())
+			{
+				UserSearchPath.Append(ModuleNamePath);
+				UserSearchPath.AppendChar(TEXT(';'));
+			}
+
+			// 3. Paths from UE_INSIGHTS_SYMBOL_PATH
+			// 4. Paths from _NT_SYMBOL_PATH
+			// 5. Paths from the user configuration file
+			Algo::ForEach(ConfigSymbolSearchPaths,
+				[&UserSearchPath](const FString& Path)
+				{
+					if (!Path.IsEmpty())
+					{
+						UserSearchPath.Append(Path);
+						UserSearchPath.AppendChar(TEXT(';'));
+					}
+				});
+		}
 
-		if (!SymSetSearchPath((HANDLE)Handle, UserSearchPath.ToString()))
+		if (!SymSetSearchPathW((HANDLE)Handle, UserSearchPath.ToString()))
 		{
-			UE_LOG(LogDbgHelp, Warning, TEXT("Unable to set symbol search path to '%hs'."), UserSearchPath.ToString());
+			UE_LOG(LogDbgHelp, Warning, TEXT("Unable to set symbol search path to '%s'."), UserSearchPath.ToString());
 		}
 		TCHAR OutPath[1024];
 		SymGetSearchPathW((HANDLE) Handle, OutPath, 1024);
@@ -379,16 +515,16 @@ void FDbgHelpResolver::LoadModuleSymbols(const FModule* Module, const TCHAR* Pat
 	}
 
 	// Attempt to load symbols
-	const DWORD64 LoadedBaseAddress = SymLoadModuleEx((HANDLE)Handle, NULL, TCHAR_TO_ANSI(Module->Name), NULL, Base, Size, NULL, 0);
+	const DWORD64 LoadedBaseAddress = SymLoadModuleExW((HANDLE)Handle, NULL, Module->Name, NULL, Base, Size, NULL, 0);
 	const bool bModuleLoaded = Base == LoadedBaseAddress;
 	bool bPdbLoaded = true;
 	bool bPdbMatchesImage = true;
-	IMAGEHLP_MODULE ModuleInfo;
+	IMAGEHLP_MODULEW ModuleInfo;
 
 	if (bModuleLoaded)
 	{
-		ModuleInfo.SizeOfStruct = sizeof(IMAGEHLP_MODULE);
-		SymGetModuleInfo((HANDLE)Handle, Base, &ModuleInfo);
+		ModuleInfo.SizeOfStruct = sizeof(IMAGEHLP_MODULEW);
+		SymGetModuleInfoW((HANDLE)Handle, Base, &ModuleInfo);
 
 		if (ModuleInfo.SymType != SymPdb)
 		{
@@ -413,7 +549,7 @@ void FDbgHelpResolver::LoadModuleSymbols(const FModule* Module, const TCHAR* Pat
 		// Unload the module, otherwise any subsequent attempts to load module with another
 		// path will fail.
 		SymUnloadModule((HANDLE)Handle, Base);
-		StatusMessage.Appendf(TEXT("Unable to load symbols for %s"), Path);
+		StatusMessage.Appendf(TEXT("Unable to load symbols for %s"), Module->Name);
 		Status = EModuleStatus::Failed;
 		++ModulesFailed;
 	}
@@ -422,47 +558,26 @@ void FDbgHelpResolver::LoadModuleSymbols(const FModule* Module, const TCHAR* Pat
 		// Unload the module, otherwise any subsequent attempts to load module with another
 		// path will fail.
 		SymUnloadModule((HANDLE)Handle, Base);
-		StatusMessage.Appendf(TEXT("Unable to load symbols for %s, pdb signature does not match."), Path);
+		StatusMessage.Appendf(TEXT("Unable to load symbols for %s, pdb signature does not match."), Module->Name);
 		Status = EModuleStatus::VersionMismatch;
 		++ModulesFailed;
 	}
 	else
 	{
-		StatusMessage.Appendf(TEXT("Loaded symbols for %s from %s."), Path, ANSI_TO_TCHAR(ModuleInfo.LoadedImageName));
+		StatusMessage.Appendf(TEXT("Loaded symbols for %s from %s."), Module->Name, ModuleInfo.LoadedPdbName);
 		Status = EModuleStatus::Loaded;
 		++ModulesLoaded;
 	}
 
-	// Update the module entry with the result
-	FScopeLock _(&ModulesCs);
-	const int32 EntryIdx = Algo::BinarySearchBy(LoadedModules, Base, [](const FModuleEntry& Entry) { return Entry.Base; });
-	check(EntryIdx != INDEX_NONE);
-	const FModuleEntry& Entry = LoadedModules[EntryIdx];
-
 	// Make the status visible to the world
-	Entry.Module->StatusMessage = Session.StoreString(StatusMessage.ToView());
-	Entry.Module->Status.store(Status);
+	Module->StatusMessage = Session.StoreString(StatusMessage.ToView());
+	Module->Status.store(Status);
 }
 
+////////////////////////////////////////////////////////////////////////////////////////////////////
 
-/////////////////////////////////////////////////////////////////////
-const FDbgHelpResolver::FModuleEntry* FDbgHelpResolver::GetModuleForAddress(uint64 Address) const
-{
-	const int32 EntryIdx = Algo::LowerBoundBy(LoadedModules, Address, [](const FModuleEntry& Entry) { return Entry.Base; }) - 1;
-	if (EntryIdx < 0 || EntryIdx >= LoadedModules.Num())
-	{
-		return nullptr;
-	}
-
-	return &LoadedModules[EntryIdx];
-}
-
+} // namespace TraceServices
 
-/////////////////////////////////////////////////////////////////////
 #include "Windows/HideWindowsPlatformTypes.h"
 
-/////////////////////////////////////////////////////////////////////
-
-} // namespace TraceServices
-
 #endif // PLATFORM_WINDOWS
diff --git a/Engine/Source/Developer/TraceServices/Private/Model/DbgHelpResolver.h b/Engine/Source/Developer/TraceServices/Private/Model/DbgHelpResolver.h
--- a/Engine/Source/Developer/TraceServices/Private/Model/DbgHelpResolver.h
+++ b/Engine/Source/Developer/TraceServices/Private/Model/DbgHelpResolver.h
@@ -5,6 +5,7 @@
 
 #if PLATFORM_WINDOWS
 
+#include "Common/PagedArray.h"
 #include "Containers/Array.h"
 #include "Containers/Queue.h"
 #include "HAL/CriticalSection.h"
@@ -15,32 +16,26 @@
 #include "TraceServices/Model/Modules.h"
 #include <atomic>
 
-/////////////////////////////////////////////////////////////////////
 namespace TraceServices {
 
-/////////////////////////////////////////////////////////////////////
 class FDbgHelpResolver : public FRunnable
 {
 public:
 	typedef TArray<TTuple<uint64, FResolvedSymbol*>> SymbolArray;
 
-	void Start();
 	FDbgHelpResolver(IAnalysisSession& InSession, IResolvedSymbolFilter& InSymbolFilter);
 	~FDbgHelpResolver();
+	void Start();
 	void QueueModuleLoad(const uint8* ImageId, uint32 ImageIdSize, FModule* Module);
-	void QueueModuleReload(const FModule* Module, const TCHAR* InPath, TFunction<void(SymbolArray&)> ResolveOnSuccess);
+	void QueueModuleReload(FModule* Module, const TCHAR* Path, TFunction<void(SymbolArray&)> ResolveOnSuccess);
 	void QueueSymbolResolve(uint64 Address, FResolvedSymbol* Symbol);
+	void OnAnalysisComplete();
 	void GetStats(IModuleProvider::FStats* OutStats) const;
 	void EnumerateSymbolSearchPaths(TFunctionRef<void(FStringView Path)> Callback) const;
-	void OnAnalysisComplete();
 
 private:
 	struct FModuleEntry
 	{
-		uint64 Base;
-		uint32 Size;
-		const TCHAR* Name;
-		const TCHAR* Path;
 		FModule* Module;
 		TArray<uint8> ImageId;
 	};
@@ -53,7 +48,7 @@ private:
 
 	struct FQueuedModule
 	{
-		const FModule* Module;
+		FModule* Module;
 		const TCHAR* Path;
 		TArrayView<const uint8> ImageId;
 	};
@@ -66,33 +61,35 @@ private:
 	void FreeSyms() const;
 	virtual uint32 Run() override;
 	virtual void Stop() override;
+	FModuleEntry* GetModuleEntry(FModule* Module) const;
+	FModuleEntry* GetModuleForAddress(uint64 Address) const;
 	static void UpdateResolvedSymbol(FResolvedSymbol& Symbol, ESymbolQueryResult Result, const TCHAR* Module, const TCHAR* Name, const TCHAR* File, uint16 Line);
-
 	void ResolveSymbol(uint64 Address, FResolvedSymbol& Target);
-	void LoadModuleSymbols(const FModule* Module, const TCHAR* Path, const TArrayView<const uint8> ImageId);
-	const FModuleEntry* GetModuleForAddress(uint64 Address) const;
+	void LoadModuleSymbols(FModule* Module, const TCHAR* Path, const TArrayView<const uint8> ImageId);
 
-	mutable FCriticalSection ModulesCs;
-	TArray<FModuleEntry> LoadedModules;
+	mutable FRWLock ModulesLock;
+	TPagedArray<FModuleEntry> Modules;
+	TArray<FModuleEntry*> SortedModules;
 	TQueue<FQueuedModule, EQueueMode::Mpsc> LoadSymbolsQueue;
 	TQueue<FQueuedAddress, EQueueMode::Mpsc> ResolveQueue;
 
+	mutable FRWLock CustomSymbolSearchPathsLock;
+	TArray<FString> CustomSymbolSearchPaths; // search paths added by user from UI in the current session
+	TArray<FString> ConfigSymbolSearchPaths; // search paths specified in environment variables and config files
+
 	std::atomic<uint32> ModulesDiscovered;
 	std::atomic<uint32> ModulesFailed;
 	std::atomic<uint32> ModulesLoaded;
 
-	mutable FCriticalSection SymbolSearchPathsLock;
-	TArray<FString> SymbolSearchPaths;
+	IAnalysisSession& Session;
+	IResolvedSymbolFilter& SymbolFilter;
+
 	bool bRunWorkerThread;
 	bool bDrainThenStop;
 	UPTRINT Handle;
-	IAnalysisSession& Session;
-	IResolvedSymbolFilter& SymbolFilter;
 	FRunnableThread* Thread;
 };
 
-/////////////////////////////////////////////////////////////////////
-
 } // namespace TraceServices
 
 #endif // PLATFORM_WINDOWS
