[Insights]
  - TraceServices: Added LowerBoundBy and UpperBoundBy algorithms in TPagedArray (in TraceServices::PagedArrayAlgo:: namespace).
  - TraceServices: Added BinarySearchClosestBy algorithm in TPagedArray. Performs binary search, resulting in position of the first element with projected value closest to Value.

diff --git a/Engine/Source/Developer/TraceServices/Public/Common/PagedArray.h b/Engine/Source/Developer/TraceServices/Public/Common/PagedArray.h
--- a/Engine/Source/Developer/TraceServices/Public/Common/PagedArray.h
+++ b/Engine/Source/Developer/TraceServices/Public/Common/PagedArray.h
@@ -7,6 +7,9 @@
 #include "Containers/ArrayView.h"
 #include "Math/Interval.h"
 #include "Templates/Function.h"
+#include "Templates/IsArithmetic.h"
+#include "Templates/Less.h"
+#include "Templates/Requires.h"
 
 #ifndef TRACESERVICES_PAGED_ARRAY_ITERATOR_V2
 #define TRACESERVICES_PAGED_ARRAY_ITERATOR_V2 0 // enables a simpler implementation of TPagedArrayIterator, for debug purposes
@@ -760,13 +763,13 @@ inline SIZE_T GetNum(const TPagedArray<ItemType, PageType>& PagedArray)
 }
 
 /**
- *	Use binary search to find the first and last element inside a TPagedArray that overlaps a given input interval.
- *	This requires the elements in the array to be sorted by the value returned from the Projection.
- *	Example usage for Timeline events would require a projection that returns Item.StartTime and the resulting range.Min
- *	will point to the last element where Item.End > StartTime and range.Max to the last element where
- *	Item.StartTime < EndTime.
+ * Use binary search to find the first and last element inside a TPagedArray that overlaps a given input interval.
+ * This requires the elements in the array to be sorted by the value returned from the Projection.
+ * Example usage for Timeline events would require a projection that returns Item.StartTime and the resulting range.Min
+ * will point to the last element where Item.End > StartTime and range.Max to the last element where
+ * Item.StartTime < EndTime.
  */
-template<typename ItemType, typename PageType>	
+template<typename ItemType, typename PageType>
 FInt32Interval GetElementRangeOverlappingGivenRange(const TPagedArray<ItemType, PageType>& PagedArray,
 	double StartTime, double EndTime,
 	TFunctionRef<double(const ItemType&)> ItemStartProjection,
@@ -781,7 +784,7 @@ FInt32Interval GetElementRangeOverlappingGivenRange(const TPagedArray<ItemType,
 	{
 		return Result;
 	}
-	
+
 	const int32 NumPoints = static_cast<int32>(PagedArray.Num());
 	const int32 NumPages = static_cast<int32>(PagedArray.NumPages());
 	const int32 PageSize = static_cast<int32>(PagedArray.GetPageSize());
@@ -790,7 +793,7 @@ FInt32Interval GetElementRangeOverlappingGivenRange(const TPagedArray<ItemType,
 	{
 		return Result;
 	}
-	
+
 	TArrayView<const PageType, int32> Pages = MakeArrayView(PageData, NumPages);
 
 	// find the page before the first page that's already inside the range (thus -1)
@@ -803,7 +806,7 @@ FInt32Interval GetElementRangeOverlappingGivenRange(const TPagedArray<ItemType,
 		Result.Min = 0;
 		StartPageIndex = 0;
 	}
-	else 
+	else
 	{
 		// if we went past the end with StartPageIndex we still have to search the last page
 		const PageType& Page = PageData[StartPageIndex];
@@ -816,11 +819,11 @@ FInt32Interval GetElementRangeOverlappingGivenRange(const TPagedArray<ItemType,
 		{
 			return Result;
 		}
-	
+
 		Result.Min = StartPageIndex * PageSize + Index;
 		check(Index <= NumPoints);
 	}
-	
+
 	TArrayView<const PageType, int32> RemainingPages = MakeArrayView(&PageData[StartPageIndex], NumPages - StartPageIndex);
 
 	int32 EndPageIndex = Algo::UpperBoundBy(RemainingPages, EndTime,
@@ -832,7 +835,7 @@ FInt32Interval GetElementRangeOverlappingGivenRange(const TPagedArray<ItemType,
 	check(EndPageIndex > 0)
 	// find the page before the first page that outside the range (thus -1)
 	EndPageIndex -= 1;
-	
+
 	{
 		const PageType& Page = PageData[EndPageIndex];
 		TArrayView<ItemType, int32> PageValues(Page.Items, static_cast<int32>(Page.Count));
@@ -853,4 +856,282 @@ FInt32Interval GetElementRangeOverlappingGivenRange(const TPagedArray<ItemType,
 	return Result;
 }
 
+namespace PagedArrayAlgoImpl
+{
+
+/**
+ * Performs binary search, resulting in position of the first element with projected value >= Value.
+ *
+ * @param PagedArray The paged array to search through; must be already sorted by SortPredicate.
+ * @param Value The value to look for
+ * @param Projection The functor or data member pointer; called via Invoke to compare to Value.
+ * @param SortPredicate The predicate for sort comparison; defaults to <.
+ *
+ * @returns The position of the first element with projected value >= Value; may be == Num.
+ */
+template<typename SizeType, typename ItemType, typename PageType, typename ValueType, typename ProjectionType, typename SortPredicateType = TLess<>()>
+FORCEINLINE SizeType LowerBoundInternal(const TPagedArray<ItemType, PageType>& PagedArray, const ValueType& Value, ProjectionType Projection, SortPredicateType SortPredicate)
+{
+	TArrayView<const PageType, int32> Pages = MakeArrayView(PagedArray.GetPages(), (int32)PagedArray.NumPages());
+
+	// Find the first page with projected value of the first item >= searched Value.
+	int32 PageIndex = Algo::LowerBoundBy(Pages, Value, [&Projection](const PageType& Page) { return Invoke(Projection, Page.Items[0]); }, SortPredicate);
+	if (PageIndex == 0)
+	{
+		return 0;
+	}
+
+	// Look backward to previous items (with projected value >= searched Value) to find the first one.
+	SizeType ElementIndex = PageIndex * PagedArray.GetPageSize();
+	if (ElementIndex > PagedArray.Num())
+	{
+		ElementIndex = PagedArray.Num();
+	}
+	check(ElementIndex > 0);
+	for (auto It = PagedArray.GetIteratorFromItem(ElementIndex - 1); It; It.PrevItem())
+	{
+		auto&& CheckValue = Invoke(Projection, *(It.GetCurrentItem()));
+		if (SortPredicate(CheckValue, Value))
+		{
+			break;
+		}
+		ElementIndex = It.GetCurrentItemIndex();
+	}
+	return ElementIndex;
+}
+
+/**
+ * Performs binary search, resulting in position of the first element with projected value > Value.
+ *
+ * @param PagedArray The paged array to search through; must be already sorted by SortPredicate.
+ * @param Value The value to look for
+ * @param Projection The functor or data member pointer; called via Invoke to compare to Value.
+ * @param SortPredicate The predicate for sort comparison; defaults to <.
+ *
+ * @returns The position of the first element with projected value > Value; may be == Num.
+ */
+template<typename SizeType, typename ItemType, typename PageType, typename ValueType, typename ProjectionType, typename SortPredicateType = TLess<>()>
+FORCEINLINE SizeType UpperBoundInternal(const TPagedArray<ItemType, PageType>& PagedArray, const ValueType& Value, ProjectionType Projection, SortPredicateType SortPredicate)
+{
+	TArrayView<const PageType, int32> Pages = MakeArrayView(PagedArray.GetPages(), (int32)PagedArray.NumPages());
+
+	// Find the first page with projected value of the first item > searched Value.
+	int32 PageIndex = Algo::UpperBoundBy(Pages, Value, [&Projection](const PageType& Page) { return Invoke(Projection, Page.Items[0]); }, SortPredicate);
+	if (PageIndex == 0)
+	{
+		return 0;
+	}
+
+	// Look backward to previous items (with projected value > searched Value) to find the first one.
+	SizeType ElementIndex = PageIndex * PagedArray.GetPageSize();
+	if (ElementIndex > PagedArray.Num())
+	{
+		ElementIndex = PagedArray.Num();
+	}
+	check(ElementIndex > 0);
+	for (auto It = PagedArray.GetIteratorFromItem(ElementIndex - 1); It; It.PrevItem())
+	{
+		auto&& CheckValue = Invoke(Projection, *(It.GetCurrentItem()));
+		if (!SortPredicate(Value, CheckValue))
+		{
+			break;
+		}
+		ElementIndex = It.GetCurrentItemIndex();
+	}
+	return ElementIndex;
+}
+
+/**
+ * Performs binary search, resulting in position of the first element with projected value closest to Value.
+ *
+ * @param PagedArray The paged array to search through; must be already sorted by SortPredicate.
+ * @param Value The value to look for
+ * @param Projection The functor or data member pointer; called via Invoke to compare to Value.
+ * @param SortPredicate The predicate for sort comparison; defaults to <.
+ *
+ * @returns The position of the first element with projected value closest to Value; == 0 if Num == 0, otherwise is a value in range [0 .. Num-1].
+ */
+template<typename SizeType, typename ItemType, typename PageType, typename ValueType, typename ProjectionType, typename SortPredicateType = TLess<>()
+	UE_REQUIRES(TIsArithmetic<ValueType>::Value)>
+FORCEINLINE SizeType BinarySearchClosestInternal(const TPagedArray<ItemType, PageType>& PagedArray, ValueType Value, ProjectionType Projection, SortPredicateType SortPredicate)
+{
+	if (PagedArray.Num() <= 1)
+	{
+		return 0;
+	}
+
+	// Find the first page with projected value of the first item >= searched Value.
+	TArrayView<const PageType, int32> Pages = MakeArrayView(PagedArray.GetPages(), (int32)PagedArray.NumPages());
+	int32 StartPageIndex = Algo::LowerBoundBy(Pages, Value, [&Projection](const PageType& Page) { return Invoke(Projection, Page.Items[0]); }, SortPredicate);
+	if (StartPageIndex == 0)
+	{
+		return 0;
+	}
+
+	// Iterate backward starting from this item (first item of the found page).
+	SizeType ClosestElementIndex = StartPageIndex * PagedArray.GetPageSize();
+	if (ClosestElementIndex >= PagedArray.Num())
+	{
+		ClosestElementIndex = PagedArray.Num() - 1;
+	}
+	if (ClosestElementIndex > 0)
+	{
+		auto&& ClosestElementValue = Invoke(Projection, PagedArray[ClosestElementIndex]);
+		ValueType ClosestTimeDelta = FMath::Abs(ClosestElementValue - Value);
+		for (auto It = PagedArray.GetIteratorFromItem(ClosestElementIndex - 1); It; It.PrevItem())
+		{
+			auto&& CheckValue = Invoke(Projection, *(It.GetCurrentItem()));
+			ValueType Delta = FMath::Abs(CheckValue - Value);
+			if (Delta > ClosestTimeDelta)
+			{
+				break;
+			}
+			ClosestTimeDelta = Delta;
+			ClosestElementIndex = It.GetCurrentItemIndex();
+		}
+	}
+	return ClosestElementIndex;
+}
+
+} // namespace PagedArrayAlgoImpl
+
+namespace PagedArrayAlgo
+{
+
+/**
+ * Performs binary search, resulting in position of the first element >= Value.
+ *
+ * @param PagedArray The paged array to search through; must be already sorted by SortPredicate.
+ * @param Value The value to look for
+ * @param SortPredicate Predicate for sort comparison; defaults to <.
+ *
+ * @returns Position of the first element >= Value; may be == Num.
+ */
+template <typename ItemType, typename PageType, typename SortPredicateType>
+FORCEINLINE auto LowerBound(const TPagedArray<ItemType, PageType>& PagedArray, const ItemType& Value, SortPredicateType SortPredicate) -> decltype(GetNum(PagedArray))
+{
+	using SizeType = decltype(GetNum(PagedArray));
+	return PagedArrayAlgoImpl::LowerBoundInternal<SizeType>(PagedArray, Value, FIdentityFunctor(), SortPredicate);
+}
+template <typename ItemType, typename PageType>
+FORCEINLINE auto LowerBound(const TPagedArray<ItemType, PageType>& PagedArray, const ItemType& Value) -> decltype(GetNum(PagedArray))
+{
+	using SizeType = decltype(GetNum(PagedArray));
+	return PagedArrayAlgoImpl::LowerBoundInternal<SizeType>(PagedArray, Value, FIdentityFunctor(), TLess<>());
+}
+
+/**
+ * Performs binary search, resulting in position of the first element with projected value >= Value.
+ *
+ * @param PagedArray The paged array to search through, must be already sorted by SortPredicate.
+ * @param Value The value to look for
+ * @param Projection Functor or data member pointer; called via Invoke to compare to Value.
+ * @param SortPredicate Predicate for sort comparison; defaults to <.
+ *
+ * @returns Position of the first element with projected value >= Value; may be == Num.
+ */
+template <typename ItemType, typename PageType, typename ValueType, typename ProjectionType, typename SortPredicateType>
+FORCEINLINE auto LowerBoundBy(const TPagedArray<ItemType, PageType>& PagedArray, const ValueType& Value, ProjectionType Projection, SortPredicateType SortPredicate) -> decltype(GetNum(PagedArray))
+{
+	using SizeType = decltype(GetNum(PagedArray));
+	return PagedArrayAlgoImpl::LowerBoundInternal<SizeType>(PagedArray, Value, Projection, SortPredicate);
+}
+template <typename ItemType, typename PageType, typename ValueType, typename ProjectionType>
+FORCEINLINE auto LowerBoundBy(const TPagedArray<ItemType, PageType>& PagedArray, const ValueType& Value, ProjectionType Projection) -> decltype(GetNum(PagedArray))
+{
+	using SizeType = decltype(GetNum(PagedArray));
+	return PagedArrayAlgoImpl::LowerBoundInternal<SizeType>(PagedArray, Value, Projection, TLess<>());
+}
+
+/**
+ * Performs binary search, resulting in position of the first element > Value.
+ *
+ * @param PagedArray The paged array to search through, must be already sorted by SortPredicate.
+ * @param Value The value to look for
+ * @param SortPredicate Predicate for sort comparison; defaults to <.
+ *
+ * @returns Position of the first element > Value; may be == Num.
+ */
+template <typename ItemType, typename PageType, typename SortPredicateType>
+FORCEINLINE auto UpperBound(const TPagedArray<ItemType, PageType>& PagedArray, const ItemType& Value, SortPredicateType SortPredicate) -> decltype(GetNum(PagedArray))
+{
+	using SizeType = decltype(GetNum(PagedArray));
+	return PagedArrayAlgoImpl::UpperBoundInternal<SizeType>(PagedArray, Value, FIdentityFunctor(), SortPredicate);
+}
+template <typename ItemType, typename PageType>
+FORCEINLINE auto UpperBound(const TPagedArray<ItemType, PageType>& PagedArray, const ItemType& Value) -> decltype(GetNum(PagedArray))
+{
+	using SizeType = decltype(GetNum(PagedArray));
+	return PagedArrayAlgoImpl::UpperBoundInternal<SizeType>(PagedArray, Value, FIdentityFunctor(), TLess<>());
+}
+
+/**
+ * Performs binary search, resulting in position of the first element with projected value > Value.
+ *
+ * @param PagedArray The paged array to search through, must be already sorted by SortPredicate.
+ * @param Value The value to look for
+ * @param Projection Functor or data member pointer; called via Invoke to compare to Value.
+ * @param SortPredicate Predicate for sort comparison; defaults to <.
+ *
+ * @returns Position of the first element with projected value > Value; may be == Num.
+ */
+template <typename ItemType, typename PageType, typename ValueType, typename ProjectionType, typename SortPredicateType>
+FORCEINLINE auto UpperBoundBy(const TPagedArray<ItemType, PageType>& PagedArray, const ValueType& Value, ProjectionType Projection, SortPredicateType SortPredicate) -> decltype(GetNum(PagedArray))
+{
+	using SizeType = decltype(GetNum(PagedArray));
+	return PagedArrayAlgoImpl::UpperBoundInternal<SizeType>(PagedArray, Value, Projection, SortPredicate);
+}
+template <typename ItemType, typename PageType, typename ValueType, typename ProjectionType>
+FORCEINLINE auto UpperBoundBy(const TPagedArray<ItemType, PageType>& PagedArray, const ValueType& Value, ProjectionType Projection) -> decltype(GetNum(PagedArray))
+{
+	using SizeType = decltype(GetNum(PagedArray));
+	return PagedArrayAlgoImpl::UpperBoundInternal<SizeType>(PagedArray, Value, Projection, TLess<>());
+}
+
+/**
+ * Performs binary search, resulting in position of the first element closest to Value.
+ *
+ * @param PagedArray The paged array to search through; must be already sorted by SortPredicate.
+ * @param Value The value to look for
+  * @param SortPredicate The predicate for sort comparison; defaults to <.
+ *
+ * @returns The position of the first element closest to Value; == 0 if Num == 0, otherwise is a value in range [0 .. Num-1].
+ */
+template <typename ItemType, typename PageType, typename SortPredicateType>
+FORCEINLINE auto BinarySearchClosest(const TPagedArray<ItemType, PageType>& PagedArray, const ItemType& Value, SortPredicateType SortPredicate) -> decltype(GetNum(PagedArray))
+{
+	using SizeType = decltype(GetNum(PagedArray));
+	return PagedArrayAlgoImpl::BinarySearchClosestInternal<SizeType>(PagedArray, Value, FIdentityFunctor(), SortPredicate);
+}
+template <typename ItemType, typename PageType>
+FORCEINLINE auto BinarySearchClosest(const TPagedArray<ItemType, PageType>& PagedArray, const ItemType& Value) -> decltype(GetNum(PagedArray))
+{
+	using SizeType = decltype(GetNum(PagedArray));
+	return PagedArrayAlgoImpl::BinarySearchClosestInternal<SizeType>(PagedArray, Value, FIdentityFunctor(), TLess<>());
+}
+
+/**
+ * Performs binary search, resulting in position of the first element with projected value closest to Value.
+ *
+ * @param PagedArray The paged array to search through; must be already sorted by SortPredicate.
+ * @param Value The value to look for
+ * @param Projection The functor or data member pointer; called via Invoke to compare to Value.
+ * @param SortPredicate The predicate for sort comparison; defaults to <.
+ *
+ * @returns The position of the first element with projected value closest to Value; == 0 if Num == 0, otherwise is a value in range [0 .. Num-1].
+ */
+template <typename ItemType, typename PageType, typename ValueType, typename ProjectionType, typename SortPredicateType>
+FORCEINLINE auto BinarySearchClosestBy(const TPagedArray<ItemType, PageType>& PagedArray, const ValueType& Value, ProjectionType Projection, SortPredicateType SortPredicate) -> decltype(GetNum(PagedArray))
+{
+	using SizeType = decltype(GetNum(PagedArray));
+	return PagedArrayAlgoImpl::BinarySearchClosestInternal<SizeType>(PagedArray, Value, Projection, SortPredicate);
+}
+template <typename ItemType, typename PageType, typename ValueType, typename ProjectionType>
+FORCEINLINE auto BinarySearchClosestBy(const TPagedArray<ItemType, PageType>& PagedArray, const ValueType& Value, ProjectionType Projection) -> decltype(GetNum(PagedArray))
+{
+	using SizeType = decltype(GetNum(PagedArray));
+	return PagedArrayAlgoImpl::BinarySearchClosestInternal<SizeType>(PagedArray, Value, Projection, TLess<>());
+}
+
+} // namespace PagedArrayAlgo
 } // namespace TraceServices
