[Insights]
  - LogView: Added LowerBoundByTime, UpperBoundByTime and BinarySearchClosestByTime in ILogProvider. Added comments to all API in ILogProvider. Also, fixed the API comments in IEditableLogProvider.
  - LogView: Added code to allow search of a log message by time.
  - LogView: Fixed issue where log messages are added in non monotonic increasing order by timestamp.

diff --git a/Engine/Source/Developer/TraceInsights/Private/Insights/Widgets/SLogView.cpp b/Engine/Source/Developer/TraceInsights/Private/Insights/Widgets/SLogView.cpp
--- a/Engine/Source/Developer/TraceInsights/Private/Insights/Widgets/SLogView.cpp
+++ b/Engine/Source/Developer/TraceInsights/Private/Insights/Widgets/SLogView.cpp
@@ -519,6 +519,7 @@ SLogView::SLogView()
 	, bIsFilteringAsyncTaskCancelRequested(false)
 	, TotalNumCategories(0)
 	, TotalNumMessages(0)
+	, TotalNumInserts(0)
 	, bIsDirty(false)
 {
 }
@@ -567,6 +568,7 @@ void SLogView::Reset()
 
 	TotalNumCategories = 0;
 	TotalNumMessages = 0;
+	TotalNumInserts = 0;
 
 	bIsDirty = false;
 	DirtyStopwatch.Stop();
@@ -751,30 +753,37 @@ void SLogView::Tick(const FGeometry& AllottedGeometry, const double InCurrentTim
 {
 	LLM_SCOPE_BYTAG(Insights);
 
-	int32 NewMessageCount = 0;
-
 	TSharedPtr<const TraceServices::IAnalysisSession> Session = FInsightsManager::Get()->GetSession();
 
 	Cache.SetSession(Session);
 
+	int32 LogProviderNumCategories = 0;
+	int32 LogProviderNumMessages = 0;
+	int32 LogProviderNumInserts = 0;
+
 	if (Session.IsValid())
 	{
 		TraceServices::FAnalysisSessionReadScope SessionReadScope(*Session.Get());
 		const TraceServices::ILogProvider& LogProvider = TraceServices::ReadLogProvider(*Session.Get());
 
-		NewMessageCount = static_cast<int32>(LogProvider.GetMessageCount());
+		LogProviderNumCategories = static_cast<int32>(LogProvider.GetCategoryCount());
+		LogProviderNumMessages = static_cast<int32>(LogProvider.GetMessageCount());
+		LogProviderNumInserts = static_cast<int32>(LogProvider.GetInsertCount());
+	}
 
-		//TODO: show only categories that are used in current trace
-		//TODO: cause of duplicates: a) runtime, b) case insensitive, c) stripped "Log" prefix
+	if (LogProviderNumCategories != TotalNumCategories)
+	{
+		TSet<FName> Categories;
+		TMap<FName, int32> DuplicatedCategories;
 
-		const int32 NumCategories = static_cast<int32>(LogProvider.GetCategoryCount());
-		if (NumCategories != TotalNumCategories)
+		if (Session.IsValid())
 		{
-			TotalNumCategories = NumCategories;
-			UE_LOG(TraceInsights, Log, TEXT("[LogView] Total Log Categories: %d"), TotalNumCategories);
+			TraceServices::FAnalysisSessionReadScope SessionReadScope(*Session.Get());
+			const TraceServices::ILogProvider& LogProvider = TraceServices::ReadLogProvider(*Session.Get());
+
+			// Re-read the number of categories (as it might have changed between the read locks).
+			LogProviderNumCategories = static_cast<int32>(LogProvider.GetCategoryCount());
 
-			TSet<FName> Categories;
-			TMap<FName, int32> DuplicatedCategories;
 			LogProvider.EnumerateCategories([&Categories, &DuplicatedCategories](const TraceServices::FLogCategoryInfo& Category)
 			{
 				FStringView CategoryStr(Category.Name);
@@ -800,33 +809,41 @@ void SLogView::Tick(const FGeometry& AllottedGeometry, const double InCurrentTim
 					Categories.Add(CategoryName);
 				}
 			});
-			Filter.SyncAvailableCategories(Categories);
-			const int32 NumAvailableLogCategories = Filter.GetAvailableLogCategories().Num();
-			if (DuplicatedCategories.Num() > 0)
+		}
+
+		TotalNumCategories = LogProviderNumCategories;
+		UE_LOG(TraceInsights, Log, TEXT("[LogView] Total Log Categories: %d"), TotalNumCategories);
+
+		Filter.SyncAvailableCategories(Categories);
+		const int32 NumAvailableLogCategories = Filter.GetAvailableLogCategories().Num();
+		if (DuplicatedCategories.Num() > 0)
+		{
+			UE_LOG(TraceInsights, Warning, TEXT("[LogView] Duplicated Log Categories: %d (+%dx)"), DuplicatedCategories.Num(), TotalNumCategories - NumAvailableLogCategories);
+			for (const auto& KV : DuplicatedCategories)
 			{
-				UE_LOG(TraceInsights, Warning, TEXT("[LogView] Duplicated Log Categories: %d (+%dx)"), DuplicatedCategories.Num(), TotalNumCategories - NumAvailableLogCategories);
-				for (const auto& KV : DuplicatedCategories)
-				{
-					UE_LOG(TraceInsights, Log, TEXT("[LogView]    \"%s\" (+%dx)"), *KV.Key.GetPlainNameString(), KV.Value);
-				}
+				UE_LOG(TraceInsights, Log, TEXT("[LogView]    \"%s\" (+%dx)"), *KV.Key.GetPlainNameString(), KV.Value);
 			}
-			UE_LOG(TraceInsights, Log, TEXT("[LogView] Unique Log Categories: %d"), NumAvailableLogCategories);
-
-			//Cache.Reset();
-			FilteredMessages.Reset();
-			TotalNumMessages = 0;
-			//ListView->RebuildList();
-			bIsDirty = true;
-			DirtyStopwatch.Start();
 		}
+		UE_LOG(TraceInsights, Log, TEXT("[LogView] Unique Log Categories: %d"), NumAvailableLogCategories);
+
+		FilteredMessages.Reset();
+		TotalNumMessages = 0;
+		bIsDirty = true;
+		DirtyStopwatch.Start();
 	}
 
-	if (NewMessageCount != TotalNumMessages)
+	if (LogProviderNumInserts != TotalNumInserts)
+	{
+		TotalNumInserts = LogProviderNumInserts;
+		Cache.Reset();
+	}
+
+	if (LogProviderNumMessages != TotalNumMessages)
 	{
 		bIsDirty = true;
 		DirtyStopwatch.Start();
 
-		if (NewMessageCount > TotalNumMessages)
+		if (LogProviderNumMessages > TotalNumMessages)
 		{
 			if (Filter.IsFilterSet())
 			{
@@ -837,7 +854,7 @@ void SLogView::Tick(const FGeometry& AllottedGeometry, const double InCurrentTim
 
 					bIsFilteringAsyncTaskCancelRequested = false;
 					FilteringStartIndex = TotalNumMessages;
-					FilteringEndIndex = NewMessageCount;
+					FilteringEndIndex = LogProviderNumMessages;
 					FilteringChangeNumber = Filter.GetChangeNumber();
 					FilteringAsyncTask = MakeUnique<FAsyncTask<FLogFilteringAsyncTask>>(FilteringStartIndex, FilteringEndIndex, Filter, SharedThis(this));
 #if !WITH_EDITOR
@@ -854,7 +871,7 @@ void SLogView::Tick(const FGeometry& AllottedGeometry, const double InCurrentTim
 				{
 					// A task is already in progress.
 					if (FilteringStartIndex == TotalNumMessages &&
-						FilteringEndIndex <= NewMessageCount &&
+						FilteringEndIndex <= LogProviderNumMessages &&
 						FilteringChangeNumber == Filter.GetChangeNumber())
 					{
 						// The filter is still valid. Just wait.
@@ -870,13 +887,13 @@ void SLogView::Tick(const FGeometry& AllottedGeometry, const double InCurrentTim
 			{
 				FilteringStopwatch.Restart();
 
-				for (int32 Index = TotalNumMessages; Index < NewMessageCount; Index++)
+				for (int32 Index = TotalNumMessages; Index < LogProviderNumMessages; Index++)
 				{
 					FilteredMessages.Add(MakeShared<FLogMessage>(Index));
 				}
 
-				const int32 NumAddedMessages = NewMessageCount - TotalNumMessages;
-				TotalNumMessages = NewMessageCount;
+				const int32 NumAddedMessages = LogProviderNumMessages - TotalNumMessages;
+				TotalNumMessages = LogProviderNumMessages;
 				TSharedPtr<FLogMessage> SelectedLogMessage = GetSelectedLogMessage();
 				ListView->RebuildList();
 				if (SelectedLogMessage.IsValid())
@@ -894,12 +911,12 @@ void SLogView::Tick(const FGeometry& AllottedGeometry, const double InCurrentTim
 				if (DurationMs > 10) // avoids spams
 				{
 					UE_LOG(TraceInsights, Log, TEXT("[LogView] Updated (no filter; %d added / %d total messages) in %llu ms."),
-						NumAddedMessages, NewMessageCount, DurationMs);
+						NumAddedMessages, LogProviderNumMessages, DurationMs);
 				}
 #endif // !WITH_EDITOR
 			}
 		}
-		else // if (NewMessageCount < TotalNumMessages)
+		else // if (LogProviderNumMessages < TotalNumMessages)
 		{
 			// Just reset. On next Tick() the list will grow if needed.
 #if !WITH_EDITOR
@@ -909,7 +926,7 @@ void SLogView::Tick(const FGeometry& AllottedGeometry, const double InCurrentTim
 			FilteredMessages.Reset();
 			TotalNumMessages = 0;
 			ListView->RebuildList();
-			bIsDirty = (NewMessageCount != 0);
+			bIsDirty = (LogProviderNumMessages != 0);
 			if (bIsDirty)
 			{
 				DirtyStopwatch.Start();
@@ -934,7 +951,7 @@ void SLogView::Tick(const FGeometry& AllottedGeometry, const double InCurrentTim
 		// A filtering async task has completed. Check if filter used is still valid.
 		if (!bIsFilteringAsyncTaskCancelRequested &&
 			FilteringStartIndex == TotalNumMessages &&
-			FilteringEndIndex <= NewMessageCount &&
+			FilteringEndIndex <= LogProviderNumMessages &&
 			FilteringChangeNumber == Filter.GetChangeNumber())
 		{
 			FLogFilteringAsyncTask& Task = FilteringAsyncTask->GetTask();
@@ -1022,6 +1039,40 @@ void SLogView::SelectLogMessageByLogIndex(int32 LogIndex)
 
 ////////////////////////////////////////////////////////////////////////////////////////////////////
 
+void SLogView::SelectLogMessageByClosestTime(double Time)
+{
+	if (FilteredMessages.Num() == 0)
+	{
+		return;
+	}
+
+	// Find the message with the closest time, in the unfiltered list of messages.
+	int32 LogIndex = 0;
+	TSharedPtr<const TraceServices::IAnalysisSession> Session = FInsightsManager::Get()->GetSession();
+	if (Session.IsValid())
+	{
+		TraceServices::FAnalysisSessionReadScope SessionReadScope(*Session.Get());
+		const TraceServices::ILogProvider& LogProvider = TraceServices::ReadLogProvider(*Session.Get());
+		LogIndex = (int32)LogProvider.BinarySearchClosestByTime(Time);
+	}
+	else
+	{
+		return;
+	}
+
+	// We are assuming the FilteredMessages list is sorted by log index.
+	// Find first filtered message with log index >= index.
+	int32 MessageIndex = Algo::LowerBoundBy(FilteredMessages, LogIndex, &FLogMessage::GetIndex);
+	if (MessageIndex >= FilteredMessages.Num())
+	{
+		MessageIndex = FilteredMessages.Num() - 1;
+	}
+
+	SelectLogMessage(FilteredMessages[MessageIndex]);
+}
+
+////////////////////////////////////////////////////////////////////////////////////////////////////
+
 void SLogView::OnSelectedLogMessageChanged(TSharedPtr<FLogMessage> LogMessage)
 {
 	if (LogMessage.IsValid())
diff --git a/Engine/Source/Developer/TraceInsights/Private/Insights/Widgets/SLogView.h b/Engine/Source/Developer/TraceInsights/Private/Insights/Widgets/SLogView.h
--- a/Engine/Source/Developer/TraceInsights/Private/Insights/Widgets/SLogView.h
+++ b/Engine/Source/Developer/TraceInsights/Private/Insights/Widgets/SLogView.h
@@ -173,11 +173,24 @@ protected:
 	/** Stopwatch used to measure how long it takes to filter the message list. */
 	mutable FStopwatch FilteringStopwatch;
 
+	/**
+	 * Total number of log categories (including duplicates; last value read from LogProvider).
+	 * Used to detect when more log categories are added in the LogProvider, so we can update UI.
+	 */
 	int32 TotalNumCategories;
 
-	/** Total number of log messages processed, from the source Trace session. Used to detect when new log messages are added in the source Trace session. */
+	/**
+	 * Total number of log messages (last value read from LogProvider).
+	 * Used to detect when new log messages are added in the LogProvider, so we can update UI.
+	 */
 	int32 TotalNumMessages;
 
+	/**
+	 * Total number of inserts in the source table of log messages (last value read from LogProvider).
+	 * Used to detect when the order of log messages changes in source table, so we can invalidate the cache.
+	 */
+	int32 TotalNumInserts;
+
 	/** true if the list of messages is not yet updated (the filter has changed and/or the source trace messages have changed) */
 	bool bIsDirty;
 
diff --git a/Engine/Source/Developer/TraceServices/Private/Model/Log.cpp b/Engine/Source/Developer/TraceServices/Private/Model/Log.cpp
--- a/Engine/Source/Developer/TraceServices/Private/Model/Log.cpp
+++ b/Engine/Source/Developer/TraceServices/Private/Model/Log.cpp
@@ -119,14 +119,29 @@ void FLogProvider::UpdateMessageSpec(uint64 LogPoint, uint64 InCategoryPointer,
 	LogMessageSpec.Verbosity = InVerbosity;
 }
 
+FLogMessageInternal& FLogProvider::AppendMessageInternal(double Time)
+{
+	// Performs binary search, resulting in position of the first log message with time > provided time value.
+	uint64 Index = TraceServices::PagedArrayAlgo::UpperBoundBy(Messages, Time,
+			[](const FLogMessageInternal& Item) { return Item.Time; });
+
+	if (Index < Messages.Num())
+	{
+		++NumInserts;
+	}
+
+	FLogMessageInternal& InternalMessage = Messages.Insert(Index);
+	InternalMessage.Time = Time;
+	return InternalMessage;
+}
+
 void FLogProvider::AppendMessage(uint64 LogPoint, double Time, const uint8* FormatArgs)
 {
 	Session.WriteAccessCheck();
 	FLogMessageSpec** FindSpec = SpecMap.Find(LogPoint);
 	if (FindSpec && (*FindSpec)->Verbosity != ELogVerbosity::SetColor)
 	{
-		FLogMessageInternal& InternalMessage = Messages.PushBack();
-		InternalMessage.Time = Time;
+		FLogMessageInternal& InternalMessage = AppendMessageInternal(Time);
 		InternalMessage.Spec = *FindSpec;
 		FFormatArgsHelper::Format(FormatBuffer, FormatBufferSize - 1, TempBuffer, FormatBufferSize - 1, InternalMessage.Spec->FormatString, FormatArgs);
 		InternalMessage.Message = Session.StoreString(FormatBuffer);
@@ -140,88 +155,111 @@ void FLogProvider::AppendMessage(uint64 LogPoint, double Time, const TCHAR* Text
 	FLogMessageSpec** FindSpec = SpecMap.Find(LogPoint);
 	if (FindSpec && (*FindSpec)->Verbosity != ELogVerbosity::SetColor)
 	{
-		FLogMessageInternal& InternalMessage = Messages.PushBack();
-		InternalMessage.Time = Time;
+		FLogMessageInternal& InternalMessage = AppendMessageInternal(Time);
 		InternalMessage.Spec = *FindSpec;
 		InternalMessage.Message = Text;
 		Session.UpdateDurationSeconds(Time);
 	}
 }
 
-void FLogProvider::AppendMessage(uint64 LogPoint, double Time, const FString& Message)
+void FLogProvider::AppendMessage(uint64 LogPoint, double Time, const FStringView Message)
 {
 	Session.WriteAccessCheck();
 	FLogMessageSpec** FindSpec = SpecMap.Find(LogPoint);
 	if (FindSpec && (*FindSpec)->Verbosity != ELogVerbosity::SetColor)
 	{
-		FLogMessageInternal& InternalMessage = Messages.PushBack();
-		InternalMessage.Time = Time;
+		FLogMessageInternal& InternalMessage = AppendMessageInternal(Time);
 		InternalMessage.Spec = *FindSpec;
 		InternalMessage.Message = Session.StoreString(Message);
 		Session.UpdateDurationSeconds(Time);
 	}
 }
 
-uint64 FLogProvider::GetMessageCount() const
-{
-	Session.ReadAccessCheck();
-	return Messages.Num();
-}
-
 bool FLogProvider::ReadMessage(uint64 Index, TFunctionRef<void(const FLogMessageInfo&)> Callback) const
 {
 	Session.ReadAccessCheck();
+
 	if (Index >= Messages.Num())
 	{
 		return false;
 	}
-	ConstructMessage(Index, Callback);
+
+	ConstructMessage(Messages[Index], Index, Callback);
+
 	return true;
 }
 
-void FLogProvider::EnumerateMessages(double IntervalStart, double IntervalEnd, TFunctionRef<void(const FLogMessageInfo&)> Callback) const
+void FLogProvider::EnumerateMessages(double StartTime, double EndTime, TFunctionRef<void(const FLogMessageInfo&)> Callback) const
 {
 	Session.ReadAccessCheck();
-	if (IntervalStart > IntervalEnd)
+
+	if (StartTime > EndTime)
 	{
 		return;
 	}
+
 	uint64 MessageCount = Messages.Num();
-	for (uint64 Index = 0; Index < MessageCount; ++Index)
+	if (MessageCount == 0)
+	{
+		return;
+	}
+
+	for (auto It = Messages.GetIteratorFromItem(0); It; ++It)
 	{
-		double Time = Messages[Index].Time;
-		if (IntervalStart <= Time && Time <= IntervalEnd)
+		double Time = It.GetCurrentItem()->Time;
+		if (StartTime <= Time && Time <= EndTime)
 		{
-			ConstructMessage(Index, Callback);
+			ConstructMessage(*It.GetCurrentItem(), It.GetCurrentItemIndex(), Callback);
 		}
 	}
 }
 
-void FLogProvider::EnumerateMessagesByIndex(uint64 Start, uint64 End, TFunctionRef<void(const FLogMessageInfo&)> Callback) const
+void FLogProvider::EnumerateMessagesByIndex(uint64 StartIndex, uint64 EndIndex, TFunctionRef<void(const FLogMessageInfo&)> Callback) const
 {
 	Session.ReadAccessCheck();
+
 	uint64 Count = Messages.Num();
-	if (Start >= Count)
+	if (EndIndex > Count)
 	{
-		return;
+		EndIndex = Count;
 	}
-	if (End > Count)
-	{
-		End = Count;
-	}
-	if (Start >= End)
+	if (StartIndex >= EndIndex)
 	{
 		return;
 	}
-	for (uint64 Index = Start; Index < End; ++Index)
+
+	for (auto It = Messages.GetIteratorFromItem(StartIndex); It && It.GetCurrentItemIndex() < EndIndex; ++It)
 	{
-		ConstructMessage(Index, Callback);
+		ConstructMessage(*It.GetCurrentItem(), It.GetCurrentItemIndex(), Callback);
 	}
 }
 
-void FLogProvider::ConstructMessage(uint64 Index, TFunctionRef<void(const FLogMessageInfo&)> Callback) const
+uint64 FLogProvider::LowerBoundByTime(double Time) const
+{
+	Session.ReadAccessCheck();
+
+	return TraceServices::PagedArrayAlgo::LowerBoundBy(Messages, Time,
+		[](const FLogMessageInternal& Item) { return Item.Time; });
+}
+
+uint64 FLogProvider::UpperBoundByTime(double Time) const
+{
+	Session.ReadAccessCheck();
+
+	return TraceServices::PagedArrayAlgo::UpperBoundBy(Messages, Time,
+		[](const FLogMessageInternal& Item) { return Item.Time; });
+}
+
+uint64 FLogProvider::BinarySearchClosestByTime(double Time) const
+{
+	Session.ReadAccessCheck();
+
+	return TraceServices::PagedArrayAlgo::BinarySearchClosestBy(Messages, Time,
+		[](const FLogMessageInternal& Item) { return Item.Time; });
+}
+
+void FLogProvider::ConstructMessage(const FLogMessageInternal& InternalMessage, uint64 Index, TFunctionRef<void(const FLogMessageInfo&)> Callback) const
 {
-	const FLogMessageInternal& InternalMessage = Messages[Index];
 	FLogMessageInfo Message;
 	Message.Index = Index;
 	Message.Time = InternalMessage.Time;
diff --git a/Engine/Source/Developer/TraceServices/Private/Model/LogPrivate.h b/Engine/Source/Developer/TraceServices/Private/Model/LogPrivate.h
--- a/Engine/Source/Developer/TraceServices/Private/Model/LogPrivate.h
+++ b/Engine/Source/Developer/TraceServices/Private/Model/LogPrivate.h
@@ -41,16 +41,22 @@ public:
 	//////////////////////////////////////////////////
 	// Read operations
 
-	virtual uint64 GetMessageCount() const override;
+	virtual uint64 GetMessageCount() const override { Session.ReadAccessCheck(); return Messages.Num(); }
 	virtual bool ReadMessage(uint64 Index, TFunctionRef<void(const FLogMessageInfo&)> Callback) const override;
-	virtual void EnumerateMessages(double IntervalStart, double IntervalEnd, TFunctionRef<void(const FLogMessageInfo&)> Callback) const override;
-	virtual void EnumerateMessagesByIndex(uint64 Start, uint64 End, TFunctionRef<void(const FLogMessageInfo&)> Callback) const override;
+	virtual void EnumerateMessagesByIndex(uint64 StartIndex, uint64 EndIndex, TFunctionRef<void(const FLogMessageInfo&)> Callback) const override;
+	virtual void EnumerateMessages(double StartTime, double EndTime, TFunctionRef<void(const FLogMessageInfo&)> Callback) const override;
 
-	virtual uint64 GetCategoryCount() const override { return Categories.Num(); }
+	virtual uint64 LowerBoundByTime(double Time) const override;
+	virtual uint64 UpperBoundByTime(double Time) const override;
+	virtual uint64 BinarySearchClosestByTime(double Time) const override;
+
+	virtual uint64 GetCategoryCount() const override { Session.ReadAccessCheck(); return Categories.Num(); }
 	virtual void EnumerateCategories(TFunctionRef<void(const FLogCategoryInfo&)> Callback) const override;
 
 	virtual const IUntypedTable& GetMessagesTable() const override { return MessagesTable; }
 
+	virtual uint64 GetInsertCount() const { Session.ReadAccessCheck(); return NumInserts; }
+
 	//////////////////////////////////////////////////
 	// Edit operations
 
@@ -65,12 +71,13 @@ public:
 	virtual void UpdateMessageSpec(uint64 LogPoint, uint64 InCategoryPointer, const TCHAR* InFormatString, const TCHAR* InFile, int32 InLine, ELogVerbosity::Type InVerbosity) override;
 	virtual void AppendMessage(uint64 LogPoint, double Time, const uint8* FormatArgs) override;
 	virtual void AppendMessage(uint64 LogPoint, double Time, const TCHAR* Text) override;
-	void AppendMessage(uint64 LogPoint, double Time, const FString& Message);
+	void AppendMessage(uint64 LogPoint, double Time, const FStringView Message);
 
 	//////////////////////////////////////////////////
 
 private:
-	void ConstructMessage(uint64 Id, TFunctionRef<void(const FLogMessageInfo&)> Callback) const;
+	FLogMessageInternal& AppendMessageInternal(double Time);
+	void ConstructMessage(const FLogMessageInternal& InternalMessage, uint64 Index, TFunctionRef<void(const FLogMessageInfo&)> Callback) const;
 
 	enum
 	{
@@ -86,6 +93,7 @@ private:
 	TCHAR FormatBuffer[FormatBufferSize];
 	TCHAR TempBuffer[FormatBufferSize];
 	TTableView<FLogMessageInternal> MessagesTable;
+	uint64 NumInserts = 0;
 };
 
 } // namespace TraceServices
diff --git a/Engine/Source/Developer/TraceServices/Public/TraceServices/Model/Log.h b/Engine/Source/Developer/TraceServices/Public/TraceServices/Model/Log.h
--- a/Engine/Source/Developer/TraceServices/Public/TraceServices/Model/Log.h
+++ b/Engine/Source/Developer/TraceServices/Public/TraceServices/Model/Log.h
@@ -35,13 +35,107 @@ class ILogProvider
 {
 public:
 	virtual ~ILogProvider() = default;
+
+	/**
+	 * Gets the number of log messages.
+	 *
+	 * @returns The current number of log messages.
+	 */
 	virtual uint64 GetMessageCount() const = 0;
-	virtual void EnumerateMessages(double IntervalStart, double IntervalEnd, TFunctionRef<void(const FLogMessageInfo&)> Callback) const = 0;
-	virtual void EnumerateMessagesByIndex(uint64 Start, uint64 End, TFunctionRef<void(const FLogMessageInfo&)> Callback) const = 0;
+
+	/**
+	 * Reads information for a single message specified by index.
+	 *
+	 * @param Index The index of the message to read.
+	 * @param Callback The function to be called with the log message's information.
+	 * @returns Whether Index is valid or not.
+	 *
+	 * @note The current number of messages could be different than the value returned by the last call to GetMessageCount().
+	 */
 	virtual bool ReadMessage(uint64 Index, TFunctionRef<void(const FLogMessageInfo&)> Callback) const = 0;
+
+	/**
+	 * Enumerates messages in the specified index interval [StartIndex, EndIndex).
+	 *
+	 * @param StartIndex The inclusive start index of the interval.
+	 * @param EndIndex The exclusive end index of the interval.
+	 * @param Callback The function to be called for each log message with the message's information.
+	 *
+	 * @note The current number of messages could be different than the value returned by the last call to GetMessageCount().
+	 */
+	virtual void EnumerateMessagesByIndex(uint64 StartIndex, uint64 EndIndex, TFunctionRef<void(const FLogMessageInfo&)> Callback) const = 0;
+
+	/**
+	 * Enumerates messages with timestamp in the specified time interval [StartTime, EndTime].
+	 *
+	 * @param StartTime The inclusive start time of the interval.
+	 * @param EndTime The inclusive end time of the interval.
+	 * @param Callback The function to be called for each log message with the message's information.
+	 *
+	 * @note The current number of messages could be different than the value returned by the last call to GetMessageCount().
+	 */
+	virtual void EnumerateMessages(double StartTime, double EndTime, TFunctionRef<void(const FLogMessageInfo&)> Callback) const = 0;
+
+	/**
+	 * Performs binary search, resulting in position of the first log message with time >= provided time value.
+	 *
+	 * @param Time The time value, in seconds.
+	 * @returns The index of the first log message with time >= provided time value; [0, MessageCount]. Returns ~0ull if not implemented.
+	 *
+	 * @note The current number of messages (MessageCount) could be different than the value returned by the last call to GetMessageCount().
+	 */
+	virtual uint64 LowerBoundByTime(double Time) const { return ~0ull; }
+
+	/**
+	 * Performs binary search, resulting in position of the first log message with time > provided time value.
+	 *
+	 * @param Time The time value, in seconds.
+	 * @returns The index of the first log message with time > provided time value; [0, MessageCount]. Returns ~0ull if not implemented.
+	 *
+	 * @note The current number of messages (MessageCount) could be different than the value returned by the last call to GetMessageCount().
+	 */
+	virtual uint64 UpperBoundByTime(double Time) const { return ~0ull; }
+
+	/**
+	 * Finds the log message with closest timestamp to the specified time.
+	 *
+	 * @param Time The time value, in seconds.
+	 * @returns The index of the message with the closest timestamp to the specified time.
+	 *          If MessageCount is 0, this function returns 0; otherwise it returns a valid index in range [0, MessageCount-1]. Returns ~0ull if not implemented.
+	 *
+	 * @note The current number of messages (MessageCount) could be different than the value returned by the last call to GetMessageCount().
+	 */
+	virtual uint64 BinarySearchClosestByTime(double Time) const { return ~0ull; }
+
+	/**
+	 * Gets the number of log categories.
+	 *
+	 * @returns The current number of log categories.
+	 */
 	virtual uint64 GetCategoryCount() const = 0;
+
+	/**
+	 * Enumerates the log categories.
+	 *
+	 * @param Callback The function to be called for each log category with the category's information.
+	 *
+	 * @note The current number of messages could be different than the value returned by the last call to GetMessageCount().
+	 */
 	virtual void EnumerateCategories(TFunctionRef<void(const FLogCategoryInfo&)> Callback) const = 0;
+
+	/**
+	 * Gets the untyped table for log messages.
+	 *
+	 * @returns The IUntypedTable reference to the messages table.
+	 */
 	virtual const IUntypedTable& GetMessagesTable() const = 0;
+
+	/**
+	 * Gets the number of inserts (when a message is inserted before other messages).
+	 *
+	 * @returns The number of inserts.
+	 */
+	virtual uint64 GetInsertCount() const { return 0; }
 };
 
 class IEditableLogProvider
@@ -50,83 +144,82 @@ class IEditableLogProvider
 public:
 	virtual ~IEditableLogProvider() = default;
 
-	/*
-	* Register a new log message category.
-	*
-	* @return The category identity.
-	*/
+	/**
+	 * Registers a new log message category.
+	 *
+	 * @return The category identity.
+	 */
 	virtual uint64 RegisterCategory() = 0;
 
-	/*
-	* Fetch the data structure for a log category.
-	*
-	* @param CategoryPointer	The unique identity (memory address) of the category instrumentation.
-	*
-	* @return A reference to the category structure.
-	*/
+	/**
+	 * Fetches the data structure for a log category.
+	 *
+	 * @param CategoryPointer The unique identity (memory address) of the category instrumentation.
+	 * @return A reference to the category structure.
+	 */
 	virtual FLogCategoryInfo& GetCategory(uint64 CategoryPointer) = 0;
 
-	/*
-	* Update a log message's Category information.
-	*
-	* @param LogPoint			The log message to update.
-	* @param InCategoryPointer	The category.
-	*/
+	/**
+	 * Updates the Category information for a log message.
+	 *
+	 * @param LogPoint The log message to update.
+	 * @param InCategoryPointer The category.
+	 */
 	virtual void UpdateMessageCategory(uint64 LogPoint, uint64 InCategoryPointer) = 0;
 
-	/*
-	* Update a log message's format string.
-	*
-	* @param LogPoint			The log message to update.
-	* @param InFormatString		The format string whose memory is stored in the Session.
-	*/
+	/**
+	 * Updates the format string for a log message.
+	 *
+	 * @param LogPoint The log message to update.
+	 * @param InFormatString The format string whose memory is stored in the Session.
+	 */
 	virtual void UpdateMessageFormatString(uint64 LogPoint, const TCHAR* InFormatString) = 0;
 
-	/*
-	* Update a log message's file location.
-	*
-	* @param LogPoint	The log message to update.
-	* @param InFile		The file path of the message's location.
-	* @param InLine		The line number of the message's location.
-	*/
+	/**
+	 * Updates the file location for a log message.
+	 *
+	 * @param LogPoint The log message to update.
+	 * @param InFile The file path of the message's location.
+	 * @param InLine The line number of the message's location.
+	 */
 	virtual void UpdateMessageFile(uint64 LogPoint, const TCHAR* InFile, int32 InLine) = 0;
 
-	/*
-	* Update a log message's verbosity.
-	*
-	* @param LogPoint		The log message to update.
-	* @param InVerbosity	The verbosity of the message.
-	*/
+	/**
+	 * Updates the verbosity for a log message.
+	 *
+	 * @param LogPoint The log message to update.
+	 * @param InVerbosity The verbosity of the message.
+	 */
 	virtual void UpdateMessageVerbosity(uint64 LogPoint, ELogVerbosity::Type InVerbosity) = 0;
 
-	/*
-	* Update a log message's information.
-	*
-	* @param LogPoint			The log message to update.
-	* @param InCategoryPointer	The category.
-	* @param InFormatString		The format string whose memory is stored in the Session.
-	* @param InFile				The file path of the message's location.
-	* @param InLine				The line number of the message's location.
-	* @param InVerbosity		The verbosity of the message.
-	*/
+	/**
+	 * Updates information for a log message.
+	 *
+	 * @param LogPoint The log message to update.
+	 * @param InCategoryPointer The category.
+	 * @param InFormatString The format string whose memory is stored in the Session.
+	 * @param InFile The file path of the message's location.
+	 * @param InLine The line number of the message's location.
+	 * @param InVerbosity The verbosity of the message.
+	 */
 	virtual void UpdateMessageSpec(uint64 LogPoint, uint64 InCategoryPointer, const TCHAR* InFormatString, const TCHAR* InFile, int32 InLine, ELogVerbosity::Type InVerbosity) = 0;
 
-	/*
-	* Append a new instance of a message from the trace session.
-	*
-	* @param LogPoint	The unique identity (memory address) of the message instrumentation.
-	* @param Time		The time in seconds of the event.
-	* @param FormatArgs	The arguments to use in conjunction with the spec's FormatString.
-	*/
+	/**
+	 * Appends a new instance of a message from the trace session.
+	 *
+	 * @param LogPoint The unique identity (memory address) of the message instrumentation.
+	 * @param Time The time in seconds of the event.
+	 * @param FormatArgs The arguments to use in conjunction with the spec's FormatString.
+	 */
 	virtual void AppendMessage(uint64 LogPoint, double Time, const uint8* FormatArgs) = 0;
 
-	/*
-	* Append a new instance of a message from the trace session.
-	*
-	* @param LogPoint	The unique identity (memory address) of the message instrumentation.
-	* @param Time		The time in seconds of the event.
-	* @param Text		The the message text.
-	*/
+	/**
+	 * Appends a new instance of a message from the trace session.
+	 *
+	 * @param LogPoint The unique identity (memory address) of the message instrumentation.
+	 * @param Time The time in seconds of the event.
+	 * @param Text The message text.
+	 */
 	virtual void AppendMessage(uint64 LogPoint, double Time, const TCHAR* Text) = 0;
 };
 
