[Insights] Removed unused variables in NetTraceAnalyzer.cpp (PR #11921).

diff --git a/Engine/Source/Developer/TraceServices/Private/Analyzers/NetTraceAnalyzer.cpp b/Engine/Source/Developer/TraceServices/Private/Analyzers/NetTraceAnalyzer.cpp
--- a/Engine/Source/Developer/TraceServices/Private/Analyzers/NetTraceAnalyzer.cpp
+++ b/Engine/Source/Developer/TraceServices/Private/Analyzers/NetTraceAnalyzer.cpp
@@ -20,7 +20,7 @@ enum ENetTraceAnalyzerVersion
 	ENetTraceAnalyzerVersion_Initial = 1,
 	ENetTraceAnalyzerVersion_BunchChannelIndex = 2,
 	ENetTraceAnalyzerVersion_BunchChannelInfo = 3,
-	ENetTraceAnalyzerVersion_FixedBunchSizeEncoding = 4,		
+	ENetTraceAnalyzerVersion_FixedBunchSizeEncoding = 4,
 };
 
 
@@ -291,7 +291,7 @@ void FNetTraceAnalyzer::HandlePacketContentEvent(const FOnEventContext& Context,
 
 						// Add to active objects
 						GameInstanceState->ActiveObjects.Add(DecodedNameOrObjectId, { ObjectInstance.ObjectIndex, ObjectInstance.NameIndex });
-						
+
 						Event.NameIndex = ObjectInstance.NameIndex;
 						Event.ObjectInstanceIndex = ObjectInstance.ObjectIndex;
 					}
@@ -432,13 +432,11 @@ void FNetTraceAnalyzer::FlushPacketEvents(FNetTraceConnectionState& ConnectionSt
 	TPagedArray<FNetProfilerContentEvent>& Events = ConnectionData.ContentEvents;
 
 	TArray<FNetProfilerContentEvent>& BunchEvents = ConnectionState.BunchEvents[ConnectionMode];
-	const int32 NumPacketEvents = BunchEvents.Num();
 
 	int32 CurrentBunchEventIndex = 0;
 
 	// Track bunch offsets
 	uint32 NextBunchOffset = 0U;
-	uint32 NextEventOffset = 0U;
 
 	int32 NonBunchEventCount = ConnectionState.BunchInfos[ConnectionMode].Num() ? ConnectionState.BunchInfos[ConnectionMode][0].FirstBunchEventIndex : BunchEvents.Num();
 
@@ -480,7 +478,6 @@ void FNetTraceAnalyzer::FlushPacketEvents(FNetTraceConnectionState& ConnectionSt
 
 			// Accumulate offset
 			NextBunchOffset += Bunch.BunchBits + Bunch.HeaderBits;
-			NextEventOffset = NextBunchOffset;
 
 			// Reset event count
 			EventsToAdd = 0U;
@@ -541,7 +538,7 @@ void FNetTraceAnalyzer::HandlePacketEvent(const FOnEventContext& Context, const
 				Packet.ConnectionState = ConnectionState->ConnectionState;
 				Packet.ContentSizeInBits = 0;
 				Packet.TotalPacketSizeInBytes = (Packet.ContentSizeInBits + 7u) >> 3u;
-			
+
 				++ExpectedSequenceNumber;
 				++ConnectionData.PacketChangeCount;
 			}
@@ -815,7 +812,7 @@ void FNetTraceAnalyzer::HandleObjectCreatedEvent(const FOnEventContext& Context,
 {
 	const uint64 TypeId = EventData.GetValue<uint64>("TypeId");
 	const uint64 ObjectId = EventData.GetValue<uint64>("ObjectId");
-	const uint32 OwnerId = EventData.GetValue<uint32>("OwnerId");
+	//const uint32 OwnerId = EventData.GetValue<uint32>("OwnerId");
 	const uint16 NameId = EventData.GetValue<uint16>("NameId");
 	const uint8 GameInstanceId = EventData.GetValue<uint8>("GameInstanceId");
 
