Speed up symbol filtering in Insights

When loading a large symbol file into insights, the time spent filtering becomes very noticable. In total, it takes about 90 to load the symbols for a large binary. Of these 90s, about 72.5s are spent in the regex matcher in FResolvedSymbolFilter::Update. The regexes we are checking for could also be checked for by just scanning for a substring. With this change, the 73s shrink to < 300ms.

diff --git a/Engine/Source/Developer/TraceServices/Private/Model/ModuleProvider.cpp b/Engine/Source/Developer/TraceServices/Private/Model/ModuleProvider.cpp
--- a/Engine/Source/Developer/TraceServices/Private/Model/ModuleProvider.cpp
+++ b/Engine/Source/Developer/TraceServices/Private/Model/ModuleProvider.cpp
@@ -13,7 +13,6 @@
 #include "Containers/StringView.h"
 #include "GenericPlatform/GenericPlatformFile.h"
 #include "HAL/PlatformFileManager.h"
-#include "Internationalization/Regex.h"
 #include "Misc/Paths.h"
 #include "Misc/PathViews.h"
 #include "Misc/ScopeRWLock.h"
@@ -48,7 +47,7 @@ public:
 
 private:
 	TArray<FString> IgnoreSymbolsByFunctionName;
-	TArray<FRegexPattern> IgnoreSymbolsByFilePath;
+	TArray<FStringView> IgnoreSymbolsByFilePath;
 };
 
 ////////////////////////////////////////////////////////////////////////////////////////////////////
@@ -483,9 +482,9 @@ FResolvedSymbolFilter::FResolvedSymbolFilter()
 	IgnoreSymbolsByFunctionName.Add(TEXT("FMallocBinned"));
 	IgnoreSymbolsByFunctionName.Add(TEXT("FD3D12Adapter::TraceMemoryAllocation"));
 
-	IgnoreSymbolsByFilePath.Add(FRegexPattern(FString(TEXT(".*/Containers/.*"))));
-	IgnoreSymbolsByFilePath.Add(FRegexPattern(FString(TEXT(".*/ConcurrentLinearAllocator.*"))));
-	IgnoreSymbolsByFilePath.Add(FRegexPattern(FString(TEXT(".*/D3D12PoolAllocator.*"))));
+	IgnoreSymbolsByFilePath.Add(TEXTVIEW("/Containers/"));
+	IgnoreSymbolsByFilePath.Add(TEXTVIEW("/ConcurrentLinearAllocator"));
+	IgnoreSymbolsByFilePath.Add(TEXTVIEW("/D3D12PoolAllocator"));
 }
 
 ////////////////////////////////////////////////////////////////////////////////////////////////////
@@ -515,13 +514,12 @@ void FResolvedSymbolFilter::Update(FResolvedSymbol& InSymbol) const
 
 	if (!bIsFiltered && InSymbol.File)
 	{
-		// Ignore symbols by file path, specified as RegexPattern strings.
-		for (const FRegexPattern& RegexPattern : IgnoreSymbolsByFilePath)
+		FString File(InSymbol.File);
+		File.ReplaceCharInline(TEXT('\\'), TEXT('/'), ESearchCase::CaseSensitive);
+		// Ignore symbols by file path, specified as substrings.
+		for (const FStringView& SubString: IgnoreSymbolsByFilePath)
 		{
-			FString File(InSymbol.File);
-			File.ReplaceCharInline(TEXT('\\'), TEXT('/'), ESearchCase::CaseSensitive);
-			FRegexMatcher RegexMatcher(RegexPattern, File);
-			if (RegexMatcher.FindNext())
+			if (File.Contains(SubString, ESearchCase::CaseSensitive))
 			{
 				bIsFiltered = true;
 				break;
