[Insights] Changed also IgnoreSymbolsByFunctionName to use FStringView.

#rnx

diff --git a/Engine/Source/Developer/TraceServices/Private/Model/ModuleProvider.cpp b/Engine/Source/Developer/TraceServices/Private/Model/ModuleProvider.cpp
--- a/Engine/Source/Developer/TraceServices/Private/Model/ModuleProvider.cpp
+++ b/Engine/Source/Developer/TraceServices/Private/Model/ModuleProvider.cpp
@@ -46,7 +46,7 @@ public:
 	virtual void Update(FResolvedSymbol& InSymbol) const override;
 
 private:
-	TArray<FString> IgnoreSymbolsByFunctionName;
+	TArray<FStringView> IgnoreSymbolsByFunctionName;
 	TArray<FStringView> IgnoreSymbolsByFilePath;
 };
 
@@ -465,22 +465,22 @@ uint32 TModuleProvider<SymbolResolverType>::GetNumCachedSymbolsFromModule(uint64
 
 FResolvedSymbolFilter::FResolvedSymbolFilter()
 {
-	IgnoreSymbolsByFunctionName.Add(TEXT("FMemory::"));
-	IgnoreSymbolsByFunctionName.Add(TEXT("FMallocWrapper::"));
-	IgnoreSymbolsByFunctionName.Add(TEXT("FMallocPoisonProxy::"));
-	IgnoreSymbolsByFunctionName.Add(TEXT("FMallocLeakDetectionProxy::"));
-	IgnoreSymbolsByFunctionName.Add(TEXT("FVirtualWinApiHooks::"));
-	IgnoreSymbolsByFunctionName.Add(TEXT("Malloc"));
-	IgnoreSymbolsByFunctionName.Add(TEXT("Realloc"));
-	IgnoreSymbolsByFunctionName.Add(TEXT("Free"));
-	IgnoreSymbolsByFunctionName.Add(TEXT("MemoryTrace_"));
-	IgnoreSymbolsByFunctionName.Add(TEXT("operator new"));
-	IgnoreSymbolsByFunctionName.Add(TEXT("operator delete"));
-	IgnoreSymbolsByFunctionName.Add(TEXT("std::"));
-	IgnoreSymbolsByFunctionName.Add(TEXT("FWindowsPlatformMemory::"));
-	IgnoreSymbolsByFunctionName.Add(TEXT("FCachedOSPageAllocator::"));
-	IgnoreSymbolsByFunctionName.Add(TEXT("FMallocBinned"));
-	IgnoreSymbolsByFunctionName.Add(TEXT("FD3D12Adapter::TraceMemoryAllocation"));
+	IgnoreSymbolsByFunctionName.Add(TEXTVIEW("FMemory::"));
+	IgnoreSymbolsByFunctionName.Add(TEXTVIEW("FMallocWrapper::"));
+	IgnoreSymbolsByFunctionName.Add(TEXTVIEW("FMallocPoisonProxy::"));
+	IgnoreSymbolsByFunctionName.Add(TEXTVIEW("FMallocLeakDetectionProxy::"));
+	IgnoreSymbolsByFunctionName.Add(TEXTVIEW("FVirtualWinApiHooks::"));
+	IgnoreSymbolsByFunctionName.Add(TEXTVIEW("Malloc"));
+	IgnoreSymbolsByFunctionName.Add(TEXTVIEW("Realloc"));
+	IgnoreSymbolsByFunctionName.Add(TEXTVIEW("Free"));
+	IgnoreSymbolsByFunctionName.Add(TEXTVIEW("MemoryTrace_"));
+	IgnoreSymbolsByFunctionName.Add(TEXTVIEW("operator new"));
+	IgnoreSymbolsByFunctionName.Add(TEXTVIEW("operator delete"));
+	IgnoreSymbolsByFunctionName.Add(TEXTVIEW("std::"));
+	IgnoreSymbolsByFunctionName.Add(TEXTVIEW("FWindowsPlatformMemory::"));
+	IgnoreSymbolsByFunctionName.Add(TEXTVIEW("FCachedOSPageAllocator::"));
+	IgnoreSymbolsByFunctionName.Add(TEXTVIEW("FMallocBinned"));
+	IgnoreSymbolsByFunctionName.Add(TEXTVIEW("FD3D12Adapter::TraceMemoryAllocation"));
 
 	IgnoreSymbolsByFilePath.Add(TEXTVIEW("/Containers/"));
 	IgnoreSymbolsByFilePath.Add(TEXTVIEW("/ConcurrentLinearAllocator"));
@@ -502,9 +502,9 @@ void FResolvedSymbolFilter::Update(FResolvedSymbol& InSymbol) const
 	if (!bIsFiltered && InSymbol.Name)
 	{
 		// Ignore symbols by function name prefix.
-		for (const FString& Prefix : IgnoreSymbolsByFunctionName)
+		for (const FStringView& Prefix : IgnoreSymbolsByFunctionName)
 		{
-			if (FCString::Strnicmp(InSymbol.Name, *Prefix, Prefix.Len()) == 0)
+			if (FCString::Strnicmp(InSymbol.Name, Prefix.GetData(), Prefix.Len()) == 0)
 			{
 				bIsFiltered = true;
 				break;
@@ -516,8 +516,9 @@ void FResolvedSymbolFilter::Update(FResolvedSymbol& InSymbol) const
 	{
 		FString File(InSymbol.File);
 		File.ReplaceCharInline(TEXT('\\'), TEXT('/'), ESearchCase::CaseSensitive);
+
 		// Ignore symbols by file path, specified as substrings.
-		for (const FStringView& SubString: IgnoreSymbolsByFilePath)
+		for (const FStringView& SubString : IgnoreSymbolsByFilePath)
 		{
 			if (File.Contains(SubString, ESearchCase::CaseSensitive))
 			{
