Unreal Insights psym symbol resolver, for iOS

diff --git a/Engine/Source/Developer/TraceServices/Private/Model/ModuleProvider.cpp b/Engine/Source/Developer/TraceServices/Private/Model/ModuleProvider.cpp
--- a/Engine/Source/Developer/TraceServices/Private/Model/ModuleProvider.cpp
+++ b/Engine/Source/Developer/TraceServices/Private/Model/ModuleProvider.cpp
@@ -23,6 +23,8 @@
 // If both are enabled and the RAD Syms library fails to initialize, it will fall back to DbgHelp (on Windows).
 #define USE_SYMSLIB 1
 #define USE_DBGHELP 1
+// Psym resolver supports symbols in the breakpad cross-platform text format
+#define USE_PSYMRESOLVER 1
 
 // Symbol files implementations
 #if USE_SYMSLIB
@@ -31,6 +33,9 @@
 #if USE_DBGHELP
 #include "DbgHelpResolver.h"
 #endif
+#if USE_PSYMRESOLVER
+#include "PsymResolver.h"
+#endif
 
 namespace TraceServices
 {
@@ -550,6 +555,12 @@ TSharedPtr<IModuleAnalysisProvider> CreateModuleProvider(IAnalysisSession& InSes
 		Provider = MakeShared<TModuleProvider<FDbgHelpResolver>>(InSession);
 	}
 #endif // PLATFORM_WINDOWS && USE_DBGHELP
+#if USE_PSYMRESOLVER
+	if (!Provider && InSymbolFormat.Equals("psym"))
+	{
+		Provider = MakeShared<TModuleProvider<FPsymResolver>>(InSession);
+	}
+#endif
 	return Provider;
 }
 
diff --git a/Engine/Source/Developer/TraceServices/Private/Model/PsymResolver.cpp b/Engine/Source/Developer/TraceServices/Private/Model/PsymResolver.cpp
new file mode 100644
--- /dev/null
+++ b/Engine/Source/Developer/TraceServices/Private/Model/PsymResolver.cpp
@@ -0,0 +1,497 @@
+// Copyright Epic Games, Inc. All Rights Reserved.
+
+#include "PsymResolver.h"
+#include "Algo/Sort.h"
+#include "Algo/ForEach.h"
+#include "Containers/Queue.h"
+#include "Containers/StringView.h"
+#include "Containers/StringFwd.h"
+#include "HAL/PlatformProcess.h"
+#include "HAL/RunnableThread.h"
+#include "Logging/LogMacros.h"
+#include "Misc/ConfigCacheIni.h"
+#include "Misc/ConfigContext.h"
+#include "Misc/Guid.h"
+#include "Misc/PathViews.h"
+#include "Misc/ScopeLock.h"
+#include "Misc/StringBuilder.h"
+#include "TraceServices/Model/AnalysisSession.h"
+#include "Misc/FileHelper.h"
+#include "Misc/Parse.h"
+#include <atomic>
+
+/////////////////////////////////////////////////////////////////////
+namespace TraceServices {
+
+/////////////////////////////////////////////////////////////////////
+DEFINE_LOG_CATEGORY_STATIC(LogPsymResolver, Log, All);
+
+/////////////////////////////////////////////////////////////////////
+
+class FPsymSymbolStringAllocator
+{
+public:
+	FPsymSymbolStringAllocator(ILinearAllocator& InAllocator, uint32 InBlockSize)
+		: Allocator(InAllocator)
+		, BlockSize(InBlockSize)
+	{}
+
+	const TCHAR* Store(const TCHAR* InString)
+	{
+		return Store(FStringView(InString));
+	}
+
+	const TCHAR* Store(const FStringView InString)
+	{
+		const uint32 StringSize = InString.Len() + 1;
+		check(StringSize <= BlockSize);
+		if (StringSize > BlockRemaining)
+		{
+			Block = (TCHAR*)Allocator.Allocate(BlockSize * sizeof(TCHAR));
+			BlockRemaining = BlockSize;
+			++BlockUsed;
+		}
+		const uint32 CopiedSize = InString.CopyString(Block, BlockRemaining - 1, 0);
+		check(StringSize == CopiedSize + 1);
+		Block[StringSize - 1] = TEXT('\0');
+		BlockRemaining -= StringSize;
+		const TCHAR* OutString = Block;
+		Block += StringSize;
+		return OutString;
+	}
+
+private:
+	ILinearAllocator& Allocator;
+	TCHAR* Block = nullptr;
+	uint32 BlockSize;
+	uint32 BlockRemaining = 0;
+	uint32 BlockUsed = 0;
+};
+
+/////////////////////////////////////////////////////////////////////
+FPsymResolver::FPsymResolver(IAnalysisSession& InSession, IResolvedSymbolFilter& InSymbolFilter)
+	: bRunWorkerThread(false)
+	, bDrainThenStop(false)
+	, Session(InSession)
+	, SymbolFilter(InSymbolFilter)
+{
+	// Setup search paths. The SearchPaths array is a priority stack, which
+	// means paths are searched in reversed order.
+	// 1. Any new paths entered by the user this session
+	// 2. Path of the executable (if available)
+	// 3. Paths from UE_INSIGHTS_SYMBOLPATH
+	// 4. Paths from the user configuration file
+
+	// Paths from configuration
+	FString SettingsIni;
+
+	if (FConfigContext::ReadIntoGConfig().Load(TEXT("UnrealInsightsSettings"), SettingsIni))
+	{
+		GConfig->GetArray(TEXT("Insights.MemoryProfiler"), TEXT("SymbolSearchPaths"), SymbolSearchPaths, SettingsIni);
+	}
+
+	// Paths from environment
+	FString SymbolPathEnvVar =  FPlatformMisc::GetEnvironmentVariable(TEXT("UE_INSIGHTS_SYMBOL_PATH"));
+	UE_LOG(LogPsymResolver, Log, TEXT("UE_INSIGHTS_SYMBOL_PATH: '%s'"), *SymbolPathEnvVar);
+	FString SymbolPathPart;
+	while (SymbolPathEnvVar.Split(TEXT(";"), &SymbolPathPart, &SymbolPathEnvVar))
+	{
+		SymbolSearchPaths.Emplace(SymbolPathPart);
+	}
+
+	Start();
+}
+
+/////////////////////////////////////////////////////////////////////
+FPsymResolver::~FPsymResolver()
+{
+	bRunWorkerThread = false;
+	if (Thread)
+	{
+		Thread->WaitForCompletion();
+	}
+}
+
+/////////////////////////////////////////////////////////////////////
+void FPsymResolver::QueueModuleLoad(const uint8* ImageId, uint32 ImageIdSize, FModule* Module)
+{
+	check(Module != nullptr);
+
+	FScopeLock _(&ModulesCs);
+
+	const FStringView ModuleName = FPathViews::GetCleanFilename(Module->FullName);
+
+	// Add module and sort list according to base address
+	const int32 Index = LoadedModules.Add(FModuleEntry{
+		Module->Base, Module->Size, Session.StoreString(ModuleName), Session.StoreString(Module->FullName),
+		Module, TArray(ImageId, ImageIdSize)
+	});
+
+	// Queue up module to have symbols loaded
+	LoadSymbolsQueue.Enqueue(FQueuedModule{ Module, nullptr, LoadedModules[Index].ImageId});
+
+	// Sort list according to base address
+	Algo::Sort(LoadedModules, [](const FModuleEntry& Lhs, const FModuleEntry& Rhs) { return Lhs.Base < Rhs.Base; });
+
+	++ModulesDiscovered;
+}
+
+/////////////////////////////////////////////////////////////////////
+void FPsymResolver::QueueModuleReload(const FModule* Module, const TCHAR* InPath, TFunction<void(SymbolArray&)> ResolveOnSuccess)
+{
+	FScopeLock _(&ModulesCs);
+	const uint64 ModuleBase = Module->Base;
+	const FModuleEntry* Entry = LoadedModules.FindByPredicate([ModuleBase](const FModuleEntry& Entry) { return Entry.Base == ModuleBase; });
+	if (Entry)
+	{
+		// Reset stats for reloaded module.
+		Entry->Module->Stats.Discovered.store(0u);
+		Entry->Module->Stats.Resolved.store(0u);
+		Entry->Module->Stats.Failed.store(0u);
+
+		Entry->Module->Status.store(EModuleStatus::Pending);
+		LoadSymbolsQueue.Enqueue(FQueuedModule{Module, Session.StoreString(InPath), TArrayView<const uint8>(Entry->ImageId)});
+	}
+
+	SymbolArray SymbolsToResolve;
+	ResolveOnSuccess(SymbolsToResolve);
+	for(TTuple<uint64, FResolvedSymbol*> Pair : SymbolsToResolve)
+	{
+		QueueSymbolResolve(Pair.Get<0>(), Pair.Get<1>());
+	}
+
+	if (!bRunWorkerThread && Thread)
+	{
+		// Restart the worker thread if it has stopped.
+		Start();
+	}
+}
+
+/////////////////////////////////////////////////////////////////////
+void FPsymResolver::QueueSymbolResolve(uint64 Address, FResolvedSymbol* Symbol)
+{
+	ResolveQueue.Enqueue(FQueuedAddress{Address, Symbol});
+}
+
+/////////////////////////////////////////////////////////////////////
+void FPsymResolver::GetStats(IModuleProvider::FStats* OutStats) const
+{
+	FScopeLock _(&ModulesCs);
+	FMemory::Memzero(*OutStats);
+	for(const FModuleEntry& Entry : LoadedModules)
+	{
+		OutStats->SymbolsDiscovered += Entry.Module->Stats.Discovered.load();
+		OutStats->SymbolsResolved += Entry.Module->Stats.Resolved.load();
+		OutStats->SymbolsFailed += Entry.Module->Stats.Failed.load();
+	}
+	OutStats->ModulesDiscovered = ModulesDiscovered.load();
+	OutStats->ModulesFailed = ModulesFailed.load();
+	OutStats->ModulesLoaded = ModulesLoaded.load();
+}
+
+/////////////////////////////////////////////////////////////////////
+void FPsymResolver::EnumerateSymbolSearchPaths(TFunctionRef<void(FStringView Path)> Callback) const
+{
+	FScopeLock _(&SymbolSearchPathsLock);
+	Algo::ForEach(SymbolSearchPaths, Callback);
+}
+
+/////////////////////////////////////////////////////////////////////
+void FPsymResolver::OnAnalysisComplete()
+{
+	// At this point no more module loads or symbol requests will be queued,
+	// we drain the current queue, then release resources and file locks.
+	bDrainThenStop = true;
+}
+
+/////////////////////////////////////////////////////////////////////
+uint32 FPsymResolver::Run()
+{
+	while (bRunWorkerThread)
+	{
+		// Prioritize queued module loads
+		while (!LoadSymbolsQueue.IsEmpty() && bRunWorkerThread)
+		{
+			FQueuedModule Item;
+			if (LoadSymbolsQueue.Dequeue(Item))
+			{
+				LoadModuleSymbols(Item.Module, Item.Path, Item.ImageId);
+			}
+		}
+
+		// Resolve one symbol at a time to give way for modules
+		while (!ResolveQueue.IsEmpty() && LoadSymbolsQueue.IsEmpty() && bRunWorkerThread)
+		{
+			FQueuedAddress Item;
+			if (ResolveQueue.Dequeue(Item))
+			{
+				ResolveSymbol(Item.Address, *Item.Target);
+			}
+		}
+
+		if (bDrainThenStop && ResolveQueue.IsEmpty() && LoadSymbolsQueue.IsEmpty())
+		{
+			bRunWorkerThread = false;
+		}
+
+		// ...and breathe...
+		FPlatformProcess::Sleep(0.2f);
+	}
+
+	return 0;
+}
+
+/////////////////////////////////////////////////////////////////////
+void FPsymResolver::Start()
+{
+	// Start the worker thread
+	bRunWorkerThread = true;
+	Thread = FRunnableThread::Create(this, TEXT("PSymHelpWorker"), 0, TPri_Normal);
+}
+
+/////////////////////////////////////////////////////////////////////
+void FPsymResolver::Stop()
+{
+	bRunWorkerThread = false;
+}
+
+/////////////////////////////////////////////////////////////////////
+void FPsymResolver::UpdateResolvedSymbol(FResolvedSymbol& Symbol, ESymbolQueryResult Result, const TCHAR* Module, const TCHAR* Name, const TCHAR* File, uint16 Line)
+{
+	Symbol.Module = Module;
+	Symbol.Name = Name;
+	Symbol.File = File;
+	Symbol.Line = Line;
+	Symbol.Result.store(Result, std::memory_order_release);
+}
+
+/////////////////////////////////////////////////////////////////////
+const FPsymResolver::FModuleEntry* FPsymResolver::GetModuleForAddress(uint64 Address) const
+{
+	const int32 EntryIdx = Algo::LowerBoundBy(LoadedModules, Address, [](const FModuleEntry& Entry) { return Entry.Base; }) - 1;
+	if (EntryIdx < 0 || EntryIdx >= LoadedModules.Num())
+	{
+		return nullptr;
+	}
+
+	return &LoadedModules[EntryIdx];
+}
+
+/////////////////////////////////////////////////////////////////////
+
+static inline bool IsWhiteSpace(TCHAR Value) { return Value == TCHAR(' ') || Value == TCHAR('\n') || Value == TCHAR('\r'); }
+
+static FString NextToken(const FStringView& InStringView, int32& InIndex)
+{
+	FString Result;
+
+	// prevent reallocs
+	Result.Empty(32);
+
+	// find first non-whitespace char
+	while (InIndex < InStringView.Len() && IsWhiteSpace(InStringView[InIndex]))
+	{
+		InIndex++;
+	}
+
+	// copy non-whitespace chars
+	while (InIndex < InStringView.Len() && (!IsWhiteSpace(InStringView[InIndex])))
+	{
+		Result += InStringView[InIndex];
+		InIndex++;
+	}
+
+	return Result;
+}
+
+// [REDACTED_PROFILE_URL]/breakpad/blob/master/docs/symbol_files.md
+//
+// Prefix	: Info								   : Number of spaces
+// ------------------------------------------------------------------
+// MODULE	: operatingsystem architecture id name : 4
+// FILE		: number name						   : 2
+// FUNC m	: address size parameter_size name	   : 5
+// FUNC		: address size parameter_size name	   : 4
+// address	: size line filenum					   : 3
+// PUBLIC m : address parameter_size name		   : 4
+// PUBLIC	: address parameter_size name		   : 3
+// STACK	:									   : 0 // Ignore
+// INFO		:									   : 0 // Ignore
+
+void FPsymResolver::LoadModuleSymbols(const FModule* Module, const TCHAR* Path, const TArrayView<const uint8> ImageId)
+{
+	check(Module);
+	if (Path)
+	{
+		// Find the module entry
+		FScopeLock _(&ModulesCs);
+		const int32 EntryIdx = Algo::BinarySearchBy(LoadedModules, Module->Base, [](const FModuleEntry& Entry) { return Entry.Base; });
+		check(EntryIdx != INDEX_NONE);
+		const FModuleEntry& Entry = LoadedModules[EntryIdx];
+
+		uint64 BaseAddress = Module->Base;
+
+		FPsymSymbolStringAllocator StringAllocator(Session.GetLinearAllocator(), (2 << 12) / sizeof(TCHAR));
+
+		constexpr uint32 BuildIdSize = 16;
+		uint8 BuildId[BuildIdSize] = { 0 };
+		auto LineVisitor = [this, BaseAddress, &BuildId, BuildIdSize, Entry, &StringAllocator](FStringView Line)
+			{
+				int32 Index = 0;
+				FString Command = NextToken(Line, Index);
+				if (Command == TEXT("MODULE"))
+				{
+					FString OS = NextToken(Line, Index);
+					FString Architecture = NextToken(Line, Index);
+					FString BuildIdStr = NextToken(Line, Index);
+					int i = 0;
+					for (const TCHAR* Ptr = *BuildIdStr; *Ptr && *(Ptr+1) && i < BuildIdSize; Ptr+=2, i++)
+					{
+						BuildId[i] = (uint8)((FParse::HexDigit(*Ptr) << 4) + (FParse::HexDigit(*(Ptr+1))));
+					}
+					FString Name = NextToken(Line, Index);
+				}
+				else if (Command == TEXT("FILE"))
+				{
+					FString FileIndex = NextToken(Line, Index);
+					FString FileName = NextToken(Line, Index);
+
+					PsymSourceFiles.Add(FCString::Atoi(*FileIndex), StringAllocator.Store(*FileName));
+				}
+				else if (Command == TEXT("FUNC"))
+				{
+					FString Address = NextToken(Line, Index);
+					if (Address == "m")
+					{
+						Address = NextToken(Line, Index);
+					}
+
+					FString Size = NextToken(Line, Index);
+					FString ParamSize = NextToken(Line, Index);
+					FString Name = NextToken(Line, Index);
+
+ 					++Entry.Module->Stats.Discovered;
+					PsymSymbols.Add(FPsymSymbol(FParse::HexNumber64(*Address) + BaseAddress, FParse::HexNumber(*Size), StringAllocator.Store(*Name)));
+				}
+				else if (Command == TEXT("PUBLIC"))
+				{
+					FString Address = NextToken(Line, Index);
+					if (Address == "m")
+					{
+						Address = NextToken(Line, Index);
+					}
+					FString ParamSize = NextToken(Line, Index);
+					FString Name = NextToken(Line, Index);
+
+					++Entry.Module->Stats.Discovered;
+					PsymSymbols.Add(FPsymSymbol(FParse::HexNumber64(*Address) + BaseAddress, FParse::HexNumber(*ParamSize), StringAllocator.Store(*Name)));
+				}
+				else if (Command == TEXT("STACK"))
+				{
+					// ignore
+				}
+				else if (Command == TEXT("INFO"))
+				{
+					// ignore
+				}
+				else
+				{
+					FString Size = NextToken(Line, Index);
+					FString LineNumber = NextToken(Line, Index);
+					FString FileNumber = NextToken(Line, Index);
+					
+					PsymSourceLines.Add(FPsymLine(FParse::HexNumber64(*Command) + BaseAddress, FParse::HexNumber(*Size), FCString::Atoi(*LineNumber), FCString::Atoi(*FileNumber)));
+				}
+			};
+
+		FFileHelper::LoadFileToStringWithLineVisitor(Path, LineVisitor);
+
+		Algo::Sort(PsymSymbols, [](const FPsymSymbol& Left, const FPsymSymbol& Right)
+			{
+				return Left.Address < Right.Address;
+			});
+		Algo::Sort(PsymSourceLines, [](const FPsymLine& Left, const FPsymLine& Right)
+			{
+				return Left.Address < Right.Address;
+			});
+
+		TStringBuilder<256> StatusMessage;
+		EModuleStatus Status;
+
+		// Only check the first 16 bytes of the BuildId as the psym generator appears to add a trailing 0.
+		if (FMemory::Memcmp(Entry.ImageId.GetData(), BuildId, FMath::Min<int>(Entry.ImageId.Num(), BuildIdSize)) != 0)
+		{
+			StatusMessage.Appendf(TEXT("Build ID for psym %s does not match that of trace. Is this the correct psym for %s?"), Module->Name, Path);
+			Status = EModuleStatus::VersionMismatch;
+		}
+		else
+		{
+			StatusMessage.Appendf(TEXT("Loaded symbols for %s from %s."), Module->Name, Path);
+			Status = EModuleStatus::Loaded;
+			++ModulesLoaded;
+		}
+
+		// Make the status visible to the world
+		Entry.Module->StatusMessage = Session.StoreString(StatusMessage.ToView());
+		Entry.Module->Status.store(Status);
+	}
+}
+
+/////////////////////////////////////////////////////////////////////
+
+void FPsymResolver::ResolveSymbol(uint64 Address, FResolvedSymbol& Target)
+{
+	const FModuleEntry* Entry = GetModuleForAddress(Address);
+	if (Entry)
+	{
+		if (Entry->Module->Status == EModuleStatus::Loaded)
+		{
+			int32 SymbolIndex = Algo::UpperBoundBy(PsymSymbols, Address, &FPsymSymbol::Address) - 1;
+			if (PsymSymbols.IsValidIndex(SymbolIndex))
+			{
+				int32 LineIndex = Algo::UpperBoundBy(PsymSourceLines, Address, &FPsymLine::Address) - 1;
+
+				int32 LineNumber = 0;
+				const TCHAR* FileName = TEXT("?");
+				if (PsymSourceLines.IsValidIndex(LineIndex))
+				{
+					LineNumber = PsymSourceLines[LineIndex].LineNumber;
+					int32 FileIndex = PsymSourceLines[LineIndex].FileIndex;
+				
+					const TCHAR** FileNamePtr = PsymSourceFiles.Find(FileIndex);
+					if (FileNamePtr)
+					{
+						FileName = *FileNamePtr;
+					}
+				}
+
+				UpdateResolvedSymbol(Target,
+					ESymbolQueryResult::OK,
+					Entry->Name,
+					PsymSymbols[SymbolIndex].Name,
+					FileName,
+					static_cast<uint16>(LineNumber));
+
+				SymbolFilter.Update(Target);
+				return;
+			}
+
+			++Entry->Module->Stats.Failed;
+
+			UpdateResolvedSymbol(Target,
+				ESymbolQueryResult::NotFound,
+				TEXT("?"),
+				TEXT("?"),
+				TEXT("?"),
+				0);
+			SymbolFilter.Update(Target);
+		}
+		else
+		{
+			++Entry->Module->Stats.Failed;
+		}
+	}
+}
+/////////////////////////////////////////////////////////////////////
+
+} // namespace TraceServices
diff --git a/Engine/Source/Developer/TraceServices/Private/Model/PsymResolver.h b/Engine/Source/Developer/TraceServices/Private/Model/PsymResolver.h
new file mode 100644
--- /dev/null
+++ b/Engine/Source/Developer/TraceServices/Private/Model/PsymResolver.h
@@ -0,0 +1,122 @@
+// Copyright Epic Games, Inc. All Rights Reserved.
+#pragma once
+
+#include "HAL/Platform.h"
+#include "HAL/Runnable.h"
+#include "Containers/Queue.h"
+#include "Async/MappedFileHandle.h"
+#include "Async/TaskGraphInterfaces.h"
+#include "Common/PagedArray.h"
+#include "HAL/CriticalSection.h"
+#include "HAL/Platform.h"
+#include "TraceServices/Model/AnalysisSession.h"
+#include "TraceServices/Model/Modules.h"
+
+namespace TraceServices
+{
+////////////////////////////////////////////////////////////////////////////////
+
+class FPsymResolver : public FRunnable
+{
+public:
+	typedef TArray<TTuple<uint64, FResolvedSymbol*>> SymbolArray;
+
+	void Start();
+	FPsymResolver(IAnalysisSession& InSession, IResolvedSymbolFilter& InSymbolFilter);
+	~FPsymResolver();
+	void QueueModuleLoad(const uint8* ImageId, uint32 ImageIdSize, FModule* Module);
+	void QueueModuleReload(const FModule* Module, const TCHAR* Path, TFunction<void(SymbolArray&)> ResolveOnSuccess);
+	void QueueSymbolResolve(uint64 Address, FResolvedSymbol* Symbol);
+	void OnAnalysisComplete();
+	void GetStats(IModuleProvider::FStats* OutStats) const;
+	void EnumerateSymbolSearchPaths(TFunctionRef<void(FStringView Path)> Callback) const;
+
+private:
+	struct FModuleEntry
+	{
+		uint64 Base;
+		uint32 Size;
+		const TCHAR* Name;
+		const TCHAR* Path;
+		FModule* Module;
+		TArray<uint8> ImageId;
+	};
+
+	struct FQueuedAddress
+	{
+		uint64 Address;
+		FResolvedSymbol* Target;
+	};
+
+	struct FQueuedModule
+	{
+		const FModule* Module;
+		const TCHAR* Path;
+		TArrayView<const uint8> ImageId;
+	};
+
+	enum : uint32 {
+		MaxNameLen = 512,
+	};
+
+	virtual uint32 Run() override;
+	virtual void Stop() override;
+	static void UpdateResolvedSymbol(FResolvedSymbol& Symbol, ESymbolQueryResult Result, const TCHAR* Module, const TCHAR* Name, const TCHAR* File, uint16 Line);
+
+	void ResolveSymbol(uint64 Address, FResolvedSymbol& Target);
+	void LoadModuleSymbols(const FModule* Module, const TCHAR* Path, const TArrayView<const uint8> ImageId);
+	const FModuleEntry* GetModuleForAddress(uint64 Address) const;
+
+	mutable FCriticalSection ModulesCs;
+	TArray<FModuleEntry> LoadedModules;
+	TQueue<FQueuedModule, EQueueMode::Mpsc> LoadSymbolsQueue;
+	TQueue<FQueuedAddress, EQueueMode::Mpsc> ResolveQueue;
+
+	std::atomic<uint32> ModulesDiscovered;
+	std::atomic<uint32> ModulesFailed;
+	std::atomic<uint32> ModulesLoaded;
+
+	mutable FCriticalSection SymbolSearchPathsLock;
+	TArray<FString> SymbolSearchPaths;
+	bool bRunWorkerThread;
+	bool bDrainThenStop;
+	UPTRINT Handle;
+	IAnalysisSession& Session;
+	IResolvedSymbolFilter& SymbolFilter;
+	FRunnableThread* Thread;
+	
+	struct FPsymSymbol
+	{
+		uint64 Address;
+		uint32 Size;
+		const TCHAR* Name;
+
+		FPsymSymbol(uint64 InAddress, uint32 InSize, const TCHAR* InName)
+			: Address(InAddress)
+			, Size(InSize)
+			, Name(InName)
+		{}
+	};
+	struct FPsymLine
+	{
+		uint64 Address;
+		uint32 Size;
+		int32 LineNumber;
+		int32 FileIndex;
+
+		FPsymLine(uint64 InAddress, uint32 InSize, int32 InLineNumber, int32 InFileIndex)
+			: Address(InAddress)
+			, Size(InSize)
+			, LineNumber(InLineNumber)
+			, FileIndex(InFileIndex)
+		{}
+	};
+
+	TMap<int32, const TCHAR*> PsymSourceFiles;
+	TArray<FPsymSymbol> PsymSymbols;
+	TArray<FPsymLine> PsymSourceLines;
+};
+
+/////////////////////////////////////////////////////////////////////
+
+} // namespace TraceServices
\ No newline at end of file
