[Insights] Extend Timing Regions API to allow closing regions by ID....

diff --git a/Engine/Source/Developer/TraceServices/Private/Analyzers/MiscTraceAnalysis.cpp b/Engine/Source/Developer/TraceServices/Private/Analyzers/MiscTraceAnalysis.cpp
--- a/Engine/Source/Developer/TraceServices/Private/Analyzers/MiscTraceAnalysis.cpp
+++ b/Engine/Source/Developer/TraceServices/Private/Analyzers/MiscTraceAnalysis.cpp
@@ -62,7 +62,9 @@ void FMiscTraceAnalyzer::OnAnalysisBegin(const FOnAnalysisContext& Context)
 	Builder.RouteEvent(RouteId_ScreenshotChunk, "Misc", "ScreenshotChunk");
 
 	Builder.RouteEvent(RouteId_RegionBegin, "Misc", "RegionBegin");
+	Builder.RouteEvent(RouteId_RegionBeginWithId, "Misc", "RegionBeginWithId");
 	Builder.RouteEvent(RouteId_RegionEnd, "Misc", "RegionEnd");
+	Builder.RouteEvent(RouteId_RegionEndWithId, "Misc", "RegionEndWithId");
 }
 
 void FMiscTraceAnalyzer::OnAnalysisEnd()
@@ -250,23 +252,42 @@ bool FMiscTraceAnalyzer::OnEvent(uint16 RouteId, EStyle Style, const FOnEventCon
 	{
 		uint64 Cycle = EventData.GetValue<uint64>("Cycle");
 		FString Name;
-		EventData.GetString("RegionName", Name);
-
+		check(EventData.GetString("RegionName", Name));
 		FProviderEditScopeLock RegionProviderScopedLock(RegionProvider);
 		RegionProvider.AppendRegionBegin(*Name, Context.EventTime.AsSeconds(Cycle));
+
 		break;
 	}
 
+	case RouteId_RegionBeginWithId:
+	{
+		uint64 CycleAndId = EventData.GetValue<uint64>("CycleAndId");
+		FString Name;
+		check(EventData.GetString("RegionName", Name));
+		FProviderEditScopeLock RegionProviderScopedLock(RegionProvider);
+		RegionProvider.AppendRegionBeginWithId(*Name, CycleAndId, Context.EventTime.AsSeconds(CycleAndId));
+
+		break;
+	}
+		
 	case RouteId_RegionEnd:
 	{
 		uint64 Cycle = EventData.GetValue<uint64>("Cycle");
-		FString Name = TEXT("Invalid");
+		FString Name;
 		EventData.GetString("RegionName", Name);
-
 		FProviderEditScopeLock RegionProviderScopedLock(RegionProvider);
 		RegionProvider.AppendRegionEnd(*Name, Context.EventTime.AsSeconds(Cycle));
 		break;
 	}
+		
+	case RouteId_RegionEndWithId:
+	{
+		uint64 Cycle = EventData.GetValue<uint64>("Cycle");
+		uint64_t RegionId = EventData.GetValue<uint64_t>("RegionId", 0);
+		FProviderEditScopeLock RegionProviderScopedLock(RegionProvider);
+		RegionProvider.AppendRegionEndWithId(RegionId, Context.EventTime.AsSeconds(Cycle));
+		break;
+	}
 	}
 
 	return true;
diff --git a/Engine/Source/Developer/TraceServices/Private/Analyzers/MiscTraceAnalysis.h b/Engine/Source/Developer/TraceServices/Private/Analyzers/MiscTraceAnalysis.h
--- a/Engine/Source/Developer/TraceServices/Private/Analyzers/MiscTraceAnalysis.h
+++ b/Engine/Source/Developer/TraceServices/Private/Analyzers/MiscTraceAnalysis.h
@@ -53,7 +53,9 @@ private:
 		RouteId_ScreenshotHeader,
 		RouteId_ScreenshotChunk,
 		RouteId_RegionBegin,
-		RouteId_RegionEnd
+		RouteId_RegionBeginWithId,
+		RouteId_RegionEnd,
+		RouteId_RegionEndWithId
 	};
 
 	struct FThreadState
diff --git a/Engine/Source/Developer/TraceServices/Private/Model/Regions.cpp b/Engine/Source/Developer/TraceServices/Private/Model/Regions.cpp
--- a/Engine/Source/Developer/TraceServices/Private/Model/Regions.cpp
+++ b/Engine/Source/Developer/TraceServices/Private/Model/Regions.cpp
@@ -4,6 +4,7 @@
 #include "Model/RegionsPrivate.h"
 
 #include "AnalysisServicePrivate.h"
+#include "Algo/ForEach.h"
 #include "Common/FormatArgs.h"
 #include "Common/Utils.h"
 #include "Internationalization/Internationalization.h"
@@ -42,44 +43,81 @@ const FRegionLane* FRegionProvider::GetLane(int32 index) const
 	}
 	return nullptr;
 }
-
+	
 void FRegionProvider::AppendRegionBegin(const TCHAR* Name, double Time)
 {
 	EditAccessCheck();
 
-	FTimeRegion** OpenRegion = OpenRegions.Find(Name);
+	check(Name)
+
+	FTimeRegion** OpenRegion = OpenRegionsByName.Find(Name);
+	
 	if (OpenRegion)
 	{
-		if (FCString::Strcmp(Name, TEXT("<SlowTask>")) != 0)
+		++NumWarnings;
+		if (NumWarnings <= MaxWarningMessages)
 		{
-			++NumWarnings;
-			if (NumWarnings <= MaxWarningMessages)
-			{
-				UE_LOG(LogTraceServices, Warning, TEXT("[Regions] A region begin event (BeginTime=%f, Name=\"%s\") was encountered while a region with same name is already open."), Time, Name)
-			}
+			UE_LOG(LogTraceServices, Warning, TEXT("[Regions] A region begin event (BeginTime=%f, Name=\"%s\") was encountered while a region with same name is already open."), Time, Name)
 		}
 
 		// Automatically end the previous region.
 		AppendRegionEnd(Name, Time);
 	}
 
-	{
-		FTimeRegion Region;
-		Region.BeginTime = Time;
-		Region.Text = Session.StoreString(Name);
-		Region.Depth = CalculateRegionDepth(Region);
+	FTimeRegion* NewRegion = InsertNewRegion(Time, Name, 0);
+	OpenRegionsByName.Add(NewRegion->Text, NewRegion);
+}
+
+void FRegionProvider::AppendRegionBeginWithId(const TCHAR* Name, uint64_t Id, double Time)
+{
+	EditAccessCheck();
 
-		if (Region.Depth == Lanes.Num())
+	check(Name && Id)
+	FTimeRegion** OpenRegion = OpenRegionsById.Find(Id);
+	
+	if (OpenRegion)
+	{
+		++NumWarnings;
+		if (NumWarnings <= MaxWarningMessages)
 		{
-			Lanes.Emplace(Session.GetLinearAllocator());
+			UE_LOG(LogTraceServices, Warning, TEXT("[Regions] A region begin event (BeginTime=%f, Name=\"%s\", Id=\"%llu\") was encountered while a region with same name is already open."), Time, Name, Id)
 		}
 
-		Lanes[Region.Depth].Regions.EmplaceBack(Region);
-		FTimeRegion* NewOpenRegion = &(Lanes[Region.Depth].Regions.Last());
+		// Automatically end the previous region.
+		AppendRegionEndWithId(Id, Time);
+	}
+
+	FTimeRegion* NewRegion = InsertNewRegion(Time, Name, Id);
+	OpenRegionsById.Add(Id, NewRegion);
+}
+	
+void FRegionProvider::AppendRegionEnd(const TCHAR* Name, double Time)
+{
+	EditAccessCheck();
 
-		OpenRegions.Add(Region.Text, NewOpenRegion);
-		UpdateCounter++;
+	check(Name)
+	FTimeRegion** OpenRegionPos = OpenRegionsByName.Find(Name);
+	
+	if (!OpenRegionPos)
+	{
+		++NumWarnings;
+		if (NumWarnings <= MaxWarningMessages)
+		{
+			UE_LOG(LogTraceServices, Warning, TEXT("[Regions] A region end event (EndTime=%f, Name=\"%s\") was encountered without having seen a matching region begin event first."), Time, Name)
+		}
+		
+		AppendRegionBegin(Name, Time);
+		OpenRegionPos = OpenRegionsByName.Find(Name);
+		check(OpenRegionPos);
 	}
+	
+	FTimeRegion* OpenRegion = *OpenRegionPos;
+	check(OpenRegion);
+	
+	OpenRegion->EndTime = Time;
+
+	OpenRegionsByName.Remove(OpenRegion->Text);
+	UpdateCounter++;
 
 	// Update session time
 	{
@@ -88,34 +126,36 @@ void FRegionProvider::AppendRegionBegin(const TCHAR* Name, double Time)
 	}
 }
 
-void FRegionProvider::AppendRegionEnd(const TCHAR* Name, double Time)
+void FRegionProvider::AppendRegionEndWithId(uint64_t Id, double Time)
 {
 	EditAccessCheck();
 
-	FTimeRegion** OpenRegion = OpenRegions.Find(Name);
-	if (!OpenRegion)
+	check(Id)
+	FTimeRegion** OpenRegionPos = OpenRegionsById.Find(Id);
+	
+	if (!OpenRegionPos)
 	{
-		if (FCString::Strcmp(Name, TEXT("<SlowTask>")) != 0)
+		++NumWarnings;
+		if (NumWarnings <= MaxWarningMessages)
 		{
-			++NumWarnings;
-			if (NumWarnings <= MaxWarningMessages)
-			{
-				UE_LOG(LogTraceServices, Warning, TEXT("[Regions] A region end event (EndTime=%f, Name=\"%s\") was encountered without having seen a matching region begin event first."), Time, Name)
-			}
+			UE_LOG(LogTraceServices, Warning, TEXT("[Regions] A region end event (EndTime=%f, Id=\"%llu\") was encountered without having seen a matching region begin event first."), Time, Id)
 		}
 
 		// Automatically create a new region.
-		AppendRegionBegin(Name, Time);
-		OpenRegion = OpenRegions.Find(Name);
-		check(OpenRegion);
+		// Generates a display name if we're missing a begin and are closing by ID
+		FString GeneratedName = FString::Printf(TEXT("Unknown Region (missing begin, Id=\"%llu\")"), Id);
+		AppendRegionBeginWithId(*GeneratedName, Id, Time);
+		OpenRegionPos = OpenRegionsById.Find(Id);
+		check(OpenRegionPos);
 	}
+	
+	FTimeRegion* OpenRegion = *OpenRegionPos;
+	check(OpenRegion);
+	
+	OpenRegion->EndTime = Time;
 
-	{
-		(*OpenRegion)->EndTime = Time;
-
-		OpenRegions.Remove(Name);
-		UpdateCounter++;
-	}
+	OpenRegionsById.Remove(OpenRegion->Id);
+	UpdateCounter++;
 
 	// Update session time
 	{
@@ -123,24 +163,23 @@ void FRegionProvider::AppendRegionEnd(const TCHAR* Name, double Time)
 		Session.UpdateDurationSeconds(Time);
 	}
 }
-
+	
 void FRegionProvider::OnAnalysisSessionEnded()
 {
 	EditAccessCheck();
 
-	for (const auto& KV : OpenRegions)
+	auto printOpenRegionMessage = [this](const auto& KV)
 	{
 		const FTimeRegion* Region = KV.Value;
 
-		if (FCString::Strcmp(Region->Text, TEXT("<SlowTask>")) != 0)
+		++NumWarnings;
+		if (NumWarnings <= MaxWarningMessages)
 		{
-			++NumWarnings;
-			if (NumWarnings <= MaxWarningMessages)
-			{
-				UE_LOG(LogTraceServices, Warning, TEXT("[Regions] A region (BeginTime=%f, Name=\"%s\") was never closed."), Region->BeginTime, Region->Text)
-			}
+			UE_LOG(LogTraceServices, Warning, TEXT("[Regions] A region (BeginTime=%f, Name=\"%s\", Id=\"%llu\") was never closed."), Region->BeginTime, Region->Text, Region->Id)
 		}
-	}
+	};
+	Algo::ForEach(OpenRegionsById, printOpenRegionMessage);
+	Algo::ForEach(OpenRegionsByName, printOpenRegionMessage);
 
 	if (NumWarnings > 0 || NumErrors > 0)
 	{
@@ -178,6 +217,31 @@ int32 FRegionProvider::CalculateRegionDepth(const FTimeRegion& Region) const
 	return NewDepth;
 }
 
+FTimeRegion* FRegionProvider::InsertNewRegion(double BeginTime, const TCHAR* Name, uint64_t Id)
+{
+	FTimeRegion Region;
+	Region.BeginTime = BeginTime;
+	Region.Text = Session.StoreString(Name);
+	Region.Id = Id;
+	Region.Depth = CalculateRegionDepth(Region);
+
+	if (Region.Depth == Lanes.Num())
+	{
+		Lanes.Emplace(Session.GetLinearAllocator());
+	}
+
+	Lanes[Region.Depth].Regions.EmplaceBack(Region);
+	FTimeRegion* NewOpenRegion = &(Lanes[Region.Depth].Regions.Last());
+	UpdateCounter++;
+
+	// Update session time
+	{
+		FAnalysisSessionEditScope _(Session);
+		Session.UpdateDurationSeconds(BeginTime);
+	}
+	return NewOpenRegion;
+}
+
 void FRegionProvider::EnumerateLanes(TFunctionRef<void(const FRegionLane&, int32)> Callback) const
 {
 	ReadAccessCheck();
diff --git a/Engine/Source/Developer/TraceServices/Private/Model/RegionsPrivate.h b/Engine/Source/Developer/TraceServices/Private/Model/RegionsPrivate.h
--- a/Engine/Source/Developer/TraceServices/Private/Model/RegionsPrivate.h
+++ b/Engine/Source/Developer/TraceServices/Private/Model/RegionsPrivate.h
@@ -48,23 +48,28 @@ public:
 	virtual void EditAccessCheck() const override { Lock.WriteAccessCheck(GRegionsProviderLockState); }
 
 	virtual void AppendRegionBegin(const TCHAR* Name, double Time) override;
+	virtual void AppendRegionBeginWithId(const TCHAR* Name, uint64_t Id, double Time) override;
 	virtual void AppendRegionEnd(const TCHAR* Name, double Time) override;
+	virtual void AppendRegionEndWithId(uint64_t Id, double Time) override;
 
 	virtual void OnAnalysisSessionEnded() override;
 
 	//////////////////////////////////////////////////
 
 private:
-	// Update the depth member of a region to allow overlapping regions to be displayed on separate lanes.
+	/// Update the depth member of a region to allow overlapping regions to be displayed on separate lanes.
 	int32 CalculateRegionDepth(const FTimeRegion& Item) const;
-
+	/// Calculates depth, inserts a new region into the correct lane and updates session time.
+	FTimeRegion* InsertNewRegion(double BeginTime, const TCHAR* Name, uint64_t Id);
+	
 private:
 	mutable FProviderLock Lock;
 
 	IAnalysisSession& Session;
 
 	// Open regions inside lanes
-	TMap<FStringView, FTimeRegion*> OpenRegions;
+	TMap<FStringView, FTimeRegion*> OpenRegionsByName;
+	TMap<uint64_t, FTimeRegion*> OpenRegionsById;
 
 	// Closed regions
 	TArray<FRegionLane> Lanes;
diff --git a/Engine/Source/Developer/TraceServices/Public/TraceServices/Model/Regions.h b/Engine/Source/Developer/TraceServices/Public/TraceServices/Model/Regions.h
--- a/Engine/Source/Developer/TraceServices/Public/TraceServices/Model/Regions.h
+++ b/Engine/Source/Developer/TraceServices/Public/TraceServices/Model/Regions.h
@@ -16,6 +16,8 @@ struct TRACESERVICES_API FTimeRegion
 	double BeginTime = std::numeric_limits<double>::infinity();
 	double EndTime = std::numeric_limits<double>::infinity();
 	const TCHAR* Text = nullptr;
+	// Id will be zero if the region is identified by Name only
+	uint64_t Id = 0;
 	int32 Depth = -1;
 };
 
@@ -92,20 +94,40 @@ public:
 
 	/**
 	 * Append a new begin event of a region from the trace session.
+	 * Prefer opening/closing regions with an Id, since string names are not unique.
 	 *
 	 * @param Name		The string name of the region.
 	 * @param Time		The time in seconds of the begin event of this region.
 	 */
 	virtual void AppendRegionBegin(const TCHAR* Name, double Time) = 0;
-
+	
+	/**
+	 * Append a new begin event of a region from the trace session.
+	 * Id will be used to uniquely identify the new region.
+	 *
+	 * @param Name		The string name of the region.
+	 * @param Id		The Id of the region. Used to uniquely identify regions with the same name.
+	 * @param Time		The time in seconds of the begin event of this region.
+	 */
+	virtual void AppendRegionBeginWithId(const TCHAR* Name, uint64_t Id, double Time) = 0;
+	
 	/**
 	 * Append a new end event of a region from the trace session.
+	 * Prefer opening/closing regions with an Id, since string names are not unique.
 	 *
 	 * @param Name		The string name of the region.
 	 * @param Time		The time in seconds of the end event of this region.
 	 */
 	virtual void AppendRegionEnd(const TCHAR* Name, double Time) = 0;
-
+	
+	/**
+	 * Append a new end event of a region from the trace session.
+	 * The region is identified by Id.
+	 * @param Id		The Id of the region.
+	 * @param Time		The time in seconds of the end event of this region.
+	 */
+	virtual void AppendRegionEndWithId(const uint64_t Id, double Time) = 0;
+	
 	/**
 	 * Called from the analyzer once all events have been processed.
 	 * Allows post-processing and error reporting for regions that were never closed.
diff --git a/Engine/Source/Runtime/Core/Private/ProfilingDebugging/MiscTrace.cpp b/Engine/Source/Runtime/Core/Private/ProfilingDebugging/MiscTrace.cpp
--- a/Engine/Source/Runtime/Core/Private/ProfilingDebugging/MiscTrace.cpp
+++ b/Engine/Source/Runtime/Core/Private/ProfilingDebugging/MiscTrace.cpp
@@ -31,11 +31,21 @@ UE_TRACE_EVENT_BEGIN(Misc, RegionBegin)
 	UE_TRACE_EVENT_FIELD(UE::Trace::WideString, RegionName)
 UE_TRACE_EVENT_END()
 
+UE_TRACE_EVENT_BEGIN(Misc, RegionBeginWithId)
+	UE_TRACE_EVENT_FIELD(uint64, CycleAndId)
+	UE_TRACE_EVENT_FIELD(UE::Trace::WideString, RegionName)
+UE_TRACE_EVENT_END()
+
 UE_TRACE_EVENT_BEGIN(Misc, RegionEnd)
 	UE_TRACE_EVENT_FIELD(uint64, Cycle)
 	UE_TRACE_EVENT_FIELD(UE::Trace::WideString, RegionName)
 UE_TRACE_EVENT_END()
 
+UE_TRACE_EVENT_BEGIN(Misc, RegionEndWithId)
+	UE_TRACE_EVENT_FIELD(uint64, Cycle)
+	UE_TRACE_EVENT_FIELD(uint64, RegionId)
+UE_TRACE_EVENT_END()
+
 UE_TRACE_EVENT_BEGIN(Misc, BeginFrame)
 	UE_TRACE_EVENT_FIELD(uint64, Cycle)
 	UE_TRACE_EVENT_FIELD(uint8, FrameType)
@@ -83,6 +93,15 @@ void FMiscTrace::OutputBeginRegion(const TCHAR* RegionName)
 		<< RegionBegin.RegionName(RegionName);
 }
 
+uint64_t FMiscTrace::OutputBeginRegionWithId(const TCHAR* RegionName)
+{
+	const uint64_t CycleAndId = FPlatformTime::Cycles64();
+	UE_TRACE_LOG(Misc, RegionBeginWithId, RegionChannel)
+		<< RegionBeginWithId.CycleAndId(CycleAndId)
+		<< RegionBeginWithId.RegionName(RegionName);
+	return CycleAndId;
+}
+
 void FMiscTrace::OutputEndRegion(const TCHAR* RegionName)
 {
 	UE_TRACE_LOG(Misc, RegionEnd, RegionChannel)
@@ -90,6 +109,13 @@ void FMiscTrace::OutputEndRegion(const TCHAR* RegionName)
 		<< RegionEnd.RegionName(RegionName);
 }
 
+void FMiscTrace::OutputEndRegionWithId(uint64_t RegionId)
+{
+	UE_TRACE_LOG(Misc, RegionEndWithId, RegionChannel)
+		<< RegionEndWithId.Cycle(FPlatformTime::Cycles64())
+		<< RegionEndWithId.RegionId(RegionId);	
+}
+
 void FMiscTrace::OutputBookmarkInternal(const void* BookmarkPoint, uint16 EncodedFormatArgsSize, uint8* EncodedFormatArgs)
 {
 	UE_TRACE_LOG(Misc, Bookmark, BookmarkChannel)
diff --git a/Engine/Source/Runtime/Core/Public/ProfilingDebugging/MiscTrace.h b/Engine/Source/Runtime/Core/Public/ProfilingDebugging/MiscTrace.h
--- a/Engine/Source/Runtime/Core/Public/ProfilingDebugging/MiscTrace.h
+++ b/Engine/Source/Runtime/Core/Public/ProfilingDebugging/MiscTrace.h
@@ -74,8 +74,10 @@ struct FMiscTrace
 	}
 
 	CORE_API static void OutputBeginRegion(const TCHAR* RegionName);
+	[[nodiscard]] CORE_API static uint64_t OutputBeginRegionWithId(const TCHAR* RegionName);
 	CORE_API static void OutputEndRegion(const TCHAR* RegionName);
-
+	CORE_API static void OutputEndRegionWithId(uint64_t RegionId);
+	
 	CORE_API static void OutputBeginFrame(ETraceFrameType FrameType);
 	CORE_API static void OutputEndFrame(ETraceFrameType FrameType);
 
@@ -115,9 +117,15 @@ private:
 #define TRACE_BEGIN_REGION(RegionName) \
 	FMiscTrace::OutputBeginRegion(RegionName);
 
+#define TRACE_BEGIN_REGION_WITH_ID(RegionName) \
+	FMiscTrace::OutputBeginRegionWithId(RegionName);
+
 #define TRACE_END_REGION(RegionName) \
 	FMiscTrace::OutputEndRegion(RegionName);
 
+#define TRACE_END_REGION_WITH_ID(RegionId) \
+	FMiscTrace::OutputEndRegionWithId(RegionId);
+
 #define TRACE_BEGIN_FRAME(FrameType) \
 	FMiscTrace::OutputBeginFrame(FrameType);
 
@@ -138,7 +146,9 @@ private:
 #define TRACE_BOOKMARK(...)
 #define TRACE_BOOKMARK_CYCLES(...)
 #define TRACE_BEGIN_REGION(...)
+#define TRACE_BEGIN_REGION_WITH_ID(...) 0;
 #define TRACE_END_REGION(...)
+#define TRACE_END_REGION_WITH_ID(...)
 #define TRACE_BEGIN_FRAME(...)
 #define TRACE_END_FRAME(...)
 #define TRACE_SCREENSHOT(...)
