[Insights] TraceServices: Fixed session's main string store to not duplicate strings with same hash. This fixes issue with some of the duplicated CPU timers (CpuProfilerTraceAnalysis).

diff --git a/Engine/Source/Developer/TraceServices/Private/Common/StringStore.cpp b/Engine/Source/Developer/TraceServices/Private/Common/StringStore.cpp
--- a/Engine/Source/Developer/TraceServices/Private/Common/StringStore.cpp
+++ b/Engine/Source/Developer/TraceServices/Private/Common/StringStore.cpp
@@ -3,14 +3,12 @@
 #include "Common/StringStore.h"
 #include "Misc/ScopeLock.h"
 
-
 namespace TraceServices
 {
 
 FStringStore::FStringStore(FSlabAllocator& InAllocator)
 	: Allocator(InAllocator)
 {
-
 }
 
 const TCHAR* FStringStore::Store(const TCHAR* String)
@@ -21,13 +19,19 @@ const TCHAR* FStringStore::Store(const TCHAR* String)
 const TCHAR* FStringStore::Store(const FStringView& String)
 {
 	FScopeLock _(&Cs);
+
 	uint32 Hash = GetTypeHash(String);
-	const TCHAR** AlreadyStored = StoredStrings.Find(Hash);
-	if (AlreadyStored && !String.Compare(FStringView(*AlreadyStored)))
+
+	FindStoredStrings.Reset();
+	StoredStrings.MultiFind(Hash, FindStoredStrings);
+	for (const TCHAR* FoundStoredString : FindStoredStrings)
 	{
-		return *AlreadyStored;
+		if (!String.Compare(FStringView(FoundStoredString)))
+		{
+			return FoundStoredString;
+		}
 	}
-	
+
 	int32 StringLength = String.Len() + 1;
 	if (BufferLeft < StringLength)
 	{
@@ -40,12 +44,10 @@ const TCHAR* FStringStore::Store(const FStringView& String)
 	BufferPtr[StringLength - 1] = TEXT('\0');
 	BufferLeft -= StringLength;
 	BufferPtr += StringLength;
-	if (!AlreadyStored)
-	{
-		StoredStrings.Add(Hash, Stored);
-	}
+
+	StoredStrings.Add(Hash, Stored);
+
 	return Stored;
 }
 
-
 } // namespace TraceServices
diff --git a/Engine/Source/Developer/TraceServices/Private/Common/StringStore.h b/Engine/Source/Developer/TraceServices/Private/Common/StringStore.h
--- a/Engine/Source/Developer/TraceServices/Private/Common/StringStore.h
+++ b/Engine/Source/Developer/TraceServices/Private/Common/StringStore.h
@@ -3,8 +3,9 @@
 #pragma once
 
 #include "CoreTypes.h"
+
 #include "Common/SlabAllocator.h"
-#include "Common/PagedArray.h"
+#include "Containers/Array.h"
 #include "Containers/Map.h"
 #include "Containers/StringView.h"
 #include "HAL/CriticalSection.h"
@@ -26,11 +27,11 @@ private:
 	};
 	FCriticalSection Cs;
 	FSlabAllocator& Allocator;
-	TMap<uint32, const TCHAR*> StoredStrings;
+	TMultiMap<uint32, const TCHAR*> StoredStrings;
+	TArray<const TCHAR*> FindStoredStrings;
 	TCHAR* BufferPtr = nullptr;
 	uint64 BufferLeft = 0;
 	uint64 BlockCount = 0;
-
 };
 
 } // namespace TraceServices
