[Insights]
  - CpuProfilerTraceAnalysis: Improved analysis performance for traces with huge number of timers.
  - CpuProfilerTraceAnalysis: Added "[CpuProfiler] Analysis completed" log message (with analysis stats) when the analysis is completed.
  - CpuProfilerTraceAnalysis: Added warnings when analysis detects timers defined multiple times and for possible incorrect use of FCpuProfilerTrace::OutputEventType().
  - CpuProfilerTraceAnalysis: Clean-up memory when analysis completes.

diff --git a/Engine/Source/Developer/TraceServices/Private/Analyzers/CpuProfilerTraceAnalysis.cpp b/Engine/Source/Developer/TraceServices/Private/Analyzers/CpuProfilerTraceAnalysis.cpp
--- a/Engine/Source/Developer/TraceServices/Private/Analyzers/CpuProfilerTraceAnalysis.cpp
+++ b/Engine/Source/Developer/TraceServices/Private/Analyzers/CpuProfilerTraceAnalysis.cpp
@@ -72,6 +72,42 @@ void FCpuProfilerAnalyzer::OnAnalysisEnd()
 		check(ThreadState.PendingEvents.Num() == 0); // no pending events
 		check(ThreadState.ScopeStack.Num() == 0); // no open events
 	}
+
+	bool bPossibleOutputEventTypeIssue = false;
+	ScopeNameToTimerIdMap.ValueSort([](const FTimerInfo& A, const FTimerInfo& B) { return A.Count > B.Count; });
+	for (auto& KV : ScopeNameToTimerIdMap)
+	{
+		if (KV.Value.Count < 1000)
+		{
+			break;
+		}
+		UE_LOG(LogTraceServices, Warning, TEXT("[CpuProfiler] Timer defined %u times! (id=%d name=\"%s\")"), KV.Value.Count, KV.Value.Id, (const TCHAR*)KV.Key);
+		bPossibleOutputEventTypeIssue = true;
+	}
+	if (bPossibleOutputEventTypeIssue)
+	{
+		UE_LOG(LogTraceServices, Warning, TEXT("[CpuProfiler] Possible incorrect use of FCpuProfilerTrace::OutputEventType()!"));
+	}
+
+	UE_LOG(LogTraceServices, Log, TEXT("[CpuProfiler] Analysis completed (%d timers, %d spec ids, %llu scopes, %llu bytes, %.1f bytes/scope)."),
+		ScopeNameToTimerIdMap.Num(),
+		SpecIdToTimerIdMap.Num(),
+		TotalScopeCount,
+		TotalEventSize,
+		(double)TotalEventSize / (double)TotalScopeCount);
+
+	// Clean-up...
+	for (auto& KV : ThreadStatesMap)
+	{
+		FThreadState* ThreadState = KV.Value;
+		delete ThreadState;
+	}
+	ThreadStatesMap.Reset();
+	ThreadStatesMap.Shrink();
+	SpecIdToTimerIdMap.Reset();
+	SpecIdToTimerIdMap.Shrink();
+	ScopeNameToTimerIdMap.Reset();
+	ScopeNameToTimerIdMap.Shrink();
 }
 
 ////////////////////////////////////////////////////////////////////////////////////////////////////
@@ -130,8 +166,7 @@ bool FCpuProfilerAnalyzer::OnEvent(uint16 RouteId, EStyle Style, const FOnEventC
 			Line = EventData.GetValue<uint32>("Line");
 		}
 
-		constexpr bool bMergeByName = true;
-		DefineTimer(SpecId, Session.StoreString(TimerName), FileName, Line, bMergeByName);
+		DefineTimer(SpecId, Session.StoreString(TimerName), FileName, Line);
 		break;
 	}
 
@@ -187,7 +222,6 @@ bool FCpuProfilerAnalyzer::OnEvent(uint16 RouteId, EStyle Style, const FOnEventC
 		}
 
 		TotalEventSize += BufferSize;
-		BytesPerScope = double(TotalEventSize) / double(TotalScopeCount);
 		break;
 	}
 
@@ -230,7 +264,6 @@ bool FCpuProfilerAnalyzer::OnEvent(uint16 RouteId, EStyle Style, const FOnEventC
 		}
 
 		TotalEventSize += BufferSize;
-		BytesPerScope = double(TotalEventSize) / double(TotalScopeCount);
 		break;
 	}
 
@@ -663,8 +696,7 @@ void FCpuProfilerAnalyzer::OnCpuScopeEnter(const FOnEventContext& Context)
 	{
 		FString ScopeName;
 		ScopeName += Context.EventData.GetTypeInfo().GetName();
-		constexpr bool bMergeByName = true;
-		TimerId = DefineTimer(SpecId, Session.StoreString(*ScopeName), nullptr, 0, bMergeByName);
+		TimerId = DefineTimer(SpecId, Session.StoreString(*ScopeName), nullptr, 0);
 	}
 
 	TArray<uint8> CborData;
@@ -700,66 +732,77 @@ void FCpuProfilerAnalyzer::OnCpuScopeLeave(const FOnEventContext& Context)
 
 ////////////////////////////////////////////////////////////////////////////////////////////////////
 
-uint32 FCpuProfilerAnalyzer::DefineTimer(uint32 SpecId, const TCHAR* Name, const TCHAR* File, uint32 Line, bool bMergeByName)
+uint32 FCpuProfilerAnalyzer::DefineTimer(uint32 SpecId, const TCHAR* Name, const TCHAR* File, uint32 Line)
 {
+	// Expected: Name is already a pointer in a string store.
+
 	// The CPU scoped events (timers) will be merged by name.
 	// Ex.: If there are multiple timers defined in code with same name,
 	//      those will appear in Insights as a single timer.
 
 	// Check if a timer with same name was already defined.
-	uint32* FindTimerIdByName = bMergeByName ? ScopeNameToTimerIdMap.Find(Name) : nullptr;
+	FTimerInfo* FindTimerIdByName = ScopeNameToTimerIdMap.Find((uint64)Name);
 	if (FindTimerIdByName)
 	{
 		// Yes, a timer with same name was already defined.
+
+		++(FindTimerIdByName->Count);
+
+		uint32 TimerId;
+
 		// Check if SpecId is already mapped to timer.
-		uint32* FindTimerId = SpecIdToTimerIdMap.Find(SpecId);
+		const uint32* FindTimerId = SpecIdToTimerIdMap.Find(SpecId);
 		if (FindTimerId)
 		{
 			// Yes, SpecId was already mapped to a timer (ex. as an <unknown> timer).
+			TimerId = *FindTimerId;
+
 			// Update name for mapped timer.
-			EditableTimingProfilerProvider.SetTimerNameAndLocation(*FindTimerId, Name, File, Line);
+			EditableTimingProfilerProvider.SetTimerNameAndLocation(TimerId, Name, File, Line);
+
 			// In this case, we do not remap the SpecId to the previously defined timer with same name.
 			// This is because the two timers are already used in timelines.
 			// So we will continue to use separate timers, even if those have same name.
-			return *FindTimerId;
 		}
 		else
 		{
+			// Use the previously defined timer with same name.
+			TimerId = FindTimerIdByName->Id;
+
 			// Map this SpecId to the previously defined timer with same name.
-			SpecIdToTimerIdMap.Add(SpecId, *FindTimerIdByName);
-			return *FindTimerIdByName;
+			SpecIdToTimerIdMap.Add(SpecId, TimerId);
 		}
+
+		return TimerId;
 	}
 	else
 	{
 		// No, a timer with same name was not defined (or we do not want to merge by name).
+
+		uint32 TimerId;
+
 		// Check if SpecId is already mapped to timer.
-		uint32* FindTimerId = SpecIdToTimerIdMap.Find(SpecId);
+		const uint32* FindTimerId = SpecIdToTimerIdMap.Find(SpecId);
 		if (FindTimerId)
 		{
 			// Yes, SpecId was already mapped to a timer (ex. as an <unknown> timer).
+			TimerId = *FindTimerId;
+
 			// Update name for mapped timer.
-			EditableTimingProfilerProvider.SetTimerNameAndLocation(*FindTimerId, Name, File, Line);
-			if (bMergeByName)
-			{
-				// Map the name to the timer.
-				ScopeNameToTimerIdMap.Add(Name, *FindTimerId);
-			}
-			return *FindTimerId;
+			EditableTimingProfilerProvider.SetTimerNameAndLocation(TimerId, Name, File, Line);
 		}
 		else
 		{
 			// Define a new CPU timer.
-			uint32 NewTimerId = EditableTimingProfilerProvider.AddCpuTimer(Name, File, Line);
+			TimerId = EditableTimingProfilerProvider.AddCpuTimer(Name, File, Line);
+
 			// Map the SpecId to the timer.
-			SpecIdToTimerIdMap.Add(SpecId, NewTimerId);
-			if (bMergeByName)
-			{
-				// Map the name to the timer.
-				ScopeNameToTimerIdMap.Add(Name, NewTimerId);
-			}
-			return NewTimerId;
+			SpecIdToTimerIdMap.Add(SpecId, TimerId);
 		}
+
+		// Map the name to the timer.
+		ScopeNameToTimerIdMap.Add((uint64)Name, { TimerId, 1 });
+		return TimerId;
 	}
 }
 
diff --git a/Engine/Source/Developer/TraceServices/Private/Analyzers/CpuProfilerTraceAnalysis.h b/Engine/Source/Developer/TraceServices/Private/Analyzers/CpuProfilerTraceAnalysis.h
--- a/Engine/Source/Developer/TraceServices/Private/Analyzers/CpuProfilerTraceAnalysis.h
+++ b/Engine/Source/Developer/TraceServices/Private/Analyzers/CpuProfilerTraceAnalysis.h
@@ -49,17 +49,11 @@ private:
 		bool bShouldIgnorePendingEvents = false; // becomes true when we detect first pending event with incorrect timestamp (i.e < LastCycle)
 	};
 
-	void ProcessBuffer(const FEventTime& EventTime, FThreadState& ThreadState, const uint8* BufferPtr, uint32 BufferSize);
-	void ProcessBufferV2(const FEventTime& EventTime, FThreadState& ThreadState, const uint8* BufferPtr, uint32 BufferSize);
-	void DispatchPendingEvents(uint64& LastCycle, uint64 CurrentCycle, FThreadState& ThreadState, const FPendingEvent*& PendingCursor, int32& RemainingPending, bool bIsBeginEvent);
-	void DispatchRemainingPendingEvents(FThreadState& ThreadState);
-	void EndOpenEvents(FThreadState& ThreadState, double Timestamp);
-	void OnCpuScopeEnter(const FOnEventContext& Context);
-	void OnCpuScopeLeave(const FOnEventContext& Context);
-	uint32 DefineTimer(uint32 SpecId, const TCHAR* TimerName, const TCHAR* File, uint32 Line, bool bMergeByName); // returns the TimerId
-	uint32 DefineNewTimerChecked(uint32 SpecId, const TCHAR* TimerName, const TCHAR* File = nullptr, uint32 Line = 0); // returns the TimerId
-	uint32 GetTimerId(uint32 SpecId);
-	FThreadState& GetThreadState(uint32 ThreadId);
+	struct FTimerInfo
+	{
+		uint32 Id;
+		uint32 Count;
+	};
 
 	enum : uint16
 	{
@@ -71,12 +65,27 @@ private:
 		RouteId_CpuScope,
 	};
 
+	void ProcessBuffer(const FEventTime& EventTime, FThreadState& ThreadState, const uint8* BufferPtr, uint32 BufferSize);
+	void ProcessBufferV2(const FEventTime& EventTime, FThreadState& ThreadState, const uint8* BufferPtr, uint32 BufferSize);
+	void DispatchPendingEvents(uint64& LastCycle, uint64 CurrentCycle, FThreadState& ThreadState, const FPendingEvent*& PendingCursor, int32& RemainingPending, bool bIsBeginEvent);
+	void DispatchRemainingPendingEvents(FThreadState& ThreadState);
+	void EndOpenEvents(FThreadState& ThreadState, double Timestamp);
+	void OnCpuScopeEnter(const FOnEventContext& Context);
+	void OnCpuScopeLeave(const FOnEventContext& Context);
+	uint32 DefineTimer(uint32 SpecId, const TCHAR* TimerName, const TCHAR* File, uint32 Line); // returns the TimerId
+	uint32 DefineNewTimerChecked(uint32 SpecId, const TCHAR* TimerName, const TCHAR* File = nullptr, uint32 Line = 0); // returns the TimerId
+	uint32 GetTimerId(uint32 SpecId);
+	FThreadState& GetThreadState(uint32 ThreadId);
+
+private:
 	IAnalysisSession& Session;
 	IEditableTimingProfilerProvider& EditableTimingProfilerProvider;
 	IEditableThreadProvider& EditableThreadProvider;
+
 	TMap<uint32, FThreadState*> ThreadStatesMap;
-	TMap<uint32, uint32> SpecIdToTimerIdMap;
-	TMap<const TCHAR*, uint32, FDefaultSetAllocator, TStringPointerMapKeyFuncs_DEPRECATED<const TCHAR*, uint32>> ScopeNameToTimerIdMap;
+	TMap<uint32, uint32> SpecIdToTimerIdMap; // SpecId --> TimerId
+	TMap<uint64, FTimerInfo> ScopeNameToTimerIdMap; // (uint64)Name --> FTimerInfo
+
 	uint32 CoroutineTimerId = ~0;
 	uint32 CoroutineUnknownTimerId = ~0;
 	uint64 TotalEventSize = 0;
