Keep Unreal Insights' swimlane time bounds updated if a trace only has counters.

#rnx

diff --git a/Engine/Source/Developer/TraceServices/Private/Analyzers/CountersTraceAnalysis.cpp b/Engine/Source/Developer/TraceServices/Private/Analyzers/CountersTraceAnalysis.cpp
--- a/Engine/Source/Developer/TraceServices/Private/Analyzers/CountersTraceAnalysis.cpp
+++ b/Engine/Source/Developer/TraceServices/Private/Analyzers/CountersTraceAnalysis.cpp
@@ -66,6 +66,7 @@ bool FCountersAnalyzer::OnEvent(uint16 RouteId, EStyle Style, const FOnEventCont
 		IEditableCounter* FindEditableCounter = EditableCountersMap.FindRef(CounterId);
 		if (ensure(FindEditableCounter))
 		{
+			Session.UpdateDurationSeconds(Timestamp);
 			FindEditableCounter->SetValue(Timestamp, Value);
 		}
 		break;
@@ -78,6 +79,7 @@ bool FCountersAnalyzer::OnEvent(uint16 RouteId, EStyle Style, const FOnEventCont
 		IEditableCounter* FindEditableCounter = EditableCountersMap.FindRef(CounterId);
 		if (ensure(FindEditableCounter))
 		{
+			Session.UpdateDurationSeconds(Timestamp);
 			FindEditableCounter->SetValue(Timestamp, Value);
 		}
 		break;
