[Insights] Fixed crash in Memory Insights analysis when analyzing a corrupted trace ("check" detection for an address allocated multiple times).

diff --git a/Engine/Source/Developer/TraceServices/Private/Model/AllocationsProvider.cpp b/Engine/Source/Developer/TraceServices/Private/Model/AllocationsProvider.cpp
--- a/Engine/Source/Developer/TraceServices/Private/Model/AllocationsProvider.cpp
+++ b/Engine/Source/Developer/TraceServices/Private/Model/AllocationsProvider.cpp
@@ -2203,7 +2203,10 @@ void FAllocationsProvider::EditUnmarkAllocationAsHeap(uint32 ThreadId, double Ti
 			}
 		}
 
-		check(RootHeap.LiveAllocs->FindRef(Address) == nullptr);
+#if INSIGHTS_VALIDATE_ALLOC_EVENTS && INSIGHTS_DOUBLE_ALLOC_FREE_PREVIOUS
+		// This can fail if Address is allocated multiple times. See INSIGHTS_VALIDATE_ALLOC_EVENTS.
+		ensure(RootHeap.LiveAllocs->FindRef(Address) == nullptr);
+#endif
 
 		// Mark allocation as a "heap" allocation.
 		Alloc->Flags = Alloc->Flags | EMemoryTraceHeapAllocationFlags::Heap;
