[Insights] Memory Insights: Automatically removed the "CustomName" parent tag (in order to make the list of tags consistent with one from "LLM Tags" panel).

diff --git a/Engine/Source/Developer/TraceServices/Private/Model/AllocationsProvider.cpp b/Engine/Source/Developer/TraceServices/Private/Model/AllocationsProvider.cpp
--- a/Engine/Source/Developer/TraceServices/Private/Model/AllocationsProvider.cpp
+++ b/Engine/Source/Developer/TraceServices/Private/Model/AllocationsProvider.cpp
@@ -205,6 +205,17 @@ void FTagTracker::AddTagSpec(TagIdType InTag, TagIdType InParentTag, const TCHAR
 		return;
 	}
 
+	// Identify the special "CustomName" tag.
+	if (FCString::Strcmp(InDisplay, TEXT("CustomName")) == 0)
+	{
+		CustomNameTag = InTag;
+	}
+	// Remove the parent tag if it is "CustomName".
+	if (CustomNameTag != InvalidTagId && InParentTag == CustomNameTag)
+	{
+		InParentTag = InvalidTagId;
+	}
+
 	if (!InDisplay || *InDisplay == TEXT('\0'))
 	{
 		++NumErrors;
diff --git a/Engine/Source/Developer/TraceServices/Private/Model/AllocationsProvider.h b/Engine/Source/Developer/TraceServices/Private/Model/AllocationsProvider.h
--- a/Engine/Source/Developer/TraceServices/Private/Model/AllocationsProvider.h
+++ b/Engine/Source/Developer/TraceServices/Private/Model/AllocationsProvider.h
@@ -86,6 +86,7 @@ private:
 	TMap<uint32, FThreadState> TrackerThreadStates;
 	TMap<TagIdType, FTagEntry> TagMap;
 	TArray<TTuple<TagIdType, FString>> PendingTags;
+	TagIdType CustomNameTag = InvalidTagId;
 	uint32 NumErrors = 0;
 	uint32 NumWarnings = 0;
 };
