[Insights]
  - Changed the log message at end of analysis that displayes number of warnings to be a warning (instead of an error).
  - Replaced usage of uint64_t with uint64.

#rnx

diff --git a/Engine/Source/Developer/TraceServices/Private/Analyzers/MiscTraceAnalysis.cpp b/Engine/Source/Developer/TraceServices/Private/Analyzers/MiscTraceAnalysis.cpp
--- a/Engine/Source/Developer/TraceServices/Private/Analyzers/MiscTraceAnalysis.cpp
+++ b/Engine/Source/Developer/TraceServices/Private/Analyzers/MiscTraceAnalysis.cpp
@@ -45,22 +45,30 @@ void FMiscTraceAnalyzer::OnAnalysisBegin(const FOnAnalysisContext& Context)
 {
 	auto& Builder = Context.InterfaceBuilder;
 
+	// Threads
 	Builder.RouteEvent(RouteId_RegisterGameThread, "Misc", "RegisterGameThread");
 	Builder.RouteEvent(RouteId_CreateThread, "Misc", "CreateThread");
 	Builder.RouteEvent(RouteId_SetThreadGroup, "Misc", "SetThreadGroup");
 	Builder.RouteEvent(RouteId_BeginThreadGroupScope, "Misc", "BeginThreadGroupScope");
 	Builder.RouteEvent(RouteId_EndThreadGroupScope, "Misc", "EndThreadGroupScope");
+
+	// Frames
 	Builder.RouteEvent(RouteId_BeginFrame, "Misc", "BeginFrame");
 	Builder.RouteEvent(RouteId_EndFrame, "Misc", "EndFrame");
 	Builder.RouteEvent(RouteId_BeginGameFrame, "Misc", "BeginGameFrame");
 	Builder.RouteEvent(RouteId_EndGameFrame, "Misc", "EndGameFrame");
 	Builder.RouteEvent(RouteId_BeginRenderFrame, "Misc", "BeginRenderFrame");
 	Builder.RouteEvent(RouteId_EndRenderFrame, "Misc", "EndRenderFrame");
+
+	// Channels
 	Builder.RouteEvent(RouteId_ChannelAnnounce, "Trace", "ChannelAnnounce");
 	Builder.RouteEvent(RouteId_ChannelToggle, "Trace", "ChannelToggle");
+
+	// Screenshots
 	Builder.RouteEvent(RouteId_ScreenshotHeader, "Misc", "ScreenshotHeader");
 	Builder.RouteEvent(RouteId_ScreenshotChunk, "Misc", "ScreenshotChunk");
 
+	// Timing Regions
 	Builder.RouteEvent(RouteId_RegionBegin, "Misc", "RegionBegin");
 	Builder.RouteEvent(RouteId_RegionBeginWithId, "Misc", "RegionBeginWithId");
 	Builder.RouteEvent(RouteId_RegionEnd, "Misc", "RegionEnd");
@@ -255,7 +263,6 @@ bool FMiscTraceAnalyzer::OnEvent(uint16 RouteId, EStyle Style, const FOnEventCon
 		check(EventData.GetString("RegionName", Name));
 		FProviderEditScopeLock RegionProviderScopedLock(RegionProvider);
 		RegionProvider.AppendRegionBegin(*Name, Context.EventTime.AsSeconds(Cycle));
-
 		break;
 	}
 
@@ -266,10 +273,9 @@ bool FMiscTraceAnalyzer::OnEvent(uint16 RouteId, EStyle Style, const FOnEventCon
 		check(EventData.GetString("RegionName", Name));
 		FProviderEditScopeLock RegionProviderScopedLock(RegionProvider);
 		RegionProvider.AppendRegionBeginWithId(*Name, CycleAndId, Context.EventTime.AsSeconds(CycleAndId));
-
 		break;
 	}
-		
+
 	case RouteId_RegionEnd:
 	{
 		uint64 Cycle = EventData.GetValue<uint64>("Cycle");
@@ -279,11 +285,11 @@ bool FMiscTraceAnalyzer::OnEvent(uint16 RouteId, EStyle Style, const FOnEventCon
 		RegionProvider.AppendRegionEnd(*Name, Context.EventTime.AsSeconds(Cycle));
 		break;
 	}
-		
+
 	case RouteId_RegionEndWithId:
 	{
 		uint64 Cycle = EventData.GetValue<uint64>("Cycle");
-		uint64_t RegionId = EventData.GetValue<uint64_t>("RegionId", 0);
+		uint64 RegionId = EventData.GetValue<uint64>("RegionId", 0);
 		FProviderEditScopeLock RegionProviderScopedLock(RegionProvider);
 		RegionProvider.AppendRegionEndWithId(RegionId, Context.EventTime.AsSeconds(Cycle));
 		break;
diff --git a/Engine/Source/Developer/TraceServices/Private/Model/Regions.cpp b/Engine/Source/Developer/TraceServices/Private/Model/Regions.cpp
--- a/Engine/Source/Developer/TraceServices/Private/Model/Regions.cpp
+++ b/Engine/Source/Developer/TraceServices/Private/Model/Regions.cpp
@@ -43,7 +43,7 @@ const FRegionLane* FRegionProvider::GetLane(int32 index) const
 	}
 	return nullptr;
 }
-	
+
 void FRegionProvider::AppendRegionBegin(const TCHAR* Name, double Time)
 {
 	EditAccessCheck();
@@ -51,7 +51,7 @@ void FRegionProvider::AppendRegionBegin(const TCHAR* Name, double Time)
 	check(Name)
 
 	FTimeRegion** OpenRegion = OpenRegionsByName.Find(Name);
-	
+
 	if (OpenRegion)
 	{
 		++NumWarnings;
@@ -68,19 +68,19 @@ void FRegionProvider::AppendRegionBegin(const TCHAR* Name, double Time)
 	OpenRegionsByName.Add(NewRegion->Text, NewRegion);
 }
 
-void FRegionProvider::AppendRegionBeginWithId(const TCHAR* Name, uint64_t Id, double Time)
+void FRegionProvider::AppendRegionBeginWithId(const TCHAR* Name, uint64 Id, double Time)
 {
 	EditAccessCheck();
 
 	check(Name && Id)
 	FTimeRegion** OpenRegion = OpenRegionsById.Find(Id);
-	
+
 	if (OpenRegion)
 	{
 		++NumWarnings;
 		if (NumWarnings <= MaxWarningMessages)
 		{
-			UE_LOG(LogTraceServices, Warning, TEXT("[Regions] A region begin event (BeginTime=%f, Name=\"%s\", Id=\"%llu\") was encountered while a region with same name is already open."), Time, Name, Id)
+			UE_LOG(LogTraceServices, Warning, TEXT("[Regions] A region begin event (BeginTime=%f, Name=\"%s\", Id=%llu) was encountered while a region with same name is already open."), Time, Name, Id)
 		}
 
 		// Automatically end the previous region.
@@ -90,14 +90,14 @@ void FRegionProvider::AppendRegionBeginWithId(const TCHAR* Name, uint64_t Id, do
 	FTimeRegion* NewRegion = InsertNewRegion(Time, Name, Id);
 	OpenRegionsById.Add(Id, NewRegion);
 }
-	
+
 void FRegionProvider::AppendRegionEnd(const TCHAR* Name, double Time)
 {
 	EditAccessCheck();
 
 	check(Name)
 	FTimeRegion** OpenRegionPos = OpenRegionsByName.Find(Name);
-	
+
 	if (!OpenRegionPos)
 	{
 		++NumWarnings;
@@ -105,15 +105,15 @@ void FRegionProvider::AppendRegionEnd(const TCHAR* Name, double Time)
 		{
 			UE_LOG(LogTraceServices, Warning, TEXT("[Regions] A region end event (EndTime=%f, Name=\"%s\") was encountered without having seen a matching region begin event first."), Time, Name)
 		}
-		
+
 		AppendRegionBegin(Name, Time);
 		OpenRegionPos = OpenRegionsByName.Find(Name);
 		check(OpenRegionPos);
 	}
-	
+
 	FTimeRegion* OpenRegion = *OpenRegionPos;
 	check(OpenRegion);
-	
+
 	OpenRegion->EndTime = Time;
 
 	OpenRegionsByName.Remove(OpenRegion->Text);
@@ -126,32 +126,32 @@ void FRegionProvider::AppendRegionEnd(const TCHAR* Name, double Time)
 	}
 }
 
-void FRegionProvider::AppendRegionEndWithId(uint64_t Id, double Time)
+void FRegionProvider::AppendRegionEndWithId(uint64 Id, double Time)
 {
 	EditAccessCheck();
 
 	check(Id)
 	FTimeRegion** OpenRegionPos = OpenRegionsById.Find(Id);
-	
+
 	if (!OpenRegionPos)
 	{
 		++NumWarnings;
 		if (NumWarnings <= MaxWarningMessages)
 		{
-			UE_LOG(LogTraceServices, Warning, TEXT("[Regions] A region end event (EndTime=%f, Id=\"%llu\") was encountered without having seen a matching region begin event first."), Time, Id)
+			UE_LOG(LogTraceServices, Warning, TEXT("[Regions] A region end event (EndTime=%f, Id=%llu) was encountered without having seen a matching region begin event first."), Time, Id)
 		}
 
 		// Automatically create a new region.
 		// Generates a display name if we're missing a begin and are closing by ID
-		FString GeneratedName = FString::Printf(TEXT("Unknown Region (missing begin, Id=\"%llu\")"), Id);
+		FString GeneratedName = FString::Printf(TEXT("Unknown Region (missing begin, Id=%llu)"), Id);
 		AppendRegionBeginWithId(*GeneratedName, Id, Time);
 		OpenRegionPos = OpenRegionsById.Find(Id);
 		check(OpenRegionPos);
 	}
-	
+
 	FTimeRegion* OpenRegion = *OpenRegionPos;
 	check(OpenRegion);
-	
+
 	OpenRegion->EndTime = Time;
 
 	OpenRegionsById.Remove(OpenRegion->Id);
@@ -163,7 +163,7 @@ void FRegionProvider::AppendRegionEndWithId(uint64_t Id, double Time)
 		Session.UpdateDurationSeconds(Time);
 	}
 }
-	
+
 void FRegionProvider::OnAnalysisSessionEnded()
 {
 	EditAccessCheck();
@@ -175,15 +175,19 @@ void FRegionProvider::OnAnalysisSessionEnded()
 		++NumWarnings;
 		if (NumWarnings <= MaxWarningMessages)
 		{
-			UE_LOG(LogTraceServices, Warning, TEXT("[Regions] A region (BeginTime=%f, Name=\"%s\", Id=\"%llu\") was never closed."), Region->BeginTime, Region->Text, Region->Id)
+			UE_LOG(LogTraceServices, Warning, TEXT("[Regions] A region (BeginTime=%f, Name=\"%s\", Id=%llu) was never closed."), Region->BeginTime, Region->Text, Region->Id)
 		}
 	};
 	Algo::ForEach(OpenRegionsById, printOpenRegionMessage);
 	Algo::ForEach(OpenRegionsByName, printOpenRegionMessage);
 
-	if (NumWarnings > 0 || NumErrors > 0)
+	if (NumWarnings > 0)
+	{
+		UE_LOG(LogTraceServices, Warning, TEXT("[Regions] %u warnings"), NumWarnings);
+	}
+	if (NumErrors > 0)
 	{
-		UE_LOG(LogTraceServices, Error, TEXT("[Regions] %u warnings; %u errors"), NumWarnings, NumErrors);
+		UE_LOG(LogTraceServices, Error, TEXT("[Regions] %u errors"), NumErrors);
 	}
 
 	uint64 TotalRegionCount = GetRegionCount();
@@ -217,7 +221,7 @@ int32 FRegionProvider::CalculateRegionDepth(const FTimeRegion& Region) const
 	return NewDepth;
 }
 
-FTimeRegion* FRegionProvider::InsertNewRegion(double BeginTime, const TCHAR* Name, uint64_t Id)
+FTimeRegion* FRegionProvider::InsertNewRegion(double BeginTime, const TCHAR* Name, uint64 Id)
 {
 	FTimeRegion Region;
 	Region.BeginTime = BeginTime;
diff --git a/Engine/Source/Developer/TraceServices/Private/Model/RegionsPrivate.h b/Engine/Source/Developer/TraceServices/Private/Model/RegionsPrivate.h
--- a/Engine/Source/Developer/TraceServices/Private/Model/RegionsPrivate.h
+++ b/Engine/Source/Developer/TraceServices/Private/Model/RegionsPrivate.h
@@ -48,9 +48,9 @@ public:
 	virtual void EditAccessCheck() const override { Lock.WriteAccessCheck(GRegionsProviderLockState); }
 
 	virtual void AppendRegionBegin(const TCHAR* Name, double Time) override;
-	virtual void AppendRegionBeginWithId(const TCHAR* Name, uint64_t Id, double Time) override;
+	virtual void AppendRegionBeginWithId(const TCHAR* Name, uint64 Id, double Time) override;
 	virtual void AppendRegionEnd(const TCHAR* Name, double Time) override;
-	virtual void AppendRegionEndWithId(uint64_t Id, double Time) override;
+	virtual void AppendRegionEndWithId(uint64 Id, double Time) override;
 
 	virtual void OnAnalysisSessionEnded() override;
 
@@ -60,8 +60,8 @@ private:
 	/// Update the depth member of a region to allow overlapping regions to be displayed on separate lanes.
 	int32 CalculateRegionDepth(const FTimeRegion& Item) const;
 	/// Calculates depth, inserts a new region into the correct lane and updates session time.
-	FTimeRegion* InsertNewRegion(double BeginTime, const TCHAR* Name, uint64_t Id);
-	
+	FTimeRegion* InsertNewRegion(double BeginTime, const TCHAR* Name, uint64 Id);
+
 private:
 	mutable FProviderLock Lock;
 
@@ -69,7 +69,7 @@ private:
 
 	// Open regions inside lanes
 	TMap<FStringView, FTimeRegion*> OpenRegionsByName;
-	TMap<uint64_t, FTimeRegion*> OpenRegionsById;
+	TMap<uint64, FTimeRegion*> OpenRegionsById;
 
 	// Closed regions
 	TArray<FRegionLane> Lanes;
diff --git a/Engine/Source/Developer/TraceServices/Public/TraceServices/Model/Regions.h b/Engine/Source/Developer/TraceServices/Public/TraceServices/Model/Regions.h
--- a/Engine/Source/Developer/TraceServices/Public/TraceServices/Model/Regions.h
+++ b/Engine/Source/Developer/TraceServices/Public/TraceServices/Model/Regions.h
@@ -17,7 +17,7 @@ struct TRACESERVICES_API FTimeRegion
 	double EndTime = std::numeric_limits<double>::infinity();
 	const TCHAR* Text = nullptr;
 	// Id will be zero if the region is identified by Name only
-	uint64_t Id = 0;
+	uint64 Id = 0;
 	int32 Depth = -1;
 };
 
@@ -100,7 +100,7 @@ public:
 	 * @param Time		The time in seconds of the begin event of this region.
 	 */
 	virtual void AppendRegionBegin(const TCHAR* Name, double Time) = 0;
-	
+
 	/**
 	 * Append a new begin event of a region from the trace session.
 	 * Id will be used to uniquely identify the new region.
@@ -109,8 +109,8 @@ public:
 	 * @param Id		The Id of the region. Used to uniquely identify regions with the same name.
 	 * @param Time		The time in seconds of the begin event of this region.
 	 */
-	virtual void AppendRegionBeginWithId(const TCHAR* Name, uint64_t Id, double Time) = 0;
-	
+	virtual void AppendRegionBeginWithId(const TCHAR* Name, uint64 Id, double Time) = 0;
+
 	/**
 	 * Append a new end event of a region from the trace session.
 	 * Prefer opening/closing regions with an Id, since string names are not unique.
@@ -119,15 +119,15 @@ public:
 	 * @param Time		The time in seconds of the end event of this region.
 	 */
 	virtual void AppendRegionEnd(const TCHAR* Name, double Time) = 0;
-	
+
 	/**
 	 * Append a new end event of a region from the trace session.
 	 * The region is identified by Id.
 	 * @param Id		The Id of the region.
 	 * @param Time		The time in seconds of the end event of this region.
 	 */
-	virtual void AppendRegionEndWithId(const uint64_t Id, double Time) = 0;
-	
+	virtual void AppendRegionEndWithId(const uint64 Id, double Time) = 0;
+
 	/**
 	 * Called from the analyzer once all events have been processed.
 	 * Allows post-processing and error reporting for regions that were never closed.
diff --git a/Engine/Source/Runtime/Core/Private/ProfilingDebugging/MiscTrace.cpp b/Engine/Source/Runtime/Core/Private/ProfilingDebugging/MiscTrace.cpp
--- a/Engine/Source/Runtime/Core/Private/ProfilingDebugging/MiscTrace.cpp
+++ b/Engine/Source/Runtime/Core/Private/ProfilingDebugging/MiscTrace.cpp
@@ -93,9 +93,9 @@ void FMiscTrace::OutputBeginRegion(const TCHAR* RegionName)
 		<< RegionBegin.RegionName(RegionName);
 }
 
-uint64_t FMiscTrace::OutputBeginRegionWithId(const TCHAR* RegionName)
+uint64 FMiscTrace::OutputBeginRegionWithId(const TCHAR* RegionName)
 {
-	const uint64_t CycleAndId = FPlatformTime::Cycles64();
+	const uint64 CycleAndId = FPlatformTime::Cycles64();
 	UE_TRACE_LOG(Misc, RegionBeginWithId, RegionChannel)
 		<< RegionBeginWithId.CycleAndId(CycleAndId)
 		<< RegionBeginWithId.RegionName(RegionName);
@@ -109,11 +109,11 @@ void FMiscTrace::OutputEndRegion(const TCHAR* RegionName)
 		<< RegionEnd.RegionName(RegionName);
 }
 
-void FMiscTrace::OutputEndRegionWithId(uint64_t RegionId)
+void FMiscTrace::OutputEndRegionWithId(uint64 RegionId)
 {
 	UE_TRACE_LOG(Misc, RegionEndWithId, RegionChannel)
 		<< RegionEndWithId.Cycle(FPlatformTime::Cycles64())
-		<< RegionEndWithId.RegionId(RegionId);	
+		<< RegionEndWithId.RegionId(RegionId);
 }
 
 void FMiscTrace::OutputBookmarkInternal(const void* BookmarkPoint, uint16 EncodedFormatArgsSize, uint8* EncodedFormatArgs)
diff --git a/Engine/Source/Runtime/Core/Public/ProfilingDebugging/MiscTrace.h b/Engine/Source/Runtime/Core/Public/ProfilingDebugging/MiscTrace.h
--- a/Engine/Source/Runtime/Core/Public/ProfilingDebugging/MiscTrace.h
+++ b/Engine/Source/Runtime/Core/Public/ProfilingDebugging/MiscTrace.h
@@ -74,10 +74,10 @@ struct FMiscTrace
 	}
 
 	CORE_API static void OutputBeginRegion(const TCHAR* RegionName);
-	[[nodiscard]] CORE_API static uint64_t OutputBeginRegionWithId(const TCHAR* RegionName);
+	[[nodiscard]] CORE_API static uint64 OutputBeginRegionWithId(const TCHAR* RegionName);
 	CORE_API static void OutputEndRegion(const TCHAR* RegionName);
-	CORE_API static void OutputEndRegionWithId(uint64_t RegionId);
-	
+	CORE_API static void OutputEndRegionWithId(uint64 RegionId);
+
 	CORE_API static void OutputBeginFrame(ETraceFrameType FrameType);
 	CORE_API static void OutputEndFrame(ETraceFrameType FrameType);
 
