Fix crashes in UnrealInsights cache files.

diff --git a/Engine/Source/Developer/TraceServices/Private/AnalysisCache.cpp b/Engine/Source/Developer/TraceServices/Private/AnalysisCache.cpp
--- a/Engine/Source/Developer/TraceServices/Private/AnalysisCache.cpp
+++ b/Engine/Source/Developer/TraceServices/Private/AnalysisCache.cpp
@@ -19,7 +19,7 @@
 #include "Serialization/CompactBinaryWriter.h"
 #include "Serialization/MemoryReader.h"
 
-DEFINE_LOG_CATEGORY_STATIC(LogAnalysisCache, Log, All);
+DEFINE_LOG_CATEGORY_STATIC(LogAnalysisCache, VeryVerbose, All);
 
 namespace TraceServices {
 //////////////////////////////////////////////////////////////////////
@@ -46,12 +46,14 @@ FAnalysisCache::FFileContents::FFileContents(const TCHAR* FilePath)
 			if (const bool Result = Load(); !Result)
 			{
 				UE_LOG(LogAnalysisCache, Error, TEXT("Failed to open cache file table of contents."));
+				bTransientMode = true;
 				//todo: Recover by deleting file?
 			}
 
 			// Additional sanity check. A common error scenario is that Insights crashed after writing block but before
 			// committing them to the table of contents. Detect that scenario here.
-			if (FileSize > ReservedSize && Blocks.IsEmpty())
+			uint32 MinimalExpectedSizePerVersion[] = {0, ReservedSizeV1, ReservedSizeV2};
+			if (Version < UE_ARRAY_COUNT(MinimalExpectedSizePerVersion) && FileSize > MinimalExpectedSizePerVersion[Version] && Blocks.IsEmpty())
 			{
 				UE_LOG(LogAnalysisCache, Error, TEXT("Cache file has written several blocks but table of contents contains no blocks. This is likely caused by abnormal program termination. Please delete \"%s\". Putting cache in transient mode."), *CacheFilePath);
 				IndexEntries.Empty();
@@ -78,6 +80,12 @@ FAnalysisCache::FFileContents::FFileContents(const TCHAR* FilePath)
 			bTransientMode = true;
 		}
 	}
+
+	// If we haven't established a version use current
+	if (!Version)
+	{
+		Version = CurrentVersion;
+	}
 }
 
 //////////////////////////////////////////////////////////////////////
@@ -144,9 +152,21 @@ bool FAnalysisCache::FFileContents::Save()
 
 	File->Seek(0);
 
-	FCbWriter Writer;
-	Writer.BeginObject();
-	Writer << "Version" << CurrentVersion;
+	check(Version > 0); // Version should always be set here
+	switch(Version)
+	{
+	case 1:
+		return SaveVersion1(File);
+	case 2:
+		return SaveVersion2(File);
+	}
+
+	return false;
+}
+
+//////////////////////////////////////////////////////////////////////
+void FAnalysisCache::FFileContents::SaveVersion1Index(FCbWriter& Writer)
+{
 	Writer.BeginArray(ANSITEXTVIEW("Index"));
 	for (auto Entry : IndexEntries)
 	{
@@ -164,11 +184,47 @@ bool FAnalysisCache::FFileContents::Save()
 		Writer.AddBinary(&Entry, sizeof(FBlockEntry));
 	}
 	Writer.EndArray();
+	
+}
+
+//////////////////////////////////////////////////////////////////////
+bool FAnalysisCache::FFileContents::SaveVersion1(IFileHandle* File)
+{
+	FCbWriter Writer;
+	Writer.BeginObject();
+	Writer << "Version" << 1;
+	SaveVersion1Index(Writer);
 	Writer.EndObject();
 	
 	FCbPackage Package(Writer.Save().AsObject());
 
-	FUniqueBuffer Buffer = FUniqueBuffer::Alloc(ReservedSize);
+	FUniqueBuffer Buffer = FUniqueBuffer::Alloc(ReservedSizeV1);
+	FBufferWriter BufferWriter(Buffer.GetData(), Buffer.GetSize());
+	Package.Save(BufferWriter);
+	
+	return File->Write((uint8*)Buffer.GetData(), Buffer.GetSize());
+}
+
+//////////////////////////////////////////////////////////////////////
+bool FAnalysisCache::FFileContents::SaveVersion2(IFileHandle* File)
+{
+	check (Version != 0);
+	
+	// Header
+	ANSICHAR Magic[] = "UC";
+	File->Write((uint8*)Magic,2);
+	File->Write((uint8*)&Version, 4);
+	checkSlow(File->Tell() == IndexOffset);
+
+	// Write the index
+	FCbWriter Writer;
+	Writer.BeginObject();
+	SaveVersion1Index(Writer);
+	Writer.EndObject();
+	
+	FCbPackage Package(Writer.Save().AsObject());
+	
+	FUniqueBuffer Buffer = FUniqueBuffer::Alloc(ReservedSizeV2);
 	FBufferWriter BufferWriter(Buffer.GetData(), Buffer.GetSize());
 	Package.Save(BufferWriter);
 	
@@ -183,28 +239,65 @@ bool FAnalysisCache::FFileContents::Load()
 	{
 		return true;
 	}
-	
-	File->Seek(0);
-	
-	FUniqueBuffer Buffer = FUniqueBuffer::Alloc(ReservedSize);
-	File->Read((uint8*)Buffer.GetData(), Buffer.GetSize());
-	
-	FMemoryReaderView Ar(MakeArrayView<uint8>((uint8*)Buffer.GetData(), IntCastChecked<int32>(Buffer.GetSize())));
-	
-	FCbPackage Package;
-	if (!Package.TryLoad(Ar))
+
+	// Read header
+	Version = ReadHeader(File);
+	if (Version > CurrentVersion)
 	{
+		UE_LOG(LogAnalysisCache, Warning, TEXT("Cache file of unknown version (%u), cannot load."), Version);
 		return false;
 	}
 
-	uint32 PackageVersion = Package.GetObject().Find("Version").AsUInt32();
-	UE_LOG(LogAnalysisCache, Display, TEXT("Cache file (version %u) loaded."), PackageVersion);
-	if (PackageVersion != CurrentVersion)
+	UE_LOG(LogAnalysisCache, Display, TEXT("Loading cache file (version %u)."), Version);
+
+	switch(Version)
 	{
-		// todo: Handle this better
-		return false;
+	case 1:
+		return LoadVersion1(File);
+	case 2:
+		return LoadVersion2(File);
+	}
+
+	return false;
+}
+
+//////////////////////////////////////////////////////////////////////
+uint32 FAnalysisCache::FFileContents::ReadHeader(IFileHandle* File)
+{
+	uint32 Version = 0;
+	uint8 Header[IndexOffset];
+	File->Seek(0);
+	File->Read((uint8*)&Header, IndexOffset);
+	if (Header[0] != 'U' || Header[1] != 'C')
+	{
+		// Version 1 lacked a magic header
+		Version = 1;
+		File->Seek(0);
 	}
+	else
+	{
+		// Expect version number after magic
+		Version = *(uint32*)&Header[2];
+	}
+
+	return Version;
+}
 
+//////////////////////////////////////////////////////////////////////
+bool FAnalysisCache::FFileContents::WriteHeader(IFileHandle* File, uint32 Version)
+{
+	bool bSuccess = true;
+	File->Seek(0);
+	ANSICHAR Magic[] = "UC";
+	bSuccess &= File->Write((uint8*)Magic,2);
+	bSuccess &= File->Write((uint8*)&Version, 4);
+	checkSlow(File->Tell() == IndexOffset);
+	return bSuccess;
+}
+
+//////////////////////////////////////////////////////////////////////
+void FAnalysisCache::FFileContents::LoadVersion1Index(const FCbPackage& Package)
+{
 	FCbArrayView IndexArray = Package.GetObject().Find(ANSITEXTVIEW("Index")).AsArrayView();
 	IndexEntries.Reserve(static_cast<int32>(IndexArray.Num()));
 	for (FCbFieldView IndexEntry : IndexArray)
@@ -232,10 +325,49 @@ bool FAnalysisCache::FFileContents::Load()
 		FMutableMemoryView BlockView = FMutableMemoryView(&Block, sizeof(FBlockEntry));
 		BlockView.CopyFrom(BlockEntryView.AsBinaryView());
 	}
+}
+
+//////////////////////////////////////////////////////////////////////
+bool FAnalysisCache::FFileContents::LoadVersion1(IFileHandle* File)
+{
+	FUniqueBuffer Buffer = FUniqueBuffer::Alloc(ReservedSizeV1);
+	File->Read((uint8*)Buffer.GetData(), Buffer.GetSize());
+	
+	FMemoryReaderView Ar(MakeArrayView<uint8>((uint8*)Buffer.GetData(), IntCastChecked<int32>(Buffer.GetSize())));
+	
+	FCbPackage Package;
+	if (!Package.TryLoad(Ar))
+	{
+		return false;
+	}
+	
+	Version = Package.GetObject().Find("Version").AsUInt32();
+	
+	LoadVersion1Index(Package);
 
 	return true;
 }
 
+	
+//////////////////////////////////////////////////////////////////////
+bool FAnalysisCache::FFileContents::LoadVersion2(IFileHandle* File)
+{
+	FUniqueBuffer Buffer = FUniqueBuffer::Alloc(ReservedSizeV2);
+	File->Read((uint8*)Buffer.GetData(), Buffer.GetSize());
+
+	FMemoryReaderView Ar(MakeArrayView<uint8>((uint8*)Buffer.GetData(), IntCastChecked<int32>(Buffer.GetSize())));
+
+	FCbPackage Package;
+	if (!Package.TryLoad(Ar))
+	{
+		return false;
+	}
+
+	LoadVersion1Index(Package);
+
+	return true;
+}
+	
 //////////////////////////////////////////////////////////////////////
 uint64 FAnalysisCache::FFileContents::UpdateBlock(FMemoryView Block, BlockKeyType BlockKey)
 {
@@ -272,7 +404,7 @@ uint64 FAnalysisCache::FFileContents::UpdateBlock(FMemoryView Block, BlockKeyTyp
 		const bool bSeekResult = File->SeekFromEnd(0);
 		check(bSeekResult);
 		const uint64 Offset = File->Tell();
-		check(Offset >= ReservedSize);
+		checkf(Offset >= ReservedSizeV2, TEXT("Offset (%ull) is larger than reserved size (%u)"), Offset, ReservedSizeV2);
 		if(!File->Write((const uint8*) Block.GetData(), Block.GetSize()))
 		{
 			UE_LOG(LogAnalysisCache, Error, TEXT("Failed to update block 0x%x at offset %u kb"), BlockKey, Offset / 1024);
diff --git a/Engine/Source/Developer/TraceServices/Private/AnalysisCache.h b/Engine/Source/Developer/TraceServices/Private/AnalysisCache.h
--- a/Engine/Source/Developer/TraceServices/Private/AnalysisCache.h
+++ b/Engine/Source/Developer/TraceServices/Private/AnalysisCache.h
@@ -10,6 +10,7 @@
 
 class IMappedFileRegion;
 class IMappedFileHandle;
+class FCbPackage;
 
 namespace TraceServices {
 
@@ -40,9 +41,14 @@ private:
 	struct FFileContents
 	{
 		// Current file format version
-		static constexpr uint32 CurrentVersion = 1;
-		// Number of bytes at the start of the file reserved for table of contents.
-		static constexpr uint32 ReservedSize = 512 * 1024;
+		static constexpr uint32 CurrentVersion = 2;
+		// Number of bytes at the start of the file reserved for table of contents and header.
+		static constexpr uint32 ReservedSizeV1 = 512 * 1024;
+		static constexpr uint32 ReservedSizeV2 = 1024 * 1024;
+		// Offset of the table of contents in bytes
+		static constexpr uint32 IndexOffset = 6;
+		// Size for the index object 
+		static constexpr uint32 ReservedSizeIndex = ReservedSizeV2 - IndexOffset;
 		// Number of bytes allocated for user data (per named entry)
 		static constexpr uint32 UserDataSize = 64;
 
@@ -83,6 +89,17 @@ private:
 		TArray<FBlockEntry> Blocks;
 		// When this is set blocks or table of contents are never written.
 		bool bTransientMode = false;
+		uint32 Version = 0;
+
+	private:
+		static uint32 ReadHeader(IFileHandle* File);
+		static bool WriteHeader(IFileHandle* File, uint32 Version);
+		void LoadVersion1Index(const FCbPackage& Package);
+		void SaveVersion1Index(FCbWriter& Writer);
+		bool SaveVersion1(IFileHandle* File);
+		bool SaveVersion2(IFileHandle* File);
+		bool LoadVersion1(IFileHandle* File);
+		bool LoadVersion2(IFileHandle* File);
 	};
 
 	/**
