Revert log level increment

diff --git a/Engine/Source/Developer/TraceServices/Private/AnalysisCache.cpp b/Engine/Source/Developer/TraceServices/Private/AnalysisCache.cpp
--- a/Engine/Source/Developer/TraceServices/Private/AnalysisCache.cpp
+++ b/Engine/Source/Developer/TraceServices/Private/AnalysisCache.cpp
@@ -19,7 +19,7 @@
 #include "Serialization/CompactBinaryWriter.h"
 #include "Serialization/MemoryReader.h"
 
-DEFINE_LOG_CATEGORY_STATIC(LogAnalysisCache, VeryVerbose, All);
+DEFINE_LOG_CATEGORY_STATIC(LogAnalysisCache, Log, All);
 
 namespace TraceServices {
 //////////////////////////////////////////////////////////////////////
