[Insights] Added AddAnalyzer(TSharedRef<UE::Trace::IAnalyzer> Analyzer) to IAnalysisSession. It uses an internal wrapper for the analyzer.

#rnx

diff --git a/Engine/Source/Developer/TraceServices/Private/AnalysisService.cpp b/Engine/Source/Developer/TraceServices/Private/AnalysisService.cpp
--- a/Engine/Source/Developer/TraceServices/Private/AnalysisService.cpp
+++ b/Engine/Source/Developer/TraceServices/Private/AnalysisService.cpp
@@ -3,11 +3,17 @@
 #include "TraceServices/AnalysisService.h"
 #include "AnalysisServicePrivate.h"
 
+#include "HAL/PlatformFile.h"
+
+// TraceAnalysis
+#include "Trace/Analysis.h"
+#include "Trace/DataStream.h"
+
+// TraceServices
 #include "Analyzers/BookmarksTraceAnalysis.h"
 #include "Analyzers/LogTraceAnalysis.h"
 #include "Analyzers/MiscTraceAnalysis.h"
 #include "Analyzers/StringsAnalyzer.h"
-#include "HAL/PlatformFile.h"
 #include "Model/BookmarksPrivate.h"
 #include "Model/Channel.h"
 #include "Model/CountersPrivate.h"
@@ -20,13 +26,42 @@
 #include "Model/ScreenshotProviderPrivate.h"
 #include "Model/ThreadsPrivate.h"
 #include "ModuleServicePrivate.h"
-#include "Trace/Analysis.h"
 #include "Trace/Analyzer.h"
-#include "Trace/DataStream.h"
 
 namespace TraceServices
 {
 
+class FAnalyzerWrapper : public UE::Trace::IAnalyzer
+{
+public:
+	FAnalyzerWrapper(TSharedRef<UE::Trace::IAnalyzer> InAnalyzer)
+		: InnerAnalyzer(InAnalyzer)
+	{
+	}
+
+	virtual ~FAnalyzerWrapper()
+	{
+	}
+
+	virtual void OnAnalysisBegin(const FOnAnalysisContext& Context) override
+	{
+		InnerAnalyzer->OnAnalysisBegin(Context);
+	}
+
+	virtual void OnAnalysisEnd(/*const FOnAnalysisEndContext& Context*/) override
+	{
+		InnerAnalyzer->OnAnalysisEnd(/*Context*/);
+	}
+
+	virtual bool OnEvent(uint16 RouteId, EStyle Style, const FOnEventContext& Context) override
+	{
+		return InnerAnalyzer->OnEvent(RouteId, Style, Context);
+	}
+
+private:
+	TSharedRef<UE::Trace::IAnalyzer> InnerAnalyzer;
+};
+
 // if IProvider ever gets member data, it will duplicate state because of "the diamond problem" with multiple inheritance
 // if you hit this you should consider the implications for classes that implement multiple providers
 static_assert(sizeof(IProvider) == sizeof(uintptr_t));
@@ -173,7 +208,7 @@ uint32 FAnalysisSession::GetNumPendingMessages() const
 	return PendingMessagesCount.load();
 }
 
-TArray<FAnalysisMessage> FAnalysisSession::DrainPendingMessages() 
+TArray<FAnalysisMessage> FAnalysisSession::DrainPendingMessages()
 {
 	Lock.WriteAccessCheck();
 	PendingMessagesCount.store(0);
@@ -182,9 +217,15 @@ TArray<FAnalysisMessage> FAnalysisSession::DrainPendingMessages()
 
 void FAnalysisSession::AddAnalyzer(UE::Trace::IAnalyzer* Analyzer)
 {
+	check(Analyzer != nullptr);
 	Analyzers.Add(Analyzer);
 }
 
+void FAnalysisSession::AddAnalyzer(TSharedRef<UE::Trace::IAnalyzer> Analyzer)
+{
+	Analyzers.Add(new FAnalyzerWrapper(Analyzer));
+}
+
 void FAnalysisSession::AddProvider(const FName& InName, TSharedPtr<IProvider> Provider, TSharedPtr<IEditableProvider> EditableProvider)
 {
 	Providers.Add(InName, MakeTuple(Provider, EditableProvider));
@@ -225,7 +266,7 @@ void FAnalysisSession::OnAnalysisMessage(UE::Trace::EAnalysisMessageSeverity InS
 	case UE::Trace::EAnalysisMessageSeverity::Warning: Severity = EMessageSeverity::Type::Warning; break;
 	case UE::Trace::EAnalysisMessageSeverity::Info: Severity = EMessageSeverity::Type::Info; break;
 	}
-	
+
 	Lock.BeginEdit();
 	PendingMessages.Push(FAnalysisMessage { Severity, FString(InMessage)});
 	PendingMessagesCount.fetch_add(1);
diff --git a/Engine/Source/Developer/TraceServices/Private/AnalysisServicePrivate.h b/Engine/Source/Developer/TraceServices/Private/AnalysisServicePrivate.h
--- a/Engine/Source/Developer/TraceServices/Private/AnalysisServicePrivate.h
+++ b/Engine/Source/Developer/TraceServices/Private/AnalysisServicePrivate.h
@@ -3,13 +3,17 @@
 #pragma once
 
 #include "AnalysisCache.h"
-#include "TraceServices/Model/AnalysisSession.h"
-#include "TraceServices/AnalysisService.h"
 #include "Containers/HashTable.h"
 #include "Common/SlabAllocator.h"
 #include "Common/StringStore.h"
+
+// TraceAnalysis
 #include "Trace/Analysis.h"
 
+// TraceServices
+#include "TraceServices/Model/AnalysisSession.h"
+#include "TraceServices/AnalysisService.h"
+
 namespace TraceServices
 {
 
@@ -75,6 +79,7 @@ public:
 	virtual void WriteAccessCheck() override { return Lock.WriteAccessCheck(); }
 
 	virtual void AddAnalyzer(UE::Trace::IAnalyzer* Analyzer) override;
+	virtual void AddAnalyzer(TSharedRef<UE::Trace::IAnalyzer> Analyzer) override;
 	virtual void AddProvider(const FName& Name, TSharedPtr<IProvider> Provider, TSharedPtr<IEditableProvider> EditableProvider = nullptr) override;
 
 	const TArray<UE::Trace::IAnalyzer*> ReadAnalyzers() { return Analyzers; }
diff --git a/Engine/Source/Developer/TraceServices/Public/TraceServices/Model/AnalysisSession.h b/Engine/Source/Developer/TraceServices/Public/TraceServices/Model/AnalysisSession.h
--- a/Engine/Source/Developer/TraceServices/Public/TraceServices/Model/AnalysisSession.h
+++ b/Engine/Source/Developer/TraceServices/Public/TraceServices/Model/AnalysisSession.h
@@ -24,7 +24,7 @@ namespace TraceServices
 
 class ILinearAllocator;
 class IAnalysisCache;
-	
+
 class IProvider
 {
 public:
@@ -67,18 +67,16 @@ struct FTraceSessionMetadata
 struct FAnalysisMessage
 {
 	EMessageSeverity::Type Severity;
-	FString Message;	
+	FString Message;
 };
 
 class IAnalysisSession
 {
 public:
-
-	
 	virtual ~IAnalysisSession() = default;
 	virtual void Stop(bool bAndWait=false) const = 0;
 	virtual void Wait() const = 0;
-	
+
 	virtual const TCHAR* GetName() const = 0;
 	virtual uint32 GetTraceId() const = 0;
 	virtual bool IsAnalysisComplete() const = 0;
@@ -111,15 +109,15 @@ public:
 	virtual void BeginEdit() = 0;
 	virtual void EndEdit() = 0;
 	virtual void WriteAccessCheck() = 0;
-	
+
 	virtual ILinearAllocator& GetLinearAllocator() = 0;
 	virtual const TCHAR* StoreString(const TCHAR* String) = 0;
 	virtual const TCHAR* StoreString(const FStringView& String) = 0;
 
 	virtual IAnalysisCache& GetCache() = 0;
-	
-	virtual void AddAnalyzer(UE::Trace::IAnalyzer* Analyzer) = 0;
 
+	virtual void AddAnalyzer(UE::Trace::IAnalyzer* Analyzer) = 0;
+	virtual void AddAnalyzer(TSharedRef<UE::Trace::IAnalyzer> Analyzer) = 0;
 	virtual void AddProvider(const FName& Name, TSharedPtr<IProvider> Provider, TSharedPtr<IEditableProvider> EditableProvider = nullptr) = 0;
 
 	template<typename ProviderType>
@@ -155,7 +153,7 @@ struct FAnalysisSessionReadScope
 private:
 	const IAnalysisSession& AnalysisSession;
 };
-	
+
 struct FAnalysisSessionEditScope
 {
 	FAnalysisSessionEditScope(IAnalysisSession& InAnalysisSession)
@@ -171,5 +169,5 @@ struct FAnalysisSessionEditScope
 
 	IAnalysisSession& AnalysisSession;
 };
-	
+
 } // namespace TraceServices
