Add gpu counters for number of draws and number of primitives to Insights

diff --git a/Engine/Source/Developer/TraceServices/Private/Analyzers/GpuProfilerTraceAnalysis.cpp b/Engine/Source/Developer/TraceServices/Private/Analyzers/GpuProfilerTraceAnalysis.cpp
--- a/Engine/Source/Developer/TraceServices/Private/Analyzers/GpuProfilerTraceAnalysis.cpp
+++ b/Engine/Source/Developer/TraceServices/Private/Analyzers/GpuProfilerTraceAnalysis.cpp
@@ -6,14 +6,16 @@
 #include "CborWriter.h"
 #include "Common/Utils.h"
 #include "HAL/LowLevelMemTracker.h"
+#include "Model/CountersPrivate.h"
 #include "Serialization/MemoryWriter.h"
 
 namespace TraceServices
 {
 
-FGpuProfilerAnalyzer::FGpuProfilerAnalyzer(FAnalysisSession& InSession, FTimingProfilerProvider& InTimingProfilerProvider)
+FGpuProfilerAnalyzer::FGpuProfilerAnalyzer(FAnalysisSession& InSession, FTimingProfilerProvider& InTimingProfilerProvider, IEditableCounterProvider& InEditableCounterProvider)
 	: Session(InSession)
 	, TimingProfilerProvider(InTimingProfilerProvider)
+	, EditableCounterProvider(InEditableCounterProvider)
 {
 }
 
@@ -86,6 +88,10 @@ bool FGpuProfilerAnalyzer::OnEvent(uint16 RouteId, EStyle Style, const FOnEventC
 		const TCHAR* PersistentName = Session.StoreString(Name);
 
 		TimingProfilerProvider.AddGpuQueue(QueueId, GPU, Index, Type, PersistentName);
+
+		FQueue& Queue = GetOrAddQueue(QueueId);
+		InitCountersDesc(Queue, GPU, Index, *Name);
+
 		break;
 	}
 	case RouteId_EventFrameBoundary:
@@ -97,6 +103,19 @@ bool FGpuProfilerAnalyzer::OnEvent(uint16 RouteId, EStyle Style, const FOnEventC
 
 		FQueue& Queue = GetOrAddQueue(QueueId);
 		Queue.FrameNumber = FrameNumber;
+
+		if (Queue.LastTime > 0)
+		{
+			FAnalysisSessionEditScope _(Session);
+
+			Queue.NumDrawsCounter->SetValue(Queue.LastTime, (int64)Queue.NumDraws);
+			Queue.NumPrimitivesCounter->SetValue(Queue.LastTime, (int64)Queue.NumPrimitives);
+		}
+
+		Queue.NumDraws = 0;
+		Queue.NumPrimitives = 0;
+		Queue.LastTime = 0;
+
 		break;
 	}
 	case RouteId_EventBeginBreadcrumb:
@@ -181,7 +200,12 @@ bool FGpuProfilerAnalyzer::OnEvent(uint16 RouteId, EStyle Style, const FOnEventC
 		const uint32 QueueId = EventData.GetValue<uint32>("QueueId");
 		const uint32 NumDraws = EventData.GetValue<uint32>("NumDraws");
 		const uint32 NumPrimitives = EventData.GetValue<uint32>("NumPrimitives");
-		//TODO: counters
+
+		FQueue& Queue = GetOrAddQueue(QueueId);
+
+		Queue.NumDraws += NumDraws;
+		Queue.NumPrimitives += NumPrimitives;
+
 		break;
 	}
 
@@ -342,14 +366,17 @@ FGpuProfilerAnalyzer::FQueue& FGpuProfilerAnalyzer::GetOrAddQueue(uint32 QueueId
 	{
 		return *FoundQueue;
 	}
+
 	FQueue& NewQueue = Queues.Add(QueueId);
 	NewQueue.Id = QueueId;
+	InitCounters(NewQueue);
 	return NewQueue;
 }
 
 void FGpuProfilerAnalyzer::BeginEvent(FQueue& Queue, int32 StackIndex, double BeginEventTime, uint32 BeginEventTimerId)
 {
 	Queue.Stack[StackIndex].Push({ BeginEventTime, BeginEventTimerId });
+	Queue.LastTime = BeginEventTime;
 
 	FAnalysisSessionEditScope _(Session);
 	IEditableTimeline<FTimingProfilerEvent>* Timeline = (StackIndex == 1) ?
@@ -375,6 +402,8 @@ void FGpuProfilerAnalyzer::BeginEvent(FQueue& Queue, int32 StackIndex, double Be
 
 void FGpuProfilerAnalyzer::EndEvent(FQueue& Queue, int32 StackIndex, double EndEventTime, uint32 EndEventTimerId)
 {
+	Queue.LastTime = EndEventTime;
+
 	if (Queue.Stack[StackIndex].IsEmpty())
 	{
 		return;
@@ -420,4 +449,25 @@ void FGpuProfilerAnalyzer::EndEvent(FQueue& Queue, int32 StackIndex, double EndE
 	}
 }
 
+void FGpuProfilerAnalyzer::InitCounters(FQueue& FoundQueue)
+{
+	FoundQueue.NumDrawsCounter = EditableCounterProvider.CreateEditableCounter();
+	FoundQueue.NumDrawsCounter->SetIsFloatingPoint(false);
+
+	FoundQueue.NumPrimitivesCounter = EditableCounterProvider.CreateEditableCounter();
+	FoundQueue.NumPrimitivesCounter->SetIsFloatingPoint(false);
+}
+
+void FGpuProfilerAnalyzer::InitCountersDesc(FQueue& FoundQueue, uint8 Gpu, uint8 Index, const TCHAR* Name)
+{
+	const FString DrawsCounterName = FString::Printf(TEXT("NumDraws_GPU%u_%s_%u"), Gpu, Name, Index);
+	const FString PrimitivesCounterName = FString::Printf(TEXT("NumPrimitives_GPU%u_%s_%u"), Gpu, Name, Index);
+
+	FoundQueue.NumDrawsCounter->SetName(Session.StoreString(DrawsCounterName));
+	FoundQueue.NumDrawsCounter->SetDescription(TEXT("The number of draw calls on the specified queue."));
+
+	FoundQueue.NumPrimitivesCounter->SetName(Session.StoreString(PrimitivesCounterName));
+	FoundQueue.NumPrimitivesCounter->SetDescription(TEXT("The number of primitives on the specified queue."));
+}
+
 } // namespace TraceServices
diff --git a/Engine/Source/Developer/TraceServices/Private/Analyzers/GpuProfilerTraceAnalysis.h b/Engine/Source/Developer/TraceServices/Private/Analyzers/GpuProfilerTraceAnalysis.h
--- a/Engine/Source/Developer/TraceServices/Private/Analyzers/GpuProfilerTraceAnalysis.h
+++ b/Engine/Source/Developer/TraceServices/Private/Analyzers/GpuProfilerTraceAnalysis.h
@@ -7,8 +7,9 @@
 
 namespace TraceServices
 {
-
 	class FAnalysisSession;
+	class IEditableCounter;
+	class IEditableCounterProvider;
 
 class FGpuProfilerAnalyzer
 	: public UE::Trace::IAnalyzer
@@ -43,10 +44,15 @@ private:
 		uint32 Id = 0;
 		uint32 FrameNumber = 0;
 		TArray<FOpenEvent> Stack[2]; // [0] = breadcrumbs, [1] = work
+		IEditableCounter* NumDrawsCounter;
+		IEditableCounter* NumPrimitivesCounter;
+		uint64 NumDraws = 0;
+		uint64 NumPrimitives = 0;
+		double LastTime = 0.0f;
 	};
 
 public:
-	FGpuProfilerAnalyzer(FAnalysisSession& Session, FTimingProfilerProvider& TimingProfilerProvider);
+	FGpuProfilerAnalyzer(FAnalysisSession& Session, FTimingProfilerProvider& TimingProfilerProvider, IEditableCounterProvider& EditableCounterProvider);
 	virtual void OnAnalysisBegin(const FOnAnalysisContext& Context) override;
 	virtual void OnAnalysisEnd() override;
 	virtual bool OnEvent(uint16 RouteId, EStyle Style, const FOnEventContext& Context) override;
@@ -58,9 +64,13 @@ private:
 	void BeginEvent(FQueue& Queue, int32 StackIndex, double Time, uint32 TimerId);
 	void EndEvent(FQueue& Queue, int32 StackIndex, double Time, uint32 TimerId);
 
+	void InitCounters(FQueue& FoundQueue);
+	void InitCountersDesc(FQueue& FoundQueue, uint8 Gpu, uint8 Index, const TCHAR* Name);
+
 private:
 	FAnalysisSession& Session;
 	FTimingProfilerProvider& TimingProfilerProvider;
+	IEditableCounterProvider& EditableCounterProvider;
 
 	//////////////////////////////////////////////////
 	// The New GPU Profiler
diff --git a/Engine/Source/Developer/TraceServices/Private/Modules/TimingProfilerModule.cpp b/Engine/Source/Developer/TraceServices/Private/Modules/TimingProfilerModule.cpp
--- a/Engine/Source/Developer/TraceServices/Private/Modules/TimingProfilerModule.cpp
+++ b/Engine/Source/Developer/TraceServices/Private/Modules/TimingProfilerModule.cpp
@@ -5,6 +5,7 @@
 #include "AnalysisServicePrivate.h"
 #include "Analyzers/CpuProfilerTraceAnalysis.h"
 #include "Analyzers/GpuProfilerTraceAnalysis.h"
+#include "Model/CountersPrivate.h"
 #include "Model/ThreadsPrivate.h"
 #include "Model/TimingProfilerPrivate.h"
 
@@ -28,7 +29,7 @@ void FTimingProfilerModule::OnAnalysisBegin(IAnalysisSession& InSession)
 	TSharedPtr<FTimingProfilerProvider> TimingProfilerProvider = MakeShared<FTimingProfilerProvider>(Session);
 	Session.AddProvider(GetTimingProfilerProviderName(), TimingProfilerProvider, TimingProfilerProvider);
 	Session.AddAnalyzer(new FCpuProfilerAnalyzer(Session, *TimingProfilerProvider, EditableThreadProvider));
-	Session.AddAnalyzer(new FGpuProfilerAnalyzer(Session, *TimingProfilerProvider));
+	Session.AddAnalyzer(new FGpuProfilerAnalyzer(Session, *TimingProfilerProvider, EditCounterProvider(Session)));
 }
 
 void FTimingProfilerModule::GetLoggers(TArray<const TCHAR*>& OutLoggers)
