Network insights/NetTrace
Bumped NetDebugName DebugNameId to use 32 bits.

diff --git a/Engine/Source/Developer/TraceServices/Private/Analyzers/NetTraceAnalyzer.cpp b/Engine/Source/Developer/TraceServices/Private/Analyzers/NetTraceAnalyzer.cpp
--- a/Engine/Source/Developer/TraceServices/Private/Analyzers/NetTraceAnalyzer.cpp
+++ b/Engine/Source/Developer/TraceServices/Private/Analyzers/NetTraceAnalyzer.cpp
@@ -21,6 +21,7 @@ enum ENetTraceAnalyzerVersion
 	ENetTraceAnalyzerVersion_BunchChannelIndex = 2,
 	ENetTraceAnalyzerVersion_BunchChannelInfo = 3,
 	ENetTraceAnalyzerVersion_FixedBunchSizeEncoding = 4,
+	ENetTraceAnalyzerVersion_DebugNameIndexIs32Bits = 5,
 };
 
 
@@ -67,9 +68,9 @@ void FNetTraceAnalyzer::OnAnalysisEnd()
 {
 }
 
-uint32 FNetTraceAnalyzer::GetTracedEventTypeIndex(uint16 NameIndex, uint8 Level)
+uint32 FNetTraceAnalyzer::GetTracedEventTypeIndex(FNetProfilerNameIndexType NameIndex, uint8 Level)
 {
-	const uint32 TracedEventTypeKey = uint32(NameIndex << 8U | Level);
+	const uint64 TracedEventTypeKey = uint64(NameIndex << 8U | Level);
 
 	if (const uint32* NetProfilerEventTypeIndex = TraceEventTypeToNetProfilerEventTypeIndexMap.Find(TracedEventTypeKey))
 	{
@@ -122,7 +123,7 @@ bool FNetTraceAnalyzer::OnEvent(uint16 RouteId, EStyle Style, const FOnEventCont
 
 		case RouteId_NameEvent:
 		{
-			uint16 TraceNameId = EventData.GetValue<uint16>("NameId");
+			FNetProfilerNameIndexType TraceNameId = EventData.GetValue<FNetProfilerNameIndexType>("NameId");
 			if (TracedNameIdToNetProfilerNameIdMap.Contains(TraceNameId))
 			{
 				// need to update the name
@@ -285,7 +286,7 @@ void FNetTraceAnalyzer::HandlePacketContentEvent(const FOnEventContext& Context,
 
 						// Fill in the object data we currently have
 						ObjectInstance.LifeTime.Begin = GetLastTimestamp();
-						ObjectInstance.NameIndex = static_cast<uint16>(PendingNameIndex);
+						ObjectInstance.NameIndex = PendingNameIndex;
 						ObjectInstance.NetObjectId = DecodedNameOrObjectId;
 						ObjectInstance.TypeId = 0;
 
@@ -298,7 +299,7 @@ void FNetTraceAnalyzer::HandlePacketContentEvent(const FOnEventContext& Context,
 				}
 				else if (DecodedEventType == EContentEventType::NameId)
 				{
-					if (const uint32* NetProfilerNameIndex = TracedNameIdToNetProfilerNameIdMap.Find(IntCastChecked<uint16>(DecodedNameOrObjectId)))
+					if (const FNetProfilerNameIndexType* NetProfilerNameIndex = TracedNameIdToNetProfilerNameIdMap.Find(IntCastChecked<FNetProfilerNameIndexType>(DecodedNameOrObjectId)))
 					{
 						Event.NameIndex = *NetProfilerNameIndex;
 					}
@@ -309,13 +310,13 @@ void FNetTraceAnalyzer::HandlePacketContentEvent(const FOnEventContext& Context,
 				}
 
 				// EventTypeIndex does not match NameIndex as we might see the same name on different levels
-				Event.EventTypeIndex = GetTracedEventTypeIndex(static_cast<uint16>(Event.NameIndex), static_cast<uint8>(Event.Level));
+				Event.EventTypeIndex = GetTracedEventTypeIndex(Event.NameIndex, static_cast<uint8>(Event.Level));
 			}
 			break;
 
 			case EContentEventType::BunchEvent:
 			{
-				const uint16 DecodedNameId = IntCastChecked<uint16>(FTraceAnalyzerUtils::Decode7bit(BufferPtr));
+				const FNetProfilerNameIndexType DecodedNameId = IntCastChecked<FNetProfilerNameIndexType>(FTraceAnalyzerUtils::Decode7bit(BufferPtr));
 
 				uint32 DecodedBunchBits = 0U;
 				if (NetTraceVersion >= ENetTraceAnalyzerVersion_FixedBunchSizeEncoding)
@@ -329,7 +330,7 @@ void FNetTraceAnalyzer::HandlePacketContentEvent(const FOnEventContext& Context,
 					DecodedBunchBits = IntCastChecked<uint32>((uint32)DecodedEventEndPos + (uint32)DecodedEventStartPos);
 				}
 
-				const uint32* NetProfilerNameIndex = DecodedNameId ? TracedNameIdToNetProfilerNameIdMap.Find(DecodedNameId) : nullptr;
+				const FNetProfilerNameIndexType* NetProfilerNameIndex = DecodedNameId ? TracedNameIdToNetProfilerNameIdMap.Find(DecodedNameId) : nullptr;
 
 				FBunchInfo BunchInfo;
 
@@ -337,7 +338,7 @@ void FNetTraceAnalyzer::HandlePacketContentEvent(const FOnEventContext& Context,
 				BunchInfo.HeaderBits = 0U;
 				BunchInfo.BunchBits = DecodedBunchBits;
 				BunchInfo.FirstBunchEventIndex = Events.Num();
-				BunchInfo.NameIndex = NetProfilerNameIndex ? IntCastChecked<uint16>(*NetProfilerNameIndex) : 0U;
+				BunchInfo.NameIndex = NetProfilerNameIndex ? *NetProfilerNameIndex : 0U;
 
 				BunchInfos.Add(BunchInfo);
 
@@ -378,8 +379,8 @@ void FNetTraceAnalyzer::HandlePacketContentEvent(const FOnEventContext& Context,
 						}
 						else
 						{
-							const uint32* ExistingChannelNameIndex = GameInstanceState->ChannelNames.Find(BunchInfo.BunchInfo.ChannelIndex);
-							BunchInfo.NameIndex = ExistingChannelNameIndex ? IntCastChecked<uint16>(*ExistingChannelNameIndex) : 0U;
+							const FNetProfilerNameIndexType* ExistingChannelNameIndex = GameInstanceState->ChannelNames.Find(BunchInfo.BunchInfo.ChannelIndex);
+							BunchInfo.NameIndex = ExistingChannelNameIndex ? *ExistingChannelNameIndex : 0U;
 						}
 
 						BunchInfo.BunchInfo.bIsValid = 1U;
@@ -405,7 +406,7 @@ void FNetTraceAnalyzer::AddEvent(TPagedArray<FNetProfilerContentEvent>& Events,
 {
 	FNetProfilerContentEvent& Event = Events.PushBack();
 
-	Event.EventTypeIndex = GetTracedEventTypeIndex(IntCastChecked<uint16>(InEvent.NameIndex), IntCastChecked<uint8>(InEvent.Level + LevelOffset));
+	Event.EventTypeIndex = GetTracedEventTypeIndex(InEvent.NameIndex, IntCastChecked<uint8>(InEvent.Level + LevelOffset));
 	Event.NameIndex = InEvent.NameIndex;
 	Event.ObjectInstanceIndex = InEvent.ObjectInstanceIndex;
 	Event.StartPos = InEvent.StartPos + Offset;
@@ -414,11 +415,11 @@ void FNetTraceAnalyzer::AddEvent(TPagedArray<FNetProfilerContentEvent>& Events,
 	Event.BunchInfo = InEvent.BunchInfo;
 }
 
-void FNetTraceAnalyzer::AddEvent(TPagedArray<FNetProfilerContentEvent>& Events, uint32 StartPos, uint32 EndPos, uint32 Level, uint32 NameIndex, FNetProfilerBunchInfo BunchInfo)
+void FNetTraceAnalyzer::AddEvent(TPagedArray<FNetProfilerContentEvent>& Events, uint32 StartPos, uint32 EndPos, uint32 Level, FNetProfilerNameIndexType NameIndex, FNetProfilerBunchInfo BunchInfo)
 {
 	FNetProfilerContentEvent& Event = Events.PushBack();
 
-	Event.EventTypeIndex = GetTracedEventTypeIndex(IntCastChecked<uint16>(NameIndex), IntCastChecked<uint8>(Level));
+	Event.EventTypeIndex = GetTracedEventTypeIndex(NameIndex, IntCastChecked<uint8>(Level));
 	Event.NameIndex = NameIndex;
 	Event.ObjectInstanceIndex = 0;
 	Event.StartPos = StartPos;
@@ -681,20 +682,20 @@ void FNetTraceAnalyzer::HandleConnectionCreatedEvent(const FOnEventContext& Cont
 	ConnectionState->CurrentPacketStatsStartIndex[ENetProfilerConnectionMode::Incoming] = 0U;
 }
 
-uint32 FNetTraceAnalyzer::GetOrCreateNetProfilerStatsCounterTypeIndex(uint32 NameId, ENetProfilerStatsCounterType StatsType)
+uint32 FNetTraceAnalyzer::GetOrCreateNetProfilerStatsCounterTypeIndex(FNetProfilerNameIndexType NameId, ENetProfilerStatsCounterType StatsType)
 {
-	if (const uint32* ExistingNetProfilerStatsCounterTypeIndex = TraceNetStatsCounterIdToNetProfilerStatsCounterTypeIndexMap.Find(static_cast<uint16>(NameId)))
+	if (const uint32* ExistingNetProfilerStatsCounterTypeIndex = TraceNetStatsCounterIdToNetProfilerStatsCounterTypeIndexMap.Find(NameId))
 	{
 		return *ExistingNetProfilerStatsCounterTypeIndex;
 	}
 	else
 	{
 		// Add new counter type
-		const uint32* NetProfilerNameIndex = TracedNameIdToNetProfilerNameIdMap.Find(static_cast<uint16>(NameId));
-		const uint32 NameIndex = NetProfilerNameIndex ? *NetProfilerNameIndex : 0u;
+		const FNetProfilerNameIndexType* NetProfilerNameIndex = TracedNameIdToNetProfilerNameIdMap.Find(NameId);
+		const FNetProfilerNameIndexType NameIndex = NetProfilerNameIndex ? *NetProfilerNameIndex : 0u;
 
 		uint32 NetProfilerStatsCounterTypeIndex = NetProfilerProvider.AddNetProfilerStatsCounterType(NameIndex, StatsType);
-		TraceNetStatsCounterIdToNetProfilerStatsCounterTypeIndexMap.Add(static_cast<uint16>(NameId), NetProfilerStatsCounterTypeIndex);
+		TraceNetStatsCounterIdToNetProfilerStatsCounterTypeIndexMap.Add(NameId, NetProfilerStatsCounterTypeIndex);
 
 		return NetProfilerStatsCounterTypeIndex;
 	}
@@ -705,7 +706,7 @@ void FNetTraceAnalyzer::HandlePacketStatsCounterEvent(const FOnEventContext& Con
 	const uint32 StatsCounterValue = EventData.GetValue<uint32>("StatsValue");
 	const uint8 GameInstanceId = EventData.GetValue<uint8>("GameInstanceId");
 	const uint16 ConnectionId = EventData.GetValue<uint16>("ConnectionId");
-	const uint16 NameId = EventData.GetValue<uint16>("NameId");
+	const FNetProfilerNameIndexType NameId = EventData.GetValue<FNetProfilerNameIndexType>("NameId");
 
 	FNetProfilerStats Stats;
 	Stats.StatsCounterTypeIndex = GetOrCreateNetProfilerStatsCounterTypeIndex(NameId, ENetProfilerStatsCounterType::Packet);
@@ -724,7 +725,7 @@ void FNetTraceAnalyzer::HandleFrameStatsCounterEvent(const FOnEventContext& Cont
 	const uint64 TimestampCycles = EventData.GetValue<uint64>("Timestamp");
 	const uint32 StatsCounterValue = EventData.GetValue<uint32>("StatsValue");
 	const uint8 GameInstanceId = EventData.GetValue<uint8>("GameInstanceId");
-	const uint16 NameId = EventData.GetValue<uint16>("NameId");
+	const FNetProfilerNameIndexType NameId = EventData.GetValue<FNetProfilerNameIndexType>("NameId");
 
 	// Update LastTimestamp, later on we will be able to get timestamps piggybacked from other analyzers
 	LastTimeStamp = Context.EventTime.AsSeconds(TimestampCycles);
@@ -813,12 +814,12 @@ void FNetTraceAnalyzer::HandleObjectCreatedEvent(const FOnEventContext& Context,
 	const uint64 TypeId = EventData.GetValue<uint64>("TypeId");
 	const uint64 ObjectId = EventData.GetValue<uint64>("ObjectId");
 	//const uint32 OwnerId = EventData.GetValue<uint32>("OwnerId");
-	const uint16 NameId = EventData.GetValue<uint16>("NameId");
+	const FNetProfilerNameIndexType NameId = EventData.GetValue<FNetProfilerNameIndexType>("NameId");
 	const uint8 GameInstanceId = EventData.GetValue<uint8>("GameInstanceId");
 
 	TSharedRef<FNetTraceGameInstanceState> GameInstanceState = GetOrCreateActiveGameInstanceState(GameInstanceId);
-	const uint32* NetProfilerNameIndex = TracedNameIdToNetProfilerNameIdMap.Find(NameId);
-	const uint32 NameIndex = NetProfilerNameIndex ? *NetProfilerNameIndex : 0u;
+	const FNetProfilerNameIndexType* NetProfilerNameIndex = TracedNameIdToNetProfilerNameIdMap.Find(NameId);
+	const FNetProfilerNameIndexType NameIndex = NetProfilerNameIndex ? *NetProfilerNameIndex : 0u;
 
 	if (FNetTraceActiveObjectState* ActiveObjectInstance = GameInstanceState->ActiveObjects.Find(ObjectId))
 	{
@@ -832,7 +833,7 @@ void FNetTraceAnalyzer::HandleObjectCreatedEvent(const FOnEventContext& Context,
 				ExistingInstance->TypeId = TypeId;
 
 				// Update name in both the persistent instance and the active one
-				ExistingInstance->NameIndex = static_cast<uint16>(NameIndex);
+				ExistingInstance->NameIndex = NameIndex;
 				ActiveObjectInstance->NameIndex = NameIndex;
 
 				return;
@@ -849,7 +850,7 @@ void FNetTraceAnalyzer::HandleObjectCreatedEvent(const FOnEventContext& Context,
 
 	// Fill in object data
 	ObjectInstance.LifeTime.Begin = GetLastTimestamp();
-	ObjectInstance.NameIndex = IntCastChecked<uint16>(NameIndex);
+	ObjectInstance.NameIndex = NameIndex;
 	ObjectInstance.NetObjectId = ObjectId;
 	ObjectInstance.TypeId = TypeId;
 
diff --git a/Engine/Source/Developer/TraceServices/Private/Analyzers/NetTraceAnalyzer.h b/Engine/Source/Developer/TraceServices/Private/Analyzers/NetTraceAnalyzer.h
--- a/Engine/Source/Developer/TraceServices/Private/Analyzers/NetTraceAnalyzer.h
+++ b/Engine/Source/Developer/TraceServices/Private/Analyzers/NetTraceAnalyzer.h
@@ -61,15 +61,15 @@ private:
 	uint32 NetTraceReporterVersion;
 
 	// Default names
-	uint32 BunchHeaderNameIndex;
-	uint32 PendingNameIndex;
+	FNetProfilerNameIndexType BunchHeaderNameIndex;
+	FNetProfilerNameIndexType PendingNameIndex;
 
 	// Shared for trace
-	TMap<uint16, uint32> TracedNameIdToNetProfilerNameIdMap;
+	TMap<FNetProfilerNameIndexType, FNetProfilerNameIndexType> TracedNameIdToNetProfilerNameIdMap;
 
-	TMap<uint32, uint32> TraceEventTypeToNetProfilerEventTypeIndexMap;
+	TMap<uint64, uint32> TraceEventTypeToNetProfilerEventTypeIndexMap;
 
-	TMap<uint16, uint32> TraceNetStatsCounterIdToNetProfilerStatsCounterTypeIndexMap;
+	TMap<FNetProfilerNameIndexType, uint32> TraceNetStatsCounterIdToNetProfilerStatsCounterTypeIndexMap;
 
 	struct FBunchInfo
 	{
@@ -78,7 +78,7 @@ private:
 		uint32 HeaderBits;
 		int32 FirstBunchEventIndex;
 		uint32 EventCount;
-		uint16 NameIndex;
+		FNetProfilerNameIndexType NameIndex;
 	};
 
 	struct FNetTraceConnectionState
@@ -105,7 +105,7 @@ private:
 	struct FNetTraceActiveObjectState
 	{
 		uint32 ObjectIndex;
-		uint32 NameIndex;
+		FNetProfilerNameIndexType NameIndex;
 	};
 
 	// Map from instance id to index in the persistent instance array
@@ -113,7 +113,7 @@ private:
 	{
 		TMap<uint32, TSharedRef<FNetTraceConnectionState>> ActiveConnections;
 		TMap<uint64, FNetTraceActiveObjectState> ActiveObjects;
-		TMap<int32, uint32> ChannelNames;
+		TMap<int32, FNetProfilerNameIndexType> ChannelNames;
 
 		// StatsCounters reported during the current frame
 		TArray<FNetProfilerStats> FrameStatsCounters;
@@ -125,8 +125,8 @@ private:
 
 private:
 
-	uint32 GetTracedEventTypeIndex(uint16 NameIndex, uint8 Level);
-	uint32 GetOrCreateNetProfilerStatsCounterTypeIndex(uint32 NameId, ENetProfilerStatsCounterType StatsType);
+	uint32 GetTracedEventTypeIndex(FNetProfilerNameIndexType NameIndex, uint8 Level);
+	uint32 GetOrCreateNetProfilerStatsCounterTypeIndex(FNetProfilerNameIndexType NameId, ENetProfilerStatsCounterType StatsType);
 
 	TSharedRef<FNetTraceGameInstanceState> GetOrCreateActiveGameInstanceState(uint32 GameInstanceId);
 	void DestroyActiveGameInstanceState(uint32 GameInstanceId);
@@ -154,7 +154,7 @@ private:
 	void HandleGameInstanceUpdatedEvent(const FOnEventContext& Context, const FEventData& EventData);
 
 	void AddEvent(TPagedArray<FNetProfilerContentEvent>& Events, const FNetProfilerContentEvent& Event, uint32 Offset, uint32 LevelOffset);
-	void AddEvent(TPagedArray<FNetProfilerContentEvent>& Events, uint32 StartPos, uint32 EndPos, uint32 Level, uint32 NameIndex, FNetProfilerBunchInfo Info);
+	void AddEvent(TPagedArray<FNetProfilerContentEvent>& Events, uint32 StartPos, uint32 EndPos, uint32 Level, FNetProfilerNameIndexType NameIndex, FNetProfilerBunchInfo Info);
 
 private:
 
diff --git a/Engine/Source/Developer/TraceServices/Private/Model/NetProfilerProvider.cpp b/Engine/Source/Developer/TraceServices/Private/Model/NetProfilerProvider.cpp
--- a/Engine/Source/Developer/TraceServices/Private/Model/NetProfilerProvider.cpp
+++ b/Engine/Source/Developer/TraceServices/Private/Model/NetProfilerProvider.cpp
@@ -79,7 +79,7 @@ uint32 FNetProfilerProvider::GetNetTraceVersion() const
 	return NetTraceVersion;
 }
 
-uint32 FNetProfilerProvider::AddNetProfilerName(const TCHAR* Name)
+FNetProfilerNameIndexType FNetProfilerProvider::AddNetProfilerName(const TCHAR* Name)
 {
 	Session.WriteAccessCheck();
 
@@ -90,12 +90,12 @@ uint32 FNetProfilerProvider::AddNetProfilerName(const TCHAR* Name)
 	return NewName.NameIndex;
 }
 
-const FNetProfilerName* FNetProfilerProvider::GetNetProfilerName(uint32 NameIndex) const
+const FNetProfilerName* FNetProfilerProvider::GetNetProfilerName(FNetProfilerNameIndexType NameIndex) const
 {
 	return NameIndex < (uint32)Names.Num() ? &Names[NameIndex] : nullptr;
 }
 
-uint32 FNetProfilerProvider::AddNetProfilerEventType(uint32 NameIndex, uint32 Level)
+uint32 FNetProfilerProvider::AddNetProfilerEventType(FNetProfilerNameIndexType NameIndex, uint32 Level)
 {
 	Session.WriteAccessCheck();
 
@@ -103,7 +103,7 @@ uint32 FNetProfilerProvider::AddNetProfilerEventType(uint32 NameIndex, uint32 Le
 	NewEventType.EventTypeIndex = EventTypes.Num() - 1;
 	NewEventType.NameIndex = NameIndex;
 	NewEventType.Name = GetNetProfilerName(NameIndex)->Name;
-	NewEventType.Level = Level;
+	NewEventType.Level = IntCastChecked<uint16>(Level);
 
 	return NewEventType.EventTypeIndex;
 }
@@ -113,7 +113,7 @@ const FNetProfilerEventType* FNetProfilerProvider::GetNetProfilerEventType(uint3
 	return EventTypeIndex < (uint32)EventTypes.Num() ? &EventTypes[EventTypeIndex] : nullptr;
 }
 
-uint32 FNetProfilerProvider::AddNetProfilerStatsCounterType(uint32 NameIndex, ENetProfilerStatsCounterType Type)
+uint32 FNetProfilerProvider::AddNetProfilerStatsCounterType(FNetProfilerNameIndexType NameIndex, ENetProfilerStatsCounterType Type)
 {
 	Session.WriteAccessCheck();
 
diff --git a/Engine/Source/Developer/TraceServices/Private/Model/NetProfilerProvider.h b/Engine/Source/Developer/TraceServices/Private/Model/NetProfilerProvider.h
--- a/Engine/Source/Developer/TraceServices/Private/Model/NetProfilerProvider.h
+++ b/Engine/Source/Developer/TraceServices/Private/Model/NetProfilerProvider.h
@@ -62,7 +62,7 @@ public:
 
 	void SetNetTraceVersion(uint32 NetTraceVersion);
 
-	uint32 AddNetProfilerName(const TCHAR* Name);
+	FNetProfilerNameIndexType AddNetProfilerName(const TCHAR* Name);
 
 	uint32 AddNetProfilerEventType(uint32 NameIndex, uint32 Level);
 
@@ -134,10 +134,10 @@ public:
 	virtual ITable<FNetProfilerAggregatedStats>* CreateAggregation(uint32 ConnectionIndex, ENetProfilerConnectionMode Mode, uint32 PacketIndexIntervalStart, uint32 PacketIndexIntervalEnd, uint32 StartPosition, uint32 EndPosition) const override;
 	virtual ITable<FNetProfilerAggregatedStatsCounterStats>* CreateStatsCountersAggregation(uint32 ConnectionIndex, ENetProfilerConnectionMode Mode, uint32 PacketIndexIntervalStart, uint32 PacketIndexIntervalEnd) const override;
 
-	const TCHAR* InternalGetNetProfilerName(uint32 NameIndex) const { return Names[NameIndex].Name; }
+	const TCHAR* InternalGetNetProfilerName(FNetProfilerNameIndexType NameIndex) const { return Names[NameIndex].Name; }
 private:
 	
-	const FNetProfilerName* GetNetProfilerName(uint32 ProfilerNameId) const;
+	const FNetProfilerName* GetNetProfilerName(FNetProfilerNameIndexType ProfilerNameId) const;
 	const FNetProfilerEventType* GetNetProfilerEventType(uint32 ProfilerEventTypeId) const;
 	const FNetProfilerStatsCounterType* GetNetProfilerStatsCounterType(uint32 ProfilerStatsCounterTypeId) const;
 
diff --git a/Engine/Source/Developer/TraceServices/Public/TraceServices/Model/NetProfiler.h b/Engine/Source/Developer/TraceServices/Public/TraceServices/Model/NetProfiler.h
--- a/Engine/Source/Developer/TraceServices/Public/TraceServices/Model/NetProfiler.h
+++ b/Engine/Source/Developer/TraceServices/Public/TraceServices/Model/NetProfiler.h
@@ -56,19 +56,22 @@ enum class ENetProfilerConnectionState : uint8
 TRACESERVICES_API const TCHAR* LexToString(const ENetProfilerChannelCloseReason Value);
 TRACESERVICES_API const TCHAR* LexToString(const ENetProfilerConnectionState Value);
 
+// Mirror from FNetDebugNameId on the runtime side.
+typedef uint32 FNetProfilerNameIndexType;
+
 struct FNetProfilerName
 {
 	const TCHAR* Name;	// Name
-	uint32 NameIndex;	// Index in the type array?
+	FNetProfilerNameIndexType NameIndex;	// Index in the type array?
 };
 
 // This is our event type separate per level as the same name might be used on different levels
 struct FNetProfilerEventType
 {
-	uint32 EventTypeIndex;
 	const TCHAR* Name;
-	uint32 NameIndex : 16;
-	uint32 Level : 16;
+	uint32 EventTypeIndex;
+	FNetProfilerNameIndexType NameIndex;	
+	uint16 Level;
 };
 
 typedef double FNetProfilerTimeStamp;
@@ -88,16 +91,11 @@ struct FNetProfilerLifeTime
 
 struct FNetProfilerObjectInstance
 {
-	FNetProfilerLifeTime LifeTime;			// Lifetime of this instance
-	union
-	{
-		UE_DEPRECATED(5.3, "Please use NetObjectId instead")
-		uint32 NetId;						// NetHandleIndex or NetGUID
-		uint64 NetObjectId = 0;				// NetHandleIndex or NetGUID
-	};
-	uint64 TypeId = uint64(0);				// ProtocolIdentifier
-	uint32 ObjectIndex = 0U;				// Index in the object array
-	uint16 NameIndex = 0U;					// Index in the Name array
+	FNetProfilerLifeTime LifeTime;				// Lifetime of this instance
+	uint64 NetObjectId = 0;						// NetHandleIndex or NetGUID
+	uint64 TypeId = uint64(0);					// ProtocolIdentifier
+	uint32 ObjectIndex = 0U;					// Index in the object array
+	FNetProfilerNameIndexType NameIndex = 0U;	// Index in the Name array
 };
 
 union FNetProfilerBunchInfo
@@ -131,16 +129,16 @@ union FNetProfilerBunchInfo
 
 struct FNetProfilerContentEvent
 {
-	uint64 StartPos : 24;		// Inclusive start position in the packet
-	uint64 EndPos : 24;			// Exclusive end position in the packet; BitSize = EndPos - StartPos
-	uint64 Level : 8;			// Level
-	uint64 Padding : 8;		// Padding
+	uint64 StartPos : 24;					// Inclusive start position in the packet
+	uint64 EndPos : 24;						// Exclusive end position in the packet; BitSize = EndPos - StartPos
+	uint64 Level : 8;						// Level
+	uint64 Padding : 8;						// Padding
 
 	FNetProfilerBunchInfo BunchInfo;
 
-	uint32 EventTypeIndex;		// Will replace name index
-	uint32 NameIndex;			// Identify the name / type, should we store the actual Name as well
-	uint32 ObjectInstanceIndex;	// Object instance, Non zero if this is a NetObject, we can then look up data by indexing into ObjectInstances
+	uint32 EventTypeIndex;					// Will replace name index
+	FNetProfilerNameIndexType NameIndex;	// Identify the name / type, should we store the actual Name as well
+	uint32 ObjectInstanceIndex;				// Object instance, Non zero if this is a NetObject, we can then look up data by indexing into ObjectInstances
 };
 
 // This must be kept in sync with ENetTraceStatsCounterType in NetTrace.h
@@ -153,7 +151,7 @@ enum class ENetProfilerStatsCounterType : uint8
 struct FNetProfilerStatsCounterType
 {
 	uint32 StatsCounterTypeIndex;
-	uint32 NameIndex;
+	FNetProfilerNameIndexType NameIndex;
 	ENetProfilerStatsCounterType Type;
 };
 
diff --git a/Engine/Source/Runtime/Net/Core/Private/Net/Core/Trace/NetTraceInternal.cpp b/Engine/Source/Runtime/Net/Core/Private/Net/Core/Trace/NetTraceInternal.cpp
--- a/Engine/Source/Runtime/Net/Core/Private/Net/Core/Trace/NetTraceInternal.cpp
+++ b/Engine/Source/Runtime/Net/Core/Private/Net/Core/Trace/NetTraceInternal.cpp
@@ -27,7 +27,8 @@ struct FNetTraceInternal
 		ENetTraceVersion_Initial = 1,
 		ENetTraceVersion_BunchChannelIndex = 2,
 		ENetTraceVersion_BunchChannelInfo = 3,
-		ENetTraceVersion_FixedBunchSizeEncoding = 4,		
+		ENetTraceVersion_FixedBunchSizeEncoding = 4,
+		ENetTraceVersion_DebugNameIndexIs32Bits = 5,
 	};
 
 	struct FThreadBuffer
@@ -44,7 +45,7 @@ struct FNetTraceInternal
 
 	static inline FThreadBuffer* CreateThreadBuffer();
 
-	static constexpr ENetTraceVersion NetTraceVersion = ENetTraceVersion::ENetTraceVersion_FixedBunchSizeEncoding;
+	static constexpr ENetTraceVersion NetTraceVersion = ENetTraceVersion::ENetTraceVersion_DebugNameIndexIs32Bits;
 };
 
 static thread_local TUniquePtr<FNetTraceInternal::FThreadBuffer> ThreadBuffer;
diff --git a/Engine/Source/Runtime/Net/Core/Private/Net/Core/Trace/Reporters/NetTraceReporter.cpp b/Engine/Source/Runtime/Net/Core/Private/Net/Core/Trace/Reporters/NetTraceReporter.cpp
--- a/Engine/Source/Runtime/Net/Core/Private/Net/Core/Trace/Reporters/NetTraceReporter.cpp
+++ b/Engine/Source/Runtime/Net/Core/Private/Net/Core/Trace/Reporters/NetTraceReporter.cpp
@@ -25,7 +25,7 @@ UE_TRACE_EVENT_END()
 
 // Trace a name
 UE_TRACE_EVENT_BEGIN(NetTrace, NameEvent)
-	UE_TRACE_EVENT_FIELD(uint16, NameId)
+	UE_TRACE_EVENT_FIELD(uint32, NameId)
 	UE_TRACE_EVENT_FIELD(UE::Trace::AnsiString, Name)
 UE_TRACE_EVENT_END()
 
@@ -33,7 +33,7 @@ UE_TRACE_EVENT_BEGIN(NetTrace, ObjectCreatedEvent)
 	UE_TRACE_EVENT_FIELD(uint64, TypeId)
 	UE_TRACE_EVENT_FIELD(uint64, ObjectId)
 	UE_TRACE_EVENT_FIELD(uint32, OwnerId)
-	UE_TRACE_EVENT_FIELD(uint16, NameId)
+	UE_TRACE_EVENT_FIELD(uint32, NameId)
 	UE_TRACE_EVENT_FIELD(uint8, GameInstanceId)
 UE_TRACE_EVENT_END()
 
@@ -71,7 +71,7 @@ UE_TRACE_EVENT_END()
 // StatsCounterEvent
 UE_TRACE_EVENT_BEGIN(NetTrace, PacketStatsCounterEvent)
 	UE_TRACE_EVENT_FIELD(uint32, StatsValue)
-	UE_TRACE_EVENT_FIELD(uint16, NameId)
+	UE_TRACE_EVENT_FIELD(uint32, NameId)
 	UE_TRACE_EVENT_FIELD(uint16, ConnectionId)
 	UE_TRACE_EVENT_FIELD(uint8, GameInstanceId)
 UE_TRACE_EVENT_END()
@@ -79,7 +79,7 @@ UE_TRACE_EVENT_END()
 UE_TRACE_EVENT_BEGIN(NetTrace, FrameStatsCounterEvent)
 	UE_TRACE_EVENT_FIELD(uint64, Timestamp)
 	UE_TRACE_EVENT_FIELD(uint32, StatsValue)
-	UE_TRACE_EVENT_FIELD(uint16, NameId)
+	UE_TRACE_EVENT_FIELD(uint32, NameId)
 	UE_TRACE_EVENT_FIELD(uint8, GameInstanceId)
 UE_TRACE_EVENT_END()
 
diff --git a/Engine/Source/Runtime/Net/Core/Public/Net/Core/Trace/NetDebugName.h b/Engine/Source/Runtime/Net/Core/Public/Net/Core/Trace/NetDebugName.h
--- a/Engine/Source/Runtime/Net/Core/Public/Net/Core/Trace/NetDebugName.h
+++ b/Engine/Source/Runtime/Net/Core/Public/Net/Core/Trace/NetDebugName.h
@@ -11,7 +11,7 @@ namespace UE::Net
 /**
  * NetDebugName, carries both string pointer and a debugname id assigned by the persistent debug storage. DebugNameId is used to be able to avoid hashing already seen NetDebugNames
  */
-typedef uint16 FNetDebugNameId;
+typedef uint32 FNetDebugNameId;
 
 struct FNetDebugName
 {
diff --git a/Engine/Source/Runtime/Net/Core/Public/Net/Core/Trace/Private/NetTraceInternal.h b/Engine/Source/Runtime/Net/Core/Public/Net/Core/Trace/Private/NetTraceInternal.h
--- a/Engine/Source/Runtime/Net/Core/Public/Net/Core/Trace/Private/NetTraceInternal.h
+++ b/Engine/Source/Runtime/Net/Core/Public/Net/Core/Trace/Private/NetTraceInternal.h
@@ -548,7 +548,7 @@ void FNetTraceCollector::Reset()
 #define UE_NET_TRACE_INTERNAL_BUNCH_SCOPE(Collector, Bunch, ...) FNetTraceBunchScope<decltype(Bunch)> PREPROCESSOR_JOIN(NetTraceBunchScope, __LINE__)(Bunch, __VA_ARGS__, Collector)
 
 #define UE_NET_TRACE_INTERNAL_SCOPE(Name, Stream, Collector, Verbosity) \
-	auto PREPROCESSOR_JOIN(NetTraceNameFunc_, __LINE__) = []() { static uint16 NameId = FNetTrace::TraceName(TEXT(#Name)); return NameId; }; \
+	auto PREPROCESSOR_JOIN(NetTraceNameFunc_, __LINE__) = []() { static UE::Net::FNetDebugNameId NameId = FNetTrace::TraceName(TEXT(#Name)); return NameId; }; \
 	FNetTrace::TChooseTraceEventScope<Verbosity, decltype(Stream)>::Type PREPROCESSOR_JOIN(NetTraceScope, __LINE__)(PREPROCESSOR_JOIN(NetTraceNameFunc_, __LINE__), Stream, Collector, Verbosity)
 
 #define UE_NET_TRACE_INTERNAL_OBJECT_SCOPE(HandleOrNetGUID, Stream, Collector, Verbosity) \
@@ -559,14 +559,14 @@ void FNetTraceCollector::Reset()
 	FNetTrace::TChooseTraceEventScope<Verbosity, decltype(Stream)>::Type PREPROCESSOR_JOIN(NetTraceScope, __LINE__)(PREPROCESSOR_JOIN(NetTraceDynamicNameFunc_, __LINE__), Stream, Collector, Verbosity)
 
 #define UE_NET_TRACE_INTERNAL_NAMED_SCOPE(ScopeName, EventName, Stream, Collector, Verbosity) \
-	auto PREPROCESSOR_JOIN(NetTraceNameFunc_, __LINE__) = []() { static uint16 NameId = FNetTrace::TraceName(TEXT(#EventName)); return NameId; }; \
+	auto PREPROCESSOR_JOIN(NetTraceNameFunc_, __LINE__) = []() { static UE::Net::FNetDebugNameId  NameId = FNetTrace::TraceName(TEXT(#EventName)); return NameId; }; \
 	FNetTrace::TChooseTraceEventScope<Verbosity, decltype(Stream)>::Type ScopeName(PREPROCESSOR_JOIN(NetTraceNameFunc_, __LINE__), Stream, Collector, Verbosity)
 
 #define UE_NET_TRACE_INTERNAL_NAMED_OBJECT_SCOPE(ScopeName, HandleOrNetGUID, Stream, Collector, Verbosity) \
 	FNetTrace::TChooseTraceEventScope<Verbosity, decltype(Stream)>::Type ScopeName(GetObjectIdForNetTrace(HandleOrNetGUID), Stream, Collector, Verbosity)
 
 #define UE_NET_TRACE_INTERNAL_NAMED_DYNAMIC_NAME_SCOPE(ScopeName, EventName, Stream, Collector, Verbosity) \
-	auto PREPROCESSOR_JOIN(NetTraceDynamicNameFunc_, __LINE__) = [&]() { uint16 NameId = FNetTrace::TraceName(EventName); return NameId; }; \
+	auto PREPROCESSOR_JOIN(NetTraceDynamicNameFunc_, __LINE__) = [&]() { UE::Net::FNetDebugNameId NameId = FNetTrace::TraceName(EventName); return NameId; }; \
 	FNetTrace::TChooseTraceEventScope<Verbosity, decltype(Stream)>::Type ScopeName(PREPROCESSOR_JOIN(NetTraceDynamicNameFunc_, __LINE__), Stream, Collector, Verbosity)
 
 #define UE_NET_TRACE_INTERNAL_SET_SCOPE_NAME(ScopeName, EventName) UE_NET_TRACE_DO_IF(ScopeName.IsValid(), ScopeName.SetEventName(FNetTrace::TraceName(EventName)))
@@ -576,7 +576,7 @@ void FNetTraceCollector::Reset()
 
 #define UE_NET_TRACE_INTERNAL(Name, Collector, StartPos, EndPos, Verbosity) \
 	do { if (FNetTrace::GetCollectorAtVerbosity(Collector, Verbosity)) { \
-			static uint16 PREPROCESSOR_JOIN(NetTraceNameId_, __LINE__) = FNetTrace::TraceName(TEXT(#Name)); \
+			static UE::Net::FNetDebugNameId PREPROCESSOR_JOIN(NetTraceNameId_, __LINE__) = FNetTrace::TraceName(TEXT(#Name)); \
 			FNetTrace::TracePacketContentEvent(*Collector, PREPROCESSOR_JOIN(NetTraceNameId_, __LINE__), StartPos, EndPos, Verbosity); \
 	} } while (0)
 
@@ -594,13 +594,13 @@ void FNetTraceCollector::Reset()
 
 #define UE_NET_TRACE_INTERNAL_PACKET_STATSCOUNTER(GameInstanceId, ConnectionId, Name, StatValue, Verbosity) \
 	do { if (FNetTrace::GetNetTraceVerbosityEnabled(Verbosity)) { \
-			static uint16 PREPROCESSOR_JOIN(NetTraceNameId_, __LINE__) = FNetTrace::TraceName(TEXT(#Name)); \
+			static UE::Net::FNetDebugNameId PREPROCESSOR_JOIN(NetTraceNameId_, __LINE__) = FNetTrace::TraceName(TEXT(#Name)); \
 			FNetTrace::TracePacketStatsCounter(GameInstanceId, ConnectionId, PREPROCESSOR_JOIN(NetTraceNameId_, __LINE__), StatValue); \
 	} } while (0)
 
 #define UE_NET_TRACE_INTERNAL_FRAME_STATSCOUNTER(GameInstanceId, Name, StatValue, Verbosity) \
 	do { if (FNetTrace::GetNetTraceVerbosityEnabled(Verbosity)) { \
-			static uint16 PREPROCESSOR_JOIN(NetTraceNameId_, __LINE__) = FNetTrace::TraceName(TEXT(#Name)); \
+			static UE::Net::FNetDebugNameId PREPROCESSOR_JOIN(NetTraceNameId_, __LINE__) = FNetTrace::TraceName(TEXT(#Name)); \
 			FNetTrace::TraceFrameStatsCounter(GameInstanceId, PREPROCESSOR_JOIN(NetTraceNameId_, __LINE__), StatValue); \
 	} } while (0)
 
