Silenced PVS warning from:

Network insights/NetTrace
Bumped NetDebugName DebugNameId to use 32 bits.

diff --git a/Engine/Source/Developer/TraceServices/Private/Analyzers/NetTraceAnalyzer.cpp b/Engine/Source/Developer/TraceServices/Private/Analyzers/NetTraceAnalyzer.cpp
--- a/Engine/Source/Developer/TraceServices/Private/Analyzers/NetTraceAnalyzer.cpp
+++ b/Engine/Source/Developer/TraceServices/Private/Analyzers/NetTraceAnalyzer.cpp
@@ -70,7 +70,7 @@ void FNetTraceAnalyzer::OnAnalysisEnd()
 
 uint32 FNetTraceAnalyzer::GetTracedEventTypeIndex(FNetProfilerNameIndexType NameIndex, uint8 Level)
 {
-	const uint64 TracedEventTypeKey = uint64(NameIndex << 8U | Level);
+	const uint64 TracedEventTypeKey = (((uint64)NameIndex) << 8U) | (uint64)Level;
 
 	if (const uint32* NetProfilerEventTypeIndex = TraceEventTypeToNetProfilerEventTypeIndexMap.Find(TracedEventTypeKey))
 	{
