[Insights] Memory Insights: Added analysis support for additional LLM tag sets (asset and asset classes). Improved memory and performance for LLM stats provider (reimplemented FMemoryProvider, now using TPagedArrays for SnapshotTimes and SampleTimes).

diff --git a/Engine/Source/Developer/TraceServices/Private/Analyzers/MemoryAnalysis.cpp b/Engine/Source/Developer/TraceServices/Private/Analyzers/MemoryAnalysis.cpp
--- a/Engine/Source/Developer/TraceServices/Private/Analyzers/MemoryAnalysis.cpp
+++ b/Engine/Source/Developer/TraceServices/Private/Analyzers/MemoryAnalysis.cpp
@@ -26,6 +26,7 @@ void FMemoryAnalyzer::OnAnalysisBegin(const FOnAnalysisContext& Context)
 
 	Builder.RouteEvent(RouteId_TagsSpec, "LLM", "TagsSpec");
 	Builder.RouteEvent(RouteId_TrackerSpec, "LLM", "TrackerSpec");
+	Builder.RouteEvent(RouteId_TagSetSpec, "LLM", "TagSetSpec"); // added in UE 5.6
 	Builder.RouteEvent(RouteId_TagValue, "LLM", "TagValue");
 }
 
@@ -45,17 +46,26 @@ bool FMemoryAnalyzer::OnEvent(uint16 RouteId, EStyle Style, const FOnEventContex
 	{
 	case RouteId_TagsSpec:
 	{
-		FString Name = FTraceAnalyzerUtils::LegacyAttachmentString<TCHAR>("Name", Context);
 		const int64 TagId = EventData.GetValue<int64>("TagId");
 		const int64 ParentId = EventData.GetValue<int64>("ParentId");
-		Provider->AddEventSpec(TagId, *Name, ParentId);
+		const uint8 TagSetId = EventData.GetValue<uint8>("TagSetId"); // added in UE 5.6
+		FString Name = FTraceAnalyzerUtils::LegacyAttachmentString<TCHAR>("Name", Context);
+		Provider->AddTagSpec(TagId, Name, ParentId, (FMemoryTagSetId)TagSetId);
 	}
 	break;
 	case RouteId_TrackerSpec:
 	{
 		uint8 TrackerId = EventData.GetValue<uint8>("TrackerId");
 		FString Name = FTraceAnalyzerUtils::LegacyAttachmentString<ANSICHAR>("Name", Context);
-		Provider->AddTrackerSpec(TrackerId, *Name);
+		Provider->AddTrackerSpec((FMemoryTrackerId)TrackerId, Name);
+	}
+	break;
+	case RouteId_TagSetSpec:
+	{
+		const uint8 TagSetId = EventData.GetValue<uint8>("TagSetId");
+		FString Name;
+		EventData.GetString("Name", Name);
+		Provider->AddTagSetSpec((FMemoryTagSetId)TagSetId, Name);
 	}
 	break;
 	case RouteId_TagValue:
diff --git a/Engine/Source/Developer/TraceServices/Private/Analyzers/MemoryAnalysis.h b/Engine/Source/Developer/TraceServices/Private/Analyzers/MemoryAnalysis.h
--- a/Engine/Source/Developer/TraceServices/Private/Analyzers/MemoryAnalysis.h
+++ b/Engine/Source/Developer/TraceServices/Private/Analyzers/MemoryAnalysis.h
@@ -26,6 +26,7 @@ private:
 	enum : uint16
 	{
 		RouteId_TagsSpec,
+		RouteId_TagSetSpec,
 		RouteId_TrackerSpec,
 		RouteId_TagValue,
 	};
diff --git a/Engine/Source/Developer/TraceServices/Private/Model/Memory.cpp b/Engine/Source/Developer/TraceServices/Private/Model/Memory.cpp
--- a/Engine/Source/Developer/TraceServices/Private/Model/Memory.cpp
+++ b/Engine/Source/Developer/TraceServices/Private/Model/Memory.cpp
@@ -8,94 +8,217 @@ namespace TraceServices
 
 FMemoryProvider::FMemoryProvider(IAnalysisSession& InSession)
 	: Session(InSession)
-	, TagDescsPool(Session.GetLinearAllocator(), 1024)
-	, TagsSerial(0)
+	, SnapshotTimes(Session.GetLinearAllocator(), 8 * 1024)
 {
-	TagDescs.Reserve(256);
+	AddTagSetSpec((FMemoryTagSetId)0, TEXT("Default"));
 }
 
-void FMemoryProvider::AddEventSpec(FMemoryTagId TagId, const TCHAR* Name, FMemoryTagId ParentTagId)
+FMemoryProvider::~FMemoryProvider()
 {
-	if (TagId == FMemoryTagInfo::InvalidTagId || TagId == -1)
+	for (FTrackerData* Tracker : AvailableTrackers)
 	{
-		// invalid tag id
+		if (Tracker)
+		{
+			delete Tracker->Info;
+			delete Tracker;
+		}
+	}
+	AvailableTrackers.Reset();
+	NumTrackers = 0;
+
+	for (FTagSetData* TagSet : AvailableTagSets)
+	{
+		if (TagSet)
+		{
+			delete TagSet->Info;
+			delete TagSet;
+		}
+	}
+	AvailableTagSets.Reset();
+	NumTagSets = 0;
+
+	for (FMemoryTagInfo* TagInfo : AvailableTags)
+	{
+		delete TagInfo;
+	}
+	AvailableTags.Reset();
+	TagMap.Reset();
+}
+
+void FMemoryProvider::AddTrackerSpec(FMemoryTrackerId TrackerId, const FString& Name)
+{
+	if (TrackerId < 0 || TrackerId >= FMemoryTrackerInfo::MaxTrackers)
+	{
+		// invalid tracker id
 		return;
 	}
 
-	if (ParentTagId == -1) // backward compatibility with UE 4.27
+	while (AvailableTrackers.Num() <= TrackerId)
 	{
-		ParentTagId = FMemoryTagInfo::InvalidTagId;
+		AvailableTrackers.AddDefaulted();
 	}
 
-	if (TagDescs.Contains(TagId))
+	FTrackerData* Tracker = AvailableTrackers[TrackerId];
+	if (Tracker)
 	{
-		// tag already registered
+		// The tracker is already registered.
+		check(Tracker->Info);
+
+		// Update name.
+		Tracker->Info->Name = Name;
 		return;
 	}
 
-	// Do not remove or insert from TagDescsPool, cached references to this memory are stored and inserts / removes will
-	// result in dangerous memory moves
-	FMemoryTagInfo& NewMemoryInfo = TagDescsPool.EmplaceBack(FMemoryTagInfo{ TagId, FString(Name), ParentTagId, 0 });
+	FMemoryTrackerInfo* TrackerInfo = new FMemoryTrackerInfo({ TrackerId, Name });
+	AvailableTrackers[(int32)TrackerId] = new FTrackerData(TrackerInfo);
+	++NumTrackers;
+}
 
-	TagDescs.Add(TagId, &NewMemoryInfo);
-	++TagsSerial;
+FMemoryProvider::FTrackerData* FMemoryProvider::GetTracker(FMemoryTrackerId TrackerId) const
+{
+	if (TrackerId >= 0 && TrackerId < AvailableTrackers.Num())
+	{
+		return AvailableTrackers[TrackerId];
+	}
+	return nullptr;
 }
 
-void FMemoryProvider::AddTrackerSpec(FMemoryTrackerId TrackerId, const TCHAR* Name)
+FMemoryProvider::FTrackerData* FMemoryProvider::GetOrAddTracker(FMemoryTrackerId TrackerId)
 {
-	if (TrackerId < 0 || TrackerId > FMemoryTrackerInfo::MaxValidTrackerId)
+	FTrackerData* Tracker = (TrackerId >= 0 && TrackerId < AvailableTrackers.Num()) ? AvailableTrackers[TrackerId] : nullptr;
+	if (Tracker)
 	{
-		// invalid tracker id
+		return Tracker;
+	}
+	AddTagSetSpec(TrackerId, TEXT("<unknown>"));
+	return AvailableTrackers[TrackerId];
+}
+
+void FMemoryProvider::AddTagSetSpec(FMemoryTagSetId TagSetId, const FString& Name)
+{
+	if (TagSetId < 0 || TagSetId >= FMemoryTagSetInfo::MaxTagSets)
+	{
+		// invalid tag set id
 		return;
 	}
 
-	if (TrackerDescs.Contains(TrackerId))
+	while (AvailableTagSets.Num() <= TagSetId)
 	{
-		// tracker already registerd
+		AvailableTagSets.AddDefaulted();
+	}
+
+	FTagSetData* TagSet = AvailableTagSets[TagSetId];
+	if (TagSet)
+	{
+		// The tag set is already registered.
+		check(TagSet->Info);
+
+		// Update name.
+		TagSet->Info->Name = Name;
+
 		return;
 	}
 
-	TrackerDescs.Add(TrackerId, FMemoryTrackerInfo{ TrackerId, FString(Name) });
+	FMemoryTagSetInfo* TagSetInfo = new FMemoryTagSetInfo({ TagSetId, Name });
+	AvailableTagSets[TagSetId] = new FTagSetData(TagSetInfo);
+	++NumTagSets;
 }
 
-void FMemoryProvider::AddTagSnapshot(FMemoryTrackerId TrackerId, double Time, const int64* Tags, const FMemoryTagSample* Values, uint32 TagCount)
+FMemoryProvider::FTagSetData* FMemoryProvider::GetTagSet(FMemoryTagSetId TagSetId) const
 {
-	if (TrackerId < 0 || TrackerId >= TrackerDescs.Num())
+	if (TagSetId >= 0 && TagSetId < AvailableTagSets.Num())
+	{
+		return AvailableTagSets[TagSetId];
+	}
+	return nullptr;
+}
+
+FMemoryProvider::FTagSetData* FMemoryProvider::GetOrAddTagSet(FMemoryTagSetId TagSetId)
+{
+	FTagSetData* TagSet = (TagSetId >= 0 && TagSetId < AvailableTagSets.Num()) ? AvailableTagSets[TagSetId] : nullptr;
+	if (TagSet)
+	{
+		return TagSet;
+	}
+	AddTagSetSpec(TagSetId, TEXT("<unknown>"));
+	return AvailableTagSets[TagSetId];
+}
+
+void FMemoryProvider::AddTagSpec(FMemoryTagId TagId, const FString& Name, FMemoryTagId ParentTagId, FMemoryTagSetId TagSetId)
+{
+	if (TagId == FMemoryTagInfo::InvalidTagId || TagId == -1)
+	{
+		// invalid tag id
+		return;
+	}
+
+	if (ParentTagId == -1) // backward compatibility with UE 4.27
 	{
-		// tracker id is invalid or not yet registered
+		ParentTagId = FMemoryTagInfo::InvalidTagId;
+	}
+
+	if (TagSetId < 0 || TagSetId >= FMemoryTagSetInfo::MaxTagSets)
+	{
+		// invalid tag set id
 		return;
 	}
+	FTagSetData* TagSet = GetOrAddTagSet(TagSetId);
+	check(TagSet);
 
-	if (TrackerId >= Trackers.Num())
+	FMemoryTagInfo** FoundTagInfo = TagMap.Find(TagId);
+	if (FoundTagInfo)
 	{
-		Trackers.AddDefaulted(TrackerId - Trackers.Num() + 1);
+		// The tag is already registered.
+		(*FoundTagInfo)->ParentId = ParentTagId;
+		(*FoundTagInfo)->TagSetId = TagSetId;
+		(*FoundTagInfo)->Name = Name;
+		return;
 	}
 
-	FTrackerData& TrackerData = Trackers[TrackerId];
-	FMemoryTrackerInfo& Tracker = TrackerDescs[TrackerId];
+	uint64 Trackers = 0; // flags for trackers using this tag
+	FMemoryTagInfo* TagInfo = new FMemoryTagInfo({ TagId, ParentTagId, TagSetId, Trackers, Name });
+	TagMap.Add(TagId, TagInfo);
+
+	AvailableTags.Add(TagInfo);
+	TagSet->Tags.Add(TagId, TagInfo);
+
+	++TagsSerial;
+}
+
+void FMemoryProvider::AddTagSnapshot(FMemoryTrackerId TrackerId, double Time, const int64* Tags, const FMemoryTagSample* Values, uint32 TagCount)
+{
+	if (TrackerId < 0 || TrackerId >= FMemoryTrackerInfo::MaxTrackers)
+	{
+		// invalid tracker id
+		return;
+	}
+	FTrackerData* Tracker = GetOrAddTracker(TrackerId);
+	check(Tracker);
 
-	TrackerData.SampleTimes.Push(Time);
-	const int32 SampleCount = TrackerData.SampleTimes.Num();
+	SnapshotTimes.EmplaceBack(Time);
+	const uint64 SampleCount = SnapshotTimes.Num();
 
 	for (uint32 TagIndex = 0; TagIndex < TagCount; ++TagIndex)
 	{
 		const FMemoryTagId TagId = FMemoryTagId(Tags[TagIndex]);
+		const FMemoryTagSample& Value = Values[TagIndex];
 
-		FTagSampleData* TagSamples = TrackerData.Samples.Find(TagId);
+		FTagSampleData* TagSamples = Tracker->Samples.FindRef(TagId);
 		if (!TagSamples)
 		{
-			TagSamples = &TrackerData.Samples.Add(TagId, FTagSampleData(Session.GetLinearAllocator()));
+			FTagSampleData* TagSampleData = new FTagSampleData(Session.GetLinearAllocator());
+			TagSamples = Tracker->Samples.Add(TagId, TagSampleData);
 		}
 
 		auto& TagValues = TagSamples->Values;
 
 		// Add "past values". If new value is added, we need to ensure we have
 		// the same number of elements in the value array as in the timestamp array.
-		int32 CurrentSampleCount = int32(TagValues.Num());
+		uint64 CurrentSampleCount = TagValues.Num();
 		if (CurrentSampleCount < SampleCount - 1)
 		{
 			const FMemoryTagSample LastValue = (CurrentSampleCount > 0) ? TagValues.Last() : FMemoryTagSample { 0 };
-			for (int32 ValueIndex = SampleCount - 1; ValueIndex > CurrentSampleCount; --ValueIndex)
+			for (uint64 ValueIndex = SampleCount - 1; ValueIndex > CurrentSampleCount; --ValueIndex)
 			{
 				TagValues.PushBack() = LastValue;
 			}
@@ -104,21 +227,21 @@ void FMemoryProvider::AddTagSnapshot(FMemoryTrackerId TrackerId, double Time, co
 		TagValues.PushBack() = Values[TagIndex];
 		check(int32(TagValues.Num()) == SampleCount);
 
-		if (!TagSamples->TagPtr)
+		if (!TagSamples->TagInfo)
 		{
 			// Cache pointer to FMemoryTagInfo to avoid further lookups for this tag.
-			if (FMemoryTagInfo** TagInfo = TagDescs.Find(TagId))
+			if (FMemoryTagInfo* TagInfo = const_cast<FMemoryTagInfo*>(GetTag(TagId)))
 			{
-				TagSamples->TagPtr = *TagInfo;
+				TagSamples->TagInfo = TagInfo;
 			}
 		}
 
-		if (TagSamples->TagPtr)
+		if (TagSamples->TagInfo)
 		{
 			const uint64 TrackerFlag = 1ULL << TrackerId;
-			if ((TagSamples->TagPtr->Trackers & TrackerFlag) == 0)
+			if ((TagSamples->TagInfo->Trackers & TrackerFlag) == 0)
 			{
-				TagSamples->TagPtr->Trackers |= TrackerFlag;
+				TagSamples->TagInfo->Trackers |= TrackerFlag;
 				++TagsSerial;
 			}
 		}
@@ -132,48 +255,78 @@ uint32 FMemoryProvider::GetTagSerial() const
 
 uint32 FMemoryProvider::GetTagCount() const
 {
-	return TagDescs.Num();
+	return AvailableTags.Num();
 }
 
 void FMemoryProvider::EnumerateTags(TFunctionRef<void(const FMemoryTagInfo&)> Callback) const
 {
-	for (auto& Tag : TagDescs)
+	for (const FMemoryTagInfo* Tag : AvailableTags)
 	{
-		Callback(*Tag.Value);
+		Callback(*Tag);
 	}
 }
 
-const FMemoryTagInfo* FMemoryProvider::GetTag(FMemoryTagId TagId) const
+void FMemoryProvider::EnumerateTags(FMemoryTagSetId TagSetId, TFunctionRef<void(const FMemoryTagInfo&)> Callback) const
 {
-	if (FMemoryTagInfo* const * TagInfo = TagDescs.Find(TagId))
+	FTagSetData* TagSet = GetTagSet(TagSetId);
+	if (!TagSet)
 	{
-		return *TagInfo;
+		return;
+	}
+
+	for (const auto& KV : TagSet->Tags)
+	{
+		Callback(*KV.Value);
 	}
-	return nullptr;
+}
+
+const FMemoryTagInfo* FMemoryProvider::GetTag(FMemoryTagId TagId) const
+{
+	return TagMap.FindRef(TagId);
 }
 
 uint32 FMemoryProvider::GetTrackerCount() const
 {
-	return TrackerDescs.Num();
+	return NumTrackers;
 }
 
 void FMemoryProvider::EnumerateTrackers(TFunctionRef<void(const FMemoryTrackerInfo&)> Callback) const
 {
-	for (auto& Desc : TrackerDescs)
+	for (const FTrackerData* Tracker : AvailableTrackers)
+	{
+		if (Tracker)
+		{
+			Callback(*Tracker->Info);
+		}
+	}
+}
+
+uint32 FMemoryProvider::GetTagSetCount() const
+{
+	return NumTagSets;
+}
+
+void FMemoryProvider::EnumerateTagSets(TFunctionRef<void(const FMemoryTagSetInfo&)> Callback) const
+{
+	for (const FTagSetData* TagSet : AvailableTagSets)
 	{
-		Callback(Desc.Value);
+		if (TagSet)
+		{
+			Callback(*TagSet->Info);
+		}
 	}
 }
 
 uint64 FMemoryProvider::GetTagSampleCount(FMemoryTrackerId TrackerId, FMemoryTagId TagId) const
 {
-	if (TrackerId < 0 || TrackerId >= Trackers.Num())
+	FTrackerData* Tracker = GetTracker(TrackerId);
+	if (!Tracker)
 	{
 		// invalid tracker id
 		return 0;
 	}
 
-	const FTagSampleData* TagSamples = Trackers[TrackerId].Samples.Find(TagId);
+	const FTagSampleData* TagSamples = Tracker->Samples.FindRef(TagId);
 	if (!TagSamples)
 	{
 		// invalid tag id
@@ -185,49 +338,58 @@ uint64 FMemoryProvider::GetTagSampleCount(FMemoryTrackerId TrackerId, FMemoryTag
 
 void FMemoryProvider::EnumerateTagSamples(FMemoryTrackerId TrackerId, FMemoryTagId TagId, double StartTime, double EndTime, bool bIncludeRangeNeighbours, TFunctionRef<void(double Time, double Duration, const FMemoryTagSample&)> Callback) const
 {
-	if (TrackerId < 0 || TrackerId >= Trackers.Num())
+	FTrackerData* Tracker = GetTracker(TrackerId);
+	if (!Tracker)
 	{
 		// invalid tracker id
 		return;
 	}
 
-	const FTagSampleData* TagSamples = Trackers[TrackerId].Samples.Find(TagId);
+	const FTagSampleData* TagSamples = Tracker->Samples.FindRef(TagId);
 	if (!TagSamples)
 	{
 		// invalid tag id
 		return;
 	}
 
-	const TArray<double>& SampleTimes = Trackers[TrackerId].SampleTimes;
 	const TPagedArray<FMemoryTagSample>& SampleValues = TagSamples->Values;
 
-	int32 IndexStart = Algo::LowerBound(SampleTimes, StartTime);
-	int32 IndexEnd = Algo::UpperBound(SampleTimes, EndTime);
+	int32 IndexStart = (int32)PagedArrayAlgo::LowerBound(SnapshotTimes, StartTime);
+	int32 IndexEnd = (int32)PagedArrayAlgo::UpperBound(SnapshotTimes, EndTime);
 
 	if (bIncludeRangeNeighbours)
 	{
 		IndexStart = FMath::Max(IndexStart - 1, 0);
 		IndexEnd = FMath::Min(IndexEnd + 1, static_cast<int32>(SampleValues.Num()));
 	}
-
-	const bool bEnumerateLastSample = (IndexEnd == static_cast<uint32>(SampleValues.Num()));
-	if (bEnumerateLastSample)
+	else
 	{
-		// Skip last sample from regular enumeration.
-		--IndexEnd;
+		IndexStart = FMath::Max(IndexStart, 0);
+		IndexEnd = FMath::Min(IndexEnd, static_cast<int32>(SampleValues.Num()));
 	}
 
-	for (int32 SampleIndex = IndexStart; SampleIndex < IndexEnd; ++SampleIndex)
+	if (IndexStart >= IndexEnd)
 	{
-		const double Time = SampleTimes[SampleIndex];
-		const double Duration = SampleTimes[SampleIndex + 1] - Time;
-		Callback(Time, Duration, SampleValues[SampleIndex]);
+		return;
 	}
 
-	if (bEnumerateLastSample)
+	auto TimeIt = SnapshotTimes.GetIteratorFromItem(IndexStart);
+	auto ValueIt = SampleValues.GetIteratorFromItem(IndexStart);
+
+	double Time = *TimeIt;
+	for (int32 SampleIndex = IndexStart; SampleIndex < IndexEnd; ++SampleIndex)
 	{
-		// Last sample has 0 duration.
-		Callback(SampleTimes[IndexEnd], 0.0, SampleValues[IndexEnd]);
+		if (TimeIt.NextItem())
+		{
+			double NextTime = *TimeIt;
+			Callback(Time, NextTime - Time, SampleValues[SampleIndex]);
+			Time = NextTime;
+		}
+		else
+		{
+			// Last sample has 0 duration.
+			Callback(Time, 0.0, *ValueIt);
+		}
 	}
 }
 
diff --git a/Engine/Source/Developer/TraceServices/Private/Model/MemoryPrivate.h b/Engine/Source/Developer/TraceServices/Private/Model/MemoryPrivate.h
--- a/Engine/Source/Developer/TraceServices/Private/Model/MemoryPrivate.h
+++ b/Engine/Source/Developer/TraceServices/Private/Model/MemoryPrivate.h
@@ -14,18 +14,27 @@ class FMemoryProvider : public IMemoryProvider
 {
 public:
 	explicit FMemoryProvider(IAnalysisSession& Session);
-	virtual ~FMemoryProvider() {}
+	virtual ~FMemoryProvider();
+
+	void AddTrackerSpec(FMemoryTrackerId TrackerId, const FString& Name);
+	void AddTagSetSpec(FMemoryTagSetId TagSetId, const FString& Name);
+	void AddTagSpec(FMemoryTagId Tag, const FString& Name, FMemoryTagId ParentTag, FMemoryTagSetId TagSet);
 
-	void AddEventSpec(FMemoryTagId Tag, const TCHAR* Name, FMemoryTagId ParentTag);
-	void AddTrackerSpec(FMemoryTrackerId TrackerId, const TCHAR* Name);
 	void AddTagSnapshot(FMemoryTrackerId TrackerId, double Time, const int64* Tags, const FMemoryTagSample* Values, uint32 TagCount);
 
 	virtual uint32 GetTagSerial() const override;
 	virtual uint32 GetTagCount() const override;
-	virtual void EnumerateTags(TFunctionRef<void(const FMemoryTagInfo&)> Callback) const override;
 	virtual const FMemoryTagInfo* GetTag(FMemoryTagId Id) const override;
+
+	virtual void EnumerateTags(TFunctionRef<void(const FMemoryTagInfo&)> Callback) const override;
+	virtual void EnumerateTags(FMemoryTagSetId TagSetId, TFunctionRef<void(const FMemoryTagInfo&)> Callback) const override;
+
 	virtual uint32 GetTrackerCount() const override;
 	virtual void EnumerateTrackers(TFunctionRef<void(const FMemoryTrackerInfo&)> Callback) const override;
+
+	virtual uint32 GetTagSetCount() const override;
+	virtual void EnumerateTagSets(TFunctionRef<void(const FMemoryTagSetInfo&)> Callback) const override;
+
 	virtual uint64 GetTagSampleCount(FMemoryTrackerId Tracker, FMemoryTagId Tag) const override;
 	virtual void EnumerateTagSamples(
 		FMemoryTrackerId Tracker,
@@ -36,42 +45,71 @@ public:
 		TFunctionRef<void(double Time, double Duration, const FMemoryTagSample&)> Callback) const override;
 
 private:
-	enum {
-		DefaultPageSize = 65536,
-	};
-
 	struct FTagSampleData
 	{
 		// Sample values.
 		TPagedArray<FMemoryTagSample> Values;
 
-		// Cached pointer to the actual memory tag.
-		FMemoryTagInfo* TagPtr;
+		FMemoryTagInfo* TagInfo = nullptr;
 
-		FTagSampleData(ILinearAllocator& Allocator) : Values(Allocator, DefaultPageSize), TagPtr(nullptr) {}
-		FTagSampleData(const FTagSampleData& Other) : Values(Other.Values), TagPtr(Other.TagPtr) {}
+		FTagSampleData(ILinearAllocator& Allocator)
+			: Values(Allocator, 65536)
+		{}
+		FTagSampleData(const FTagSampleData& Other)
+			: Values(Other.Values)
+		{}
 	};
 
 	struct FTrackerData
 	{
-		// Timestamps for samples
-		TArray<double> SampleTimes;
+		// The memory tracker.
+		FMemoryTrackerInfo* Info = nullptr;
 
-		// Samples for each tag
-		TMap<FMemoryTagId, FTagSampleData> Samples;
+		// All tags and samples for this tracker.
+		TMap<FMemoryTagId, FTagSampleData*> Samples;
 
-		FTrackerData() {}
+		FTrackerData(FMemoryTrackerInfo* InInfo) : Info(InInfo) {}
 
 	private:
 		FTrackerData(const FTrackerData&) = delete;
 	};
 
+	struct FTagSetData
+	{
+		// The memory tag set.
+		FMemoryTagSetInfo* Info = nullptr;
+
+		// All tags for this tag set.
+		TMap<FMemoryTagId, FMemoryTagInfo*> Tags; // tag id --> TagInfo
+
+		FTagSetData(FMemoryTagSetInfo* InInfo) : Info(InInfo) {}
+
+	private:
+		FTagSetData(const FTagSetData&) = delete;
+	};
+
+private:
+	FTrackerData* GetTracker(FMemoryTrackerId TrackerId) const;
+	FTrackerData* GetOrAddTracker(FMemoryTrackerId TrackerId);
+	FTagSetData* GetTagSet(FMemoryTagSetId TagSetId) const;
+	FTagSetData* GetOrAddTagSet(FMemoryTagSetId TagSetId);
+
+private:
 	IAnalysisSession& Session;
-	TMap<FMemoryTagId, FMemoryTagInfo*> TagDescs;
-	TPagedArray<FMemoryTagInfo> TagDescsPool;
-	TMap<FMemoryTrackerId, FMemoryTrackerInfo> TrackerDescs;
-	TArray<FTrackerData> Trackers;
-	uint32 TagsSerial;
+
+	TArray<FTrackerData*> AvailableTrackers;
+	uint32 NumTrackers = 0;
+
+	TArray<FTagSetData*> AvailableTagSets;
+	uint32 NumTagSets = 0;
+
+	TArray<FMemoryTagInfo*> AvailableTags;
+	TMap<FMemoryTagId, FMemoryTagInfo*> TagMap; // tag id --> TagInfo
+
+	// Timestamps for samples
+	TPagedArray<double> SnapshotTimes;
+
+	uint32 TagsSerial = 0;
 };
 
 } // namespace TraceServices
diff --git a/Engine/Source/Developer/TraceServices/Public/TraceServices/Model/Memory.h b/Engine/Source/Developer/TraceServices/Public/TraceServices/Model/Memory.h
--- a/Engine/Source/Developer/TraceServices/Public/TraceServices/Model/Memory.h
+++ b/Engine/Source/Developer/TraceServices/Public/TraceServices/Model/Memory.h
@@ -13,37 +13,52 @@ namespace TraceServices
 {
 
 typedef int32 FMemoryTrackerId;
+typedef int32 FMemoryTagSetId;
 typedef int64 FMemoryTagId;
 
 struct FMemoryTrackerInfo
 {
-	static const FMemoryTrackerId InvalidTrackerId = -1;
-	static const FMemoryTrackerId MaxValidTrackerId = 63;
+	static constexpr FMemoryTrackerId InvalidTrackerId = -1;
+	static constexpr int32 MaxTrackers = 8;
 
-	// Unique identifier of the llm tracker. Can also be used as an index, limited to [0 .. 63] range.
+	// The unique identifier of the memory tracker. Can also be used as an index, limited to [0 .. MaxTrackers-1] range.
 	FMemoryTrackerId Id;
 
-	// Name of the llm tracker.
+	// The name of the memory tracker.
 	FString Name;
 };
 
-struct FMemoryTagInfo
+struct FMemoryTagSetInfo
 {
-	static const FMemoryTagId InvalidTagId = 0;
+	static constexpr FMemoryTagSetId InvalidTagSetId = -1;
+	static constexpr int32 MaxTagSets = 8;
 
-	// Unique identifier of the llm tag.
-	FMemoryTagId Id;
+	// The unique identifier of the memory tag set. Can also be used as an index, limited to [0 .. MaxTagSets-1] range.
+	FMemoryTagSetId Id;
 
-	// Name of the stat associated with the llm tag.
 	FString Name;
+};
+
+struct FMemoryTagInfo
+{
+	static constexpr FMemoryTagId InvalidTagId = 0;
 
-	// Id of parent, 0 if no parent.
+	// The unique identifier of the memory tag.
+	FMemoryTagId Id;
+
+	// The id of the parent tag, 0 if no parent.
 	FMemoryTagId ParentId;
 
-	// Bit flags for trackers using this llm tag.
+	// The set id of the memory tag.
+	FMemoryTagSetId TagSetId;
+
+	// Bit flags for trackers using this memory tag.
 	// The bit position represents the tracker id; this limits the valid tracker ids to range [0 .. 63].
 	// It can be updated during analysis (as new trackers / snapshots are analyzed).
 	uint64 Trackers;
+
+	// The name of the memory tag.
+	FString Name;
 };
 
 struct FMemoryTagSample
@@ -76,23 +91,29 @@ public:
 	virtual void EnumerateTags(TFunctionRef<void(const FMemoryTagInfo&)> Callback) const = 0;
 
 	/**
-	 * Returns the meta data for a tag id.
-	 * @param Id Tag id
+	 * Enumerate the registered tags for a specified tag set.
+	 * @param Callback Function that will be called for each registered tag.
+	 */
+	virtual void EnumerateTags(FMemoryTagSetId TagSetId, TFunctionRef<void(const FMemoryTagInfo&)> Callback) const = 0;
+
+	/**
+	 * Returns the meta data for a memory tag specified by id.
+	 * @param TagId The id of the memory tag.
 	 * @return Memory tag information.
 	 */
-	virtual const FMemoryTagInfo* GetTag(FMemoryTagId Id) const = 0;
+	virtual const FMemoryTagInfo* GetTag(FMemoryTagId TagId) const = 0;
 
 	/**
 	 * Gets the number of samples for a given tag from a given tracker.
 	 * @param Tracker Tracer index.
-	 * @param Tag The id of the LLM tag.
+	 * @param TagId The id of the memory tag.
 	 * @return Number of samples that has been recorded.
 	 */
-	virtual uint64 GetTagSampleCount(FMemoryTrackerId Tracker, FMemoryTagId Tag) const = 0;
+	virtual uint64 GetTagSampleCount(FMemoryTrackerId Tracker, FMemoryTagId TagId) const = 0;
 
 	/**
-	 * Return the number of registered tracker descriptions.
-	 * @return Number of tracker descriptions.
+	 * Return the number of registered trackers.
+	 * @return Number of trackers.
 	 */
 	virtual uint32 GetTrackerCount() const = 0;
 
@@ -102,15 +123,27 @@ public:
 	 */
 	virtual void EnumerateTrackers(TFunctionRef<void(const FMemoryTrackerInfo&)> Callback) const = 0;
 
+	/**
+	 * Return the number of registered tag sets.
+	 * @return Number of tag sets.
+	 */
+	virtual uint32 GetTagSetCount() const = 0;
+
+	/**
+	 * Enumerate the registered tag sets.
+	 * @param Callback Function that is called for each registered tag set.
+	 */
+	virtual void EnumerateTagSets(TFunctionRef<void(const FMemoryTagSetInfo&)> Callback) const = 0;
+
 	/** Enumerates samples (values) for a specified LLM tag.
-	 * @param Tracker Tracker index.
-	 * @param Tag The id/index of the LLM tag.
+	 * @param Tracker The tracker index.
+	 * @param TagId The id of the memory tag.
 	 * @param StartSample The inclusive start index in the sample array of specified memory tag.
 	 * @param EndSample The exclusive end index in the sample array of specified memory tag.
 	 * @param Callback A callback function called for each sample enumerated.
-	 * @param bIncludeRangeNeighbours Includes the sample immediately before and after the selected range.
+	 * @param bIncludeRangeNeighbors Includes the sample immediately before and after the selected range.
 	 */
-	virtual void EnumerateTagSamples(FMemoryTrackerId Tracker, FMemoryTagId Tag, double StartTime, double EndTime, bool bIncludeRangeNeighbours, TFunctionRef<void(double Time, double Duration, const FMemoryTagSample&)> Callback) const = 0;
+	virtual void EnumerateTagSamples(FMemoryTrackerId Tracker, FMemoryTagId TagId, double StartTime, double EndTime, bool bIncludeRangeNeighbors, TFunctionRef<void(double Time, double Duration, const FMemoryTagSample&)> Callback) const = 0;
 };
 
 TRACESERVICES_API FName GetMemoryProviderName();
