[Insights] Implemented workaround for wrong timestamps in some frames traced by old GPU Profiler.

diff --git a/Engine/Source/Developer/TraceServices/Private/Analyzers/GpuProfilerTraceAnalysis.cpp b/Engine/Source/Developer/TraceServices/Private/Analyzers/GpuProfilerTraceAnalysis.cpp
--- a/Engine/Source/Developer/TraceServices/Private/Analyzers/GpuProfilerTraceAnalysis.cpp
+++ b/Engine/Source/Developer/TraceServices/Private/Analyzers/GpuProfilerTraceAnalysis.cpp
@@ -273,6 +273,13 @@ bool FGpuProfilerAnalyzer::OnEvent(uint16 RouteId, EStyle Style, const FOnEventC
 				continue;
 			}
 
+			// If it advances with more than 1h, it is probably a wrong timestamp.
+			if (LastTime > ThisMinTime + 3600.0 && ThisMinTime != DBL_MIN)
+			{
+				LastTime = ThisMinTime;
+				bHasErrors = true;
+			}
+
 			// The monolithic timeline assumes that timestamps are ever increasing, but
 			// with gpu/cpu calibration and drift there can be a tiny bit of overlap between
 			// frames. So we just clamp.
