[Insights]
- CallstackAnalysis: Added warning for truncated callstacks.
- ModuleProvider: Fixed typo in comment and in a log message.
- SymslibResolver: Minor refactor for the log message with "String allocator" stats.

diff --git a/Engine/Source/Developer/TraceServices/Private/Analyzers/CallstacksAnalysis.cpp b/Engine/Source/Developer/TraceServices/Private/Analyzers/CallstacksAnalysis.cpp
--- a/Engine/Source/Developer/TraceServices/Private/Analyzers/CallstacksAnalysis.cpp
+++ b/Engine/Source/Developer/TraceServices/Private/Analyzers/CallstacksAnalysis.cpp
@@ -5,6 +5,7 @@
 #include "HAL/LowLevelMemTracker.h"
 #include "Model/CallstacksProvider.h"
 #include "TraceServices/Model/AnalysisSession.h"
+#include "Common/Utils.h"
 
 namespace TraceServices
 {
@@ -36,11 +37,19 @@ bool FCallstacksAnalyzer::OnEvent(uint16 RouteId, EStyle Style, const FOnEventCo
 			uint8 NumFrames = (uint8)FMath::Min(255u, Frames.Num());
 			if (const uint32 Id = Context.EventData.GetValue<uint32>("CallstackId"))
 			{
+				if (NumFrames != Frames.Num())
+				{
+					UE_LOG(LogTraceServices, Warning, TEXT("Callstack with Id=%u has %u frames, but it will be limited to %u frames!"), Id, Frames.Num(), NumFrames);
+				}
 				Provider->AddCallstack(Id, Frames.GetData(), NumFrames);
 			}
 			// Backward compatibility with legacy memory trace format (5.0-EA).
 			else if (const uint64 Hash = Context.EventData.GetValue<uint64>("Id"))
 			{
+				if (NumFrames != Frames.Num())
+				{
+					UE_LOG(LogTraceServices, Warning, TEXT("Callstack with Hash=%llu has %u frames, but it will be limited to %u frames!"), Hash, Frames.Num(), NumFrames);
+				}
 				Provider->AddCallstackWithHash(Hash, Frames.GetData(), NumFrames);
 			}
 			break;
diff --git a/Engine/Source/Developer/TraceServices/Private/Model/ModuleProvider.cpp b/Engine/Source/Developer/TraceServices/Private/Model/ModuleProvider.cpp
--- a/Engine/Source/Developer/TraceServices/Private/Model/ModuleProvider.cpp
+++ b/Engine/Source/Developer/TraceServices/Private/Model/ModuleProvider.cpp
@@ -99,7 +99,7 @@ private:
 	FCachedStringStore Strings;
 	// Efficient representation of symbols
 	TPagedArray<FResolvedSymbol> SymbolCache;
-	// Lookup table to for symbols
+	// Lookup table for symbols
 	TMap<uint64, const FResolvedSymbol*> SymbolCacheLookup;
 	// Number of cached symbols that was loaded (used for stats)
 	uint32 NumCachedSymbols;
@@ -428,7 +428,7 @@ void TModuleProvider<SymbolResolverType>::LoadSymbolsFromCache(IAnalysisCache& C
 		const TCHAR* File = Strings.GetStringAtOffset(Symbol.FileOffset);
 		if (Module == nullptr || Name == nullptr || File == nullptr)
 		{
-			UE_LOG(LogTraceServices, Warning, TEXT("Found cached symbol (adress %llx) which referenced unknown string."), Symbol.Address);
+			UE_LOG(LogTraceServices, Warning, TEXT("Found cached symbol (address %llx) which referenced unknown string."), Symbol.Address);
 			continue;
 		}
 		FResolvedSymbol& Resolved = SymbolCache.EmplaceBack(ESymbolQueryResult::OK, Module, Name, File, static_cast<uint16>(Symbol.Line), EResolvedSymbolFilterStatus::Unknown);
diff --git a/Engine/Source/Developer/TraceServices/Private/Model/SymslibResolver.cpp b/Engine/Source/Developer/TraceServices/Private/Model/SymslibResolver.cpp
--- a/Engine/Source/Developer/TraceServices/Private/Model/SymslibResolver.cpp
+++ b/Engine/Source/Developer/TraceServices/Private/Model/SymslibResolver.cpp
@@ -848,12 +848,14 @@ void FSymslibResolver::ResolveSymbols(TArrayView<FQueuedAddress>& QueuedWork)
 		}
 		ResolveSymbolTracked(ToResolve.Address, *ToResolve.Target, StringAllocator);
 	}
-	UE_LOG(LogSymslib, VeryVerbose, TEXT("String allocator: %.02f KiB used (%.02f KiB wasted) in %d blocks"),
-		(double)((StringAllocator.BlockUsed * StringAllocator.BlockSize - StringAllocator.BlockRemaining) * sizeof(TCHAR)) / 1024.0,
-		(double)(StringAllocator.BlockRemaining * sizeof(TCHAR)) / 1024.0,
+	uint64 BytesAllocated = StringAllocator.BlockUsed * StringAllocator.BlockSize * sizeof(TCHAR);
+	uint64 BytesWasted = StringAllocator.BlockRemaining * sizeof(TCHAR);
+	UE_LOG(LogSymslib, VeryVerbose, TEXT("String allocator: %.02f KiB used (+ %.02f KiB wasted) in %d blocks"),
+		(double)(BytesAllocated - BytesWasted) / 1024.0,
+		(double)(BytesWasted) / 1024.0,
 		StringAllocator.BlockUsed);
-	SymbolBytesAllocated.fetch_add(StringAllocator.BlockUsed * StringAllocator.BlockSize * sizeof(TCHAR));
-	SymbolBytesWasted.fetch_add(StringAllocator.BlockRemaining * sizeof(TCHAR));
+	SymbolBytesAllocated.fetch_add(BytesAllocated);
+	SymbolBytesWasted.fetch_add(BytesWasted);
 }
 
 ////////////////////////////////////////////////////////////////////////////////////////////////////
