Limit error logging from GpuProfilerAnalysis

diff --git a/Engine/Source/Developer/TraceServices/Private/Analyzers/GpuProfilerTraceAnalysis.cpp b/Engine/Source/Developer/TraceServices/Private/Analyzers/GpuProfilerTraceAnalysis.cpp
--- a/Engine/Source/Developer/TraceServices/Private/Analyzers/GpuProfilerTraceAnalysis.cpp
+++ b/Engine/Source/Developer/TraceServices/Private/Analyzers/GpuProfilerTraceAnalysis.cpp
@@ -44,6 +44,7 @@ void FGpuProfilerAnalyzer::OnAnalysisBegin(const FOnAnalysisContext& Context)
 
 void FGpuProfilerAnalyzer::OnAnalysisEnd()
 {
+	// Old Gpu Profiler errors.
 	if (NumFramesWithErrors > 0)
 	{
 		UE_LOG(LogTraceServices, Error, TEXT("[GpuProfiler] Frames with errors: %u"), NumFramesWithErrors);
@@ -53,6 +54,27 @@ void FGpuProfilerAnalyzer::OnAnalysisEnd()
 	{
 		UE_LOG(LogTraceServices, Log, TEXT("[GpuProfiler] Analysis completed (%u frames, %d timers)."), NumFrames, EventTypeMap.Num());
 	}
+
+	// New Gpu Profiler errors.
+	if (ErrorData.NumInterleavedEvents > 0)
+	{
+		UE_LOG(LogTraceServices, Error, TEXT("[GpuProfiler] Number of interleaved events: %u, max error %f"), ErrorData.NumInterleavedEvents, ErrorData.InterleavedEventsMaxDelta);
+	}
+
+	if (ErrorData.NumInterleavedAndReversedEvents > 0)
+	{
+		UE_LOG(LogTraceServices, Error, TEXT("[GpuProfiler] Number of interleaved and reversed events: %u, max error %f"), ErrorData.NumInterleavedAndReversedEvents, ErrorData.InterleavedAndReversedEventsMaxDelta);
+	}
+
+	if (ErrorData.NumMismatchedEvents > 0)
+	{
+		UE_LOG(LogTraceServices, Error, TEXT("[GpuProfiler] Number of mismatched events: %u"), ErrorData.NumMismatchedEvents);
+	}
+
+	if (ErrorData.NumNegativeDurationEvents > 0)
+	{
+		UE_LOG(LogTraceServices, Error, TEXT("[GpuProfiler] Number of negative duration events: %u, max error %f"), ErrorData.NumNegativeDurationEvents, ErrorData.NegativeDurationEventsMaxDelta);
+	}
 }
 
 bool FGpuProfilerAnalyzer::OnEvent(uint16 RouteId, EStyle Style, const FOnEventContext& Context)
@@ -445,10 +467,17 @@ void FGpuProfilerAnalyzer::BeginEvent(FQueue& Queue, int32 StackIndex, double Be
 		const double LastTimestamp = Timeline->GetLastTimestamp();
 		if (BeginEventTime < LastTimestamp)
 		{
-			UE_LOG(LogTraceServices, Error, TEXT("[GPU] %s BEGIN %u \"%s\" : Time %f < %f !!!"),
-				StackIndex == 1 ? TEXT("WORK") : TEXT("BREADCRUMB"),
-				BeginEventTimerId, GetTimerName(BeginEventTimerId),
-				BeginEventTime, LastTimestamp);
+			++ErrorData.NumInterleavedEvents;
+			ErrorData.InterleavedEventsMaxDelta = FMath::Max(ErrorData.InterleavedEventsMaxDelta, LastTimestamp - BeginEventTime);
+
+			if (ErrorData.NumInterleavedEvents < ErrorData.NumMaxErrors)
+			{
+				UE_LOG(LogTraceServices, Error, TEXT("[GPU] %s BEGIN %u \"%s\" : Time %f < %f !!!"),
+					StackIndex == 1 ? TEXT("WORK") : TEXT("BREADCRUMB"),
+					BeginEventTimerId, GetTimerName(BeginEventTimerId),
+					BeginEventTime, LastTimestamp);
+			}
+
 			BeginEventTime = LastTimestamp;
 		}
 		Timeline->AppendBeginEvent(BeginEventTime, Event);
@@ -472,17 +501,28 @@ void FGpuProfilerAnalyzer::EndEvent(FQueue& Queue, int32 StackIndex, double EndE
 	}
 	if (BeginEventTimerId != EndEventTimerId)
 	{
-		UE_LOG(LogTraceServices, Error, TEXT("[GPU] %s END %u \"%s\" (%f) != BEGIN %u \"%s\" (%f) !!!"),
-			StackIndex == 1 ? TEXT("WORK") : TEXT("BREADCRUMB"),
-			EndEventTimerId, GetTimerName(EndEventTimerId), EndEventTime,
-			BeginEventTimerId, GetTimerName(BeginEventTimerId), BeginEvent.Time);
+		++ErrorData.NumMismatchedEvents;
+
+		if (ErrorData.NumMismatchedEvents < ErrorData.NumMaxErrors)
+		{
+			UE_LOG(LogTraceServices, Error, TEXT("[GPU] %s END %u \"%s\" (%f) != BEGIN %u \"%s\" (%f) !!!"),
+				StackIndex == 1 ? TEXT("WORK") : TEXT("BREADCRUMB"),
+				EndEventTimerId, GetTimerName(EndEventTimerId), EndEventTime,
+				BeginEventTimerId, GetTimerName(BeginEventTimerId), BeginEvent.Time);
+		}
 	}
 	if (BeginEvent.Time > EndEventTime)
 	{
-		UE_LOG(LogTraceServices, Error, TEXT("[GPU] %s END %u \"%s\" (%f) < BEGIN %u \"%s\" (%f) !!!"),
-			StackIndex == 1 ? TEXT("WORK") : TEXT("BREADCRUMB"),
-			EndEventTimerId, GetTimerName(EndEventTimerId), EndEventTime,
-			BeginEventTimerId, GetTimerName(BeginEventTimerId), BeginEvent.Time);
+		++ErrorData.NumNegativeDurationEvents;
+		ErrorData.NegativeDurationEventsMaxDelta = FMath::Max(ErrorData.NegativeDurationEventsMaxDelta, BeginEvent.Time - EndEventTime);
+
+		if (ErrorData.NumNegativeDurationEvents < ErrorData.NumMaxErrors)
+		{
+			UE_LOG(LogTraceServices, Error, TEXT("[GPU] %s END %u \"%s\" (%f) < BEGIN %u \"%s\" (%f) !!!"),
+				StackIndex == 1 ? TEXT("WORK") : TEXT("BREADCRUMB"),
+				EndEventTimerId, GetTimerName(EndEventTimerId), EndEventTime,
+				BeginEventTimerId, GetTimerName(BeginEventTimerId), BeginEvent.Time);
+		}
 	}
 
 	FAnalysisSessionEditScope _(Session);
@@ -494,11 +534,17 @@ void FGpuProfilerAnalyzer::EndEvent(FQueue& Queue, int32 StackIndex, double EndE
 		const double LastTimestamp = Timeline->GetLastTimestamp();
 		if (EndEventTime < LastTimestamp)
 		{
-			UE_LOG(LogTraceServices, Error, TEXT("[GPU] %s END %u \"%s\" : Time %f < %f !!!"),
-				StackIndex == 1 ? TEXT("WORK") : TEXT("BREADCRUMB"),
-				EndEventTimerId, GetTimerName(EndEventTimerId),
-				EndEventTime, LastTimestamp);
-			EndEventTime = LastTimestamp;
+			++ErrorData.NumInterleavedAndReversedEvents;
+			ErrorData.InterleavedAndReversedEventsMaxDelta = FMath::Max(ErrorData.InterleavedAndReversedEventsMaxDelta, LastTimestamp - EndEventTime);
+
+			if (ErrorData.NumInterleavedAndReversedEvents < ErrorData.NumMaxErrors)
+			{
+				UE_LOG(LogTraceServices, Error, TEXT("[GPU] %s END %u \"%s\" : Time %f < %f !!!"),
+					StackIndex == 1 ? TEXT("WORK") : TEXT("BREADCRUMB"),
+					EndEventTimerId, GetTimerName(EndEventTimerId),
+					EndEventTime, LastTimestamp);
+				EndEventTime = LastTimestamp;
+			}
 		}
 		Timeline->AppendEndEvent(EndEventTime);
 		Session.UpdateDurationSeconds(EndEventTime);
diff --git a/Engine/Source/Developer/TraceServices/Private/Analyzers/GpuProfilerTraceAnalysis.h b/Engine/Source/Developer/TraceServices/Private/Analyzers/GpuProfilerTraceAnalysis.h
--- a/Engine/Source/Developer/TraceServices/Private/Analyzers/GpuProfilerTraceAnalysis.h
+++ b/Engine/Source/Developer/TraceServices/Private/Analyzers/GpuProfilerTraceAnalysis.h
@@ -53,6 +53,20 @@ private:
 		double LastTime = 0.0f;
 	};
 
+	struct FErrorData
+	{
+		const uint32 NumMaxErrors = 100;
+
+		uint32 NumInterleavedEvents = 0;
+		uint32 NumInterleavedAndReversedEvents = 0;
+		uint32 NumMismatchedEvents = 0;
+		uint32 NumNegativeDurationEvents = 0;
+
+		double InterleavedEventsMaxDelta = 0.0f;
+		double InterleavedAndReversedEventsMaxDelta = 0.0f;
+		double NegativeDurationEventsMaxDelta = 0.0f;
+	};
+
 public:
 	FGpuProfilerAnalyzer(FAnalysisSession& Session, FTimingProfilerProvider& TimingProfilerProvider, IEditableCounterProvider& EditableCounterProvider);
 	virtual void OnAnalysisBegin(const FOnAnalysisContext& Context) override;
@@ -82,6 +96,7 @@ private:
 	TMap<FString, uint32> BreadcrumbMap; // breadcrumb name --> GPU timer id
 	TMap<uint32, const TCHAR*> TimerMap; // GPU timer id --> persistent timer name
 	TMap<uint32, FQueue> Queues; // QueueId --> FQueue
+	FErrorData ErrorData;
 
 	//////////////////////////////////////////////////
 	// The Old GPU Profiler
