[Insights] MemoryInsights: Added analysis support for MemoryTrace_UpdateAlloc (to update context -- tag and metadata -- for an allocation) and for UE_MEMSCOPE_PTR(InPtr, InRootHeap).

diff --git a/Engine/Source/Developer/TraceServices/Private/Analyzers/AllocationsAnalysis.cpp b/Engine/Source/Developer/TraceServices/Private/Analyzers/AllocationsAnalysis.cpp
--- a/Engine/Source/Developer/TraceServices/Private/Analyzers/AllocationsAnalysis.cpp
+++ b/Engine/Source/Developer/TraceServices/Private/Analyzers/AllocationsAnalysis.cpp
@@ -58,14 +58,15 @@
 	Builder.RouteEvent(RouteId_ReallocAllocSystem,  "Memory", "ReallocAllocSystem");
 	Builder.RouteEvent(RouteId_ReallocFree,         "Memory", "ReallocFree");
 	Builder.RouteEvent(RouteId_ReallocFreeSystem,   "Memory", "ReallocFreeSystem");
+	Builder.RouteEvent(RouteId_HeapMarkAlloc,       "Memory", "HeapMarkAlloc");
+	Builder.RouteEvent(RouteId_HeapUnmarkAlloc,     "Memory", "HeapUnmarkAlloc");
+	Builder.RouteEvent(RouteId_UpdateAlloc,         "Memory", "UpdateAlloc");
+	Builder.RouteEvent(RouteId_MemSwapOp,           "Memory", "MemorySwapOp");
 	Builder.RouteEvent(RouteId_Marker,              "Memory", "Marker");
 	Builder.RouteEvent(RouteId_TagSpec,             "Memory", "TagSpec");
 	Builder.RouteEvent(RouteId_HeapSpec,            "Memory", "HeapSpec");
-	Builder.RouteEvent(RouteId_HeapMarkAlloc,       "Memory", "HeapMarkAlloc");
-	Builder.RouteEvent(RouteId_HeapUnmarkAlloc,     "Memory", "HeapUnmarkAlloc");
 	Builder.RouteEvent(RouteId_MemScopeTag,         "Memory", "MemoryScope", true);
 	Builder.RouteEvent(RouteId_MemScopePtr,         "Memory", "MemoryScopePtr", true);
-	Builder.RouteEvent(RouteId_MemSwapOp,           "Memory", "MemorySwapOp");
 
 #if INSIGHTS_MEM_TRACE_METADATA_TEST
 	{
@@ -133,19 +134,6 @@
 			break;
 		}
 
-		case RouteId_HeapSpec:
-		{
-			const HeapId Id = EventData.GetValue<uint16>("Id");
-			const HeapId ParentId = EventData.GetValue<uint16>("ParentId");
-			const EMemoryTraceHeapFlags Flags = EventData.GetValue<EMemoryTraceHeapFlags>("Flags");
-			FStringView Name;
-			EventData.GetString("Name", Name);
-
-			FProviderEditScopeLock _(AllocationsProvider);
-			AllocationsProvider.EditHeapSpec(Id, ParentId, Name, Flags);
-			break;
-		}
-
 		case RouteId_AllocSystem:
 		case RouteId_AllocVideo:
 		case RouteId_ReallocAllocSystem:
@@ -166,11 +154,8 @@
 			const uint32 ThreadId = Context.ThreadInfo.GetId();
 			const double Time = GetCurrentTime();
 
-			// CallstackId is optional. If the field is not present CallstackId will be 0 (i.e. "no callstack").
-			const uint32 CallstackId = EventData.GetValue<uint32>("CallstackId", 0);
-
 			const uint64 Address = EventData.GetValue<uint64>("Address");
-
+			const uint32 CallstackId = EventData.GetValue<uint32>("CallstackId", 0);
 			RootHeap = EventData.GetValue<uint8>("RootHeap", static_cast<uint8>(RootHeap));
 
 			uint64 SizeUpper = EventData.GetValue<uint32>("Size");
@@ -203,11 +188,8 @@
 			const uint32 ThreadId = Context.ThreadInfo.GetId();
 			const double Time = GetCurrentTime();
 
-			// CallstackId is optional. If the field is not present CallstackId will be 0 (i.e. "no callstack").
-			const uint32 CallstackId = EventData.GetValue<uint32>("CallstackId", 0);
-
 			const uint64 Address = EventData.GetValue<uint64>("Address");
-
+			const uint32 CallstackId = EventData.GetValue<uint32>("CallstackId", 0);
 			RootHeap = EventData.GetValue<uint8>("RootHeap", static_cast<uint8>(RootHeap));
 
 			FProviderEditScopeLock _(AllocationsProvider);
@@ -225,10 +207,8 @@
 			const uint32 ThreadId = Context.ThreadInfo.GetId();
 			const double Time = GetCurrentTime();
 
-			// CallstackId is optional. If the field is not present CallstackId will be 0 (i.e. "no callstack").
-			const uint32 CallstackId = EventData.GetValue<uint32>("CallstackId", 0);
-
 			const uint64 Address = EventData.GetValue<uint64>("Address");
+			const uint32 CallstackId = EventData.GetValue<uint32>("CallstackId", 0);
 			const HeapId Heap = EventData.GetValue<uint16>("Heap", 0);
 			const EMemoryTraceHeapAllocationFlags Flags = EventData.GetValue<EMemoryTraceHeapAllocationFlags>("Flags");
 
@@ -242,10 +222,8 @@
 			const uint32 ThreadId = Context.ThreadInfo.GetId();
 			const double Time = GetCurrentTime();
 
-			// CallstackId is optional. If the field is not present CallstackId will be 0 (i.e. "no callstack").
-			const uint32 CallstackId = EventData.GetValue<uint32>("CallstackId", 0);
-
 			const uint64 Address = EventData.GetValue<uint64>("Address");
+			const uint32 CallstackId = EventData.GetValue<uint32>("CallstackId", 0);
 			const HeapId Heap = EventData.GetValue<uint16>("Heap", 0);
 
 			FProviderEditScopeLock _(AllocationsProvider);
@@ -253,6 +231,34 @@
 			break;
 		}
 
+		case RouteId_UpdateAlloc: // added in UE 5.6
+		{
+			const uint32 ThreadId = Context.ThreadInfo.GetId();
+			const double Time = GetCurrentTime();
+
+			const uint64 Address = EventData.GetValue<uint64>("Address");
+			const uint32 CallstackId = EventData.GetValue<uint32>("CallstackId", 0);
+			RootHeap = EventData.GetValue<uint8>("RootHeap", uint8(EMemoryTraceRootHeap::SystemMemory));
+
+			FProviderEditScopeLock _(AllocationsProvider);
+			AllocationsProvider.EditUpdateAlloc(ThreadId, Time, CallstackId, Address, RootHeap);
+		}
+
+		case RouteId_MemSwapOp: // added in UE 5.5
+		{
+			const uint32 ThreadId = Context.ThreadInfo.GetId();
+			const double Time = GetCurrentTime();
+
+			const uint64 Address = EventData.GetValue<uint64>("Address");
+			const uint32 CallstackId = EventData.GetValue<uint32>("CallstackId", 0);
+			const uint32 CompressedSize = EventData.GetValue<uint32>("CompressedSize");
+			const EMemoryTraceSwapOperation SwapOp = (EMemoryTraceSwapOperation)EventData.GetValue<uint8>("SwapOp");
+
+			FProviderEditScopeLock _(AllocationsProvider);
+			AllocationsProvider.EditSwapOp(ThreadId, Time, Address, SwapOp, CompressedSize, CallstackId);
+			break;
+		}
+
 		case RouteId_Marker:
 		{
 			// If BaseCycle is 0, then Cycle is a 64-bit absolute value, otherwise Cycle is a 32-bit value (relative to BaseCycle).
@@ -285,7 +291,7 @@
 			const TagIdType Parent = Context.EventData.GetValue<TagIdType>("Parent");
 
 			FString Display;
-			Context.EventData.GetString("Display", Display);
+			Context.EventData.GetString("Display", Display); // Display is UE::Trace::AnsiString
 			const TCHAR* DisplayString = Session.StoreString(*Display);
 
 			FProviderEditScopeLock _(AllocationsProvider);
@@ -293,6 +299,20 @@
 			break;
 		}
 
+		case RouteId_HeapSpec:
+		{
+			const HeapId Id = EventData.GetValue<uint16>("Id");
+			const HeapId ParentId = EventData.GetValue<uint16>("ParentId");
+			const EMemoryTraceHeapFlags Flags = EventData.GetValue<EMemoryTraceHeapFlags>("Flags");
+
+			FStringView Name;
+			EventData.GetString("Name", Name); // Name is UE::Trace::WideString
+
+			FProviderEditScopeLock _(AllocationsProvider);
+			AllocationsProvider.EditHeapSpec(Id, ParentId, Name, Flags);
+			break;
+		}
+
 		case RouteId_MemScopeTag: // "MemoryScope", see UE_MEMSCOPE
 		{
 			const uint32 ThreadId = Context.ThreadInfo.GetId();
@@ -336,9 +356,7 @@
 			if (Style == EStyle::EnterScope)
 			{
 				const uint64 Ptr = Context.EventData.GetValue<uint64>("Ptr");
+				RootHeap = EventData.GetValue<uint8>("RootHeap", uint8(EMemoryTraceRootHeap::SystemMemory));
 
-				// RootHeap defaults to EMemoryTraceRootHeap::SystemMemory
-				//RootHeap = EventData.GetValue<uint8>("RootHeap", static_cast<uint8>(RootHeap)); // TODO: UE_MEMSCOPE_PTR(InPtr, InRootHeap)
-
 				{
 					FProviderEditScopeLock _(AllocationsProvider);
@@ -368,21 +386,6 @@
 			}
 			break;
 		}
-
-		case RouteId_MemSwapOp:
-		{
-			const uint32 ThreadId = Context.ThreadInfo.GetId();
-			const double Time = GetCurrentTime();
-
-			const uint64 Address = EventData.GetValue<uint64>("Address");
-			const uint32 CallstackId = EventData.GetValue<uint32>("CallstackId");
-			const uint32 CompressedSize = EventData.GetValue<uint32>("CompressedSize");
-			const EMemoryTraceSwapOperation SwapOp = (EMemoryTraceSwapOperation)EventData.GetValue<uint8>("SwapOp");
-
-			FProviderEditScopeLock _(AllocationsProvider);
-			AllocationsProvider.EditSwapOp(ThreadId, Time, Address, SwapOp, CompressedSize, CallstackId);
-			break;
-		}
 	}
 
 	return true;
diff --git a/Engine/Source/Developer/TraceServices/Private/Analyzers/AllocationsAnalysis.h b/Engine/Source/Developer/TraceServices/Private/Analyzers/AllocationsAnalysis.h
--- a/Engine/Source/Developer/TraceServices/Private/Analyzers/AllocationsAnalysis.h
+++ b/Engine/Source/Developer/TraceServices/Private/Analyzers/AllocationsAnalysis.h
@@ -28,14 +28,15 @@ private:
 		RouteId_ReallocAllocSystem,
 		RouteId_ReallocFree,
 		RouteId_ReallocFreeSystem,
+		RouteId_HeapMarkAlloc,
+		RouteId_HeapUnmarkAlloc,
+		RouteId_UpdateAlloc, // added in UE 5.6
+		RouteId_MemSwapOp, // added in UE 5.5
 		RouteId_Marker,
 		RouteId_TagSpec,
 		RouteId_HeapSpec,
-		RouteId_HeapMarkAlloc,
-		RouteId_HeapUnmarkAlloc,
 		RouteId_MemScopeTag,
 		RouteId_MemScopePtr,
-		RouteId_MemSwapOp,
 	};
 
 public:
diff --git a/Engine/Source/Developer/TraceServices/Private/Model/AllocationsProvider.cpp b/Engine/Source/Developer/TraceServices/Private/Model/AllocationsProvider.cpp
--- a/Engine/Source/Developer/TraceServices/Private/Model/AllocationsProvider.cpp
+++ b/Engine/Source/Developer/TraceServices/Private/Model/AllocationsProvider.cpp
@@ -1466,6 +1466,13 @@ void FAllocationsProvider::EditAlloc(uint32 ThreadId, double Time, uint32 Callst
 		return;
 	}
 
+	AllocInternal(ThreadId, Time, CallstackId, Address, Size, Alignment, RootHeapId);
+}
+
+////////////////////////////////////////////////////////////////////////////////////////////////////
+
+FAllocationItem* FAllocationsProvider::AllocInternal(uint32 ThreadId, double Time, uint32 CallstackId, uint64 Address, uint64 Size, uint32 Alignment, HeapId RootHeapId)
+{
 	INSIGHTS_WATCH_API_LOGF(ThreadId, Time, Address, TEXT("Alloc 0x%llX : Size=%llu Alignment=%u RootHeap=%u CallstackId=%u"), Address, Size, Alignment, RootHeapId, CallstackId);
 	INSIGHTS_FILTER_EVENT(ThreadId, Time, Address, RootHeapId, CallstackId);
 
@@ -1474,7 +1481,7 @@ void FAllocationsProvider::EditAlloc(uint32 ThreadId, double Time, uint32 Callst
 #if !INSIGHTS_WARNINGS_FOR_NULLPTR_ALLOCS
 		if (Size == 0)
 		{
-			return;
+			return nullptr;
 		}
 #endif // INSIGHTS_WARNINGS_FOR_NULLPTR_ALLOCS
 		++AllocWarnings;
@@ -1482,7 +1489,7 @@ void FAllocationsProvider::EditAlloc(uint32 ThreadId, double Time, uint32 Callst
 		{
 			UE_LOG(LogTraceServices, Warning, TEXT("[MemAlloc] Alloc at address 0 : Size=%llu, RootHeap=%u, Time=%f, CallstackId=%u"), Size, RootHeapId, Time, CallstackId);
 		}
-		return;
+		return nullptr;
 	}
 
 	if (!IsValidRootHeap(RootHeapId))
@@ -1492,7 +1499,7 @@ void FAllocationsProvider::EditAlloc(uint32 ThreadId, double Time, uint32 Callst
 		{
 			UE_LOG(LogTraceServices, Warning, TEXT("[MemAlloc] Alloc with invalid root heap id (%u)."), RootHeapId);
 		}
-		return;
+		return nullptr;
 	}
 
 	FRootHeap& RootHeap = *RootHeaps[RootHeapId];
@@ -1519,7 +1526,7 @@ void FAllocationsProvider::EditAlloc(uint32 ThreadId, double Time, uint32 Callst
 		// Free the previous allocation.
 		INSIGHTS_WATCH_INDIRECT_API_LOGF(TEXT("Free"), ThreadId, Time, Address);
 		constexpr uint32 FreeCallstackId = 0; // no callstack
-		EditFree(ThreadId, Time, FreeCallstackId, Address, RootHeapId);
+		FreeInternal(ThreadId, Time, FreeCallstackId, Address, RootHeapId);
 		RootHeap.SbTree->SetTimeForEvent(RootHeap.EventIndex, Time); // for the case where the next event is first event in a new SbTree column after changing the heap alloc
 
 #else // INSIGHTS_DOUBLE_ALLOC_FREE_PREVIOUS
@@ -1534,8 +1541,9 @@ void FAllocationsProvider::EditAlloc(uint32 ThreadId, double Time, uint32 Callst
 	}
 #endif // INSIGHTS_VALIDATE_ALLOC_EVENTS
 
+	FAllocationItem* AllocationPtr = RootHeap.LiveAllocs->AddNew(Address);
+
 	{
-		FAllocationItem* AllocationPtr = RootHeap.LiveAllocs->AddNew(Address);
 		FAllocationItem& Allocation = *AllocationPtr;
 
 		uint32 MetadataId = MetadataProvider.InvalidMetadataId;
@@ -1598,6 +1606,8 @@ void FAllocationsProvider::EditAlloc(uint32 ThreadId, double Time, uint32 Callst
 		UE_LOG(LogTraceServices, Error, TEXT("[MemAlloc] Too many events!"));
 		bInitialized = false; // ignore further events
 	}
+
+	return AllocationPtr;
 }
 
 ////////////////////////////////////////////////////////////////////////////////////////////////////
@@ -1611,6 +1621,13 @@ void FAllocationsProvider::EditFree(uint32 ThreadId, double Time, uint32 Callsta
 		return;
 	}
 
+	FreeInternal(ThreadId, Time, CallstackId, Address, RootHeapId);
+}
+
+////////////////////////////////////////////////////////////////////////////////////////////////////
+
+FAllocationItem* FAllocationsProvider::FreeInternal(uint32 ThreadId, double Time, uint32 CallstackId, uint64 Address, HeapId RootHeapId)
+{
 	INSIGHTS_WATCH_API_LOGF(ThreadId, Time, Address, TEXT("Free 0x%llX : RootHeap=%u CallstackId=%u"), Address, RootHeapId, CallstackId);
 	INSIGHTS_FILTER_EVENT(ThreadId, Time, Address, RootHeapId, CallstackId);
 
@@ -1623,7 +1640,7 @@ void FAllocationsProvider::EditFree(uint32 ThreadId, double Time, uint32 Callsta
 			UE_LOG(LogTraceServices, Warning, TEXT("[MemAlloc] Free for address 0 : RootHeap=%u, Time=%f, CallstackId=%u"), RootHeapId, Time, CallstackId);
 		}
 #endif // INSIGHTS_WARNINGS_FOR_NULLPTR_ALLOCS
-		return;
+		return nullptr;
 	}
 
 	if (!IsValidRootHeap(RootHeapId))
@@ -1633,7 +1650,7 @@ void FAllocationsProvider::EditFree(uint32 ThreadId, double Time, uint32 Callsta
 		{
 			UE_LOG(LogTraceServices, Warning, TEXT("[MemAlloc] Free with invalid root heap id (%u)."), RootHeapId);
 		}
-		return;
+		return nullptr;
 	}
 
 	FRootHeap& RootHeap = *RootHeaps[RootHeapId];
@@ -1674,7 +1691,7 @@ void FAllocationsProvider::EditFree(uint32 ThreadId, double Time, uint32 Callsta
 		constexpr uint64 FakeAllocSize = 0;
 		constexpr uint32 FakeAllocAlignment = 0;
 		INSIGHTS_WATCH_INDIRECT_API_LOGF(TEXT("Alloc"), ThreadId, Time, Address);
-		EditAlloc(ThreadId, Time, CallstackId, Address, FakeAllocSize, FakeAllocAlignment, RootHeapId);
+		AllocInternal(ThreadId, Time, CallstackId, Address, FakeAllocSize, FakeAllocAlignment, RootHeapId);
 		RootHeap.SbTree->SetTimeForEvent(RootHeap.EventIndex, Time); // for the case where the next event is first event in a new SbTree column after adding the fake alloc
 		AllocationPtr = RootHeap.LiveAllocs->Remove(Address); // we take ownership of AllocationPtr
 	}
@@ -1727,6 +1744,32 @@ void FAllocationsProvider::EditFree(uint32 ThreadId, double Time, uint32 Callsta
 		UE_LOG(LogTraceServices, Error, TEXT("[MemAlloc] Too many events!"));
 		bInitialized = false; // ignore further events
 	}
+
+	return AllocationPtr;
+}
+
+////////////////////////////////////////////////////////////////////////////////////////////////////
+
+void FAllocationsProvider::EditUpdateAlloc(uint32 ThreadId, double Time, uint32 CallstackId, uint64 Address, HeapId RootHeapId)
+{
+	EditAccessCheck();
+
+	if (!bInitialized)
+	{
+		return;
+	}
+
+	FAllocationItem* AllocationPtr = FreeInternal(ThreadId, Time, CallstackId, Address, RootHeapId);
+	if (AllocationPtr)
+	{
+		AllocInternal(ThreadId,
+			Time,
+			AllocationPtr->AllocCallstackId,
+			Address,
+			AllocationPtr->GetSize(),
+			AllocationPtr->GetAlignment(),
+			RootHeapId);
+	}
 }
 
 ////////////////////////////////////////////////////////////////////////////////////////////////////
@@ -2091,7 +2134,7 @@ void FAllocationsProvider::EditMarkAllocationAsHeap(uint32 ThreadId, double Time
 		constexpr uint64 FakeAllocSize = 0;
 		constexpr uint32 FakeAllocAlignment = 0;
 		INSIGHTS_WATCH_INDIRECT_API_LOGF(TEXT("Alloc"), ThreadId, Time, Address);
-		EditAlloc(ThreadId, Time, CallstackId, Address, FakeAllocSize, FakeAllocAlignment, RootHeap.HeapSpec->Id);
+		AllocInternal(ThreadId, Time, CallstackId, Address, FakeAllocSize, FakeAllocAlignment, RootHeap.HeapSpec->Id);
 		RootHeap.SbTree->SetTimeForEvent(RootHeap.EventIndex, Time); // for the case where the next event is first event in a new SbTree column after adding the fake alloc
 		Alloc = RootHeap.LiveAllocs->Remove(Address); // we take ownership of Alloc
 	}
@@ -2200,7 +2243,7 @@ void FAllocationsProvider::EditUnmarkAllocationAsHeap(uint32 ThreadId, double Ti
 			constexpr uint64 FakeAllocSize = 0;
 			constexpr uint32 FakeAllocAlignment = 0;
 			INSIGHTS_WATCH_INDIRECT_API_LOGF(TEXT("Alloc"), ThreadId, Time, Address);
-			EditAlloc(ThreadId, Time, CallstackId, Address, FakeAllocSize, FakeAllocAlignment, RootHeap.HeapSpec->Id);
+			AllocInternal(ThreadId, Time, CallstackId, Address, FakeAllocSize, FakeAllocAlignment, RootHeap.HeapSpec->Id);
 			RootHeap.SbTree->SetTimeForEvent(RootHeap.EventIndex, Time); // for the case where the next event is first event in a new SbTree column after adding the fake alloc
 			Alloc = RootHeap.LiveAllocs->Remove(Address); // we take ownership of Alloc
 			check(Alloc != nullptr);
@@ -2297,7 +2340,7 @@ void FAllocationsProvider::EditUnmarkAllocationAsHeap(uint32 ThreadId, double Ti
 				for (const FAllocInHeap& AllocInHeap : AllocsInHeap)
 				{
 					INSIGHTS_WATCH_INDIRECT_API_LOGF(TEXT("Free"), ThreadId, Time, AllocInHeap.Address);
-					EditFree(ThreadId, Time, CallstackId, AllocInHeap.Address, RootHeap.HeapSpec->Id);
+					FreeInternal(ThreadId, Time, CallstackId, AllocInHeap.Address, RootHeap.HeapSpec->Id);
 				}
 			}
 		}
@@ -2314,9 +2357,9 @@ void FAllocationsProvider::EditUnmarkAllocationAsHeap(uint32 ThreadId, double Ti
 		const uint8 Tracker = 1;
 		TagTracker.PushTagFromPtr(ThreadId, Tracker, Alloc->Tag);
 		INSIGHTS_WATCH_INDIRECT_API_LOGF(TEXT("Free"), ThreadId, Time, Address);
-		EditFree(ThreadId, Time, CallstackId, Address, RootHeap.HeapSpec->Id);
+		FreeInternal(ThreadId, Time, CallstackId, Address, RootHeap.HeapSpec->Id);
 		INSIGHTS_WATCH_INDIRECT_API_LOGF(TEXT("Alloc"), ThreadId, Time, Address);
-		EditAlloc(ThreadId, Time, AllocCallstackId, Address, Size, Alignment, RootHeap.HeapSpec->Id);
+		AllocInternal(ThreadId, Time, AllocCallstackId, Address, Size, Alignment, RootHeap.HeapSpec->Id);
 		TagTracker.PopTagFromPtr(ThreadId, Tracker);
 	}
 	else
diff --git a/Engine/Source/Developer/TraceServices/Private/Model/AllocationsProvider.h b/Engine/Source/Developer/TraceServices/Private/Model/AllocationsProvider.h
--- a/Engine/Source/Developer/TraceServices/Private/Model/AllocationsProvider.h
+++ b/Engine/Source/Developer/TraceServices/Private/Model/AllocationsProvider.h
@@ -374,6 +374,7 @@ public:
 
 	void EditAlloc(uint32 ThreadId, double Time, uint32 CallstackId, uint64 Address, uint64 Size, uint32 Alignment, HeapId RootHeap);
 	void EditFree(uint32 ThreadId, double Time, uint32 CallstackId, uint64 Address, HeapId RootHeap);
+	void EditUpdateAlloc(uint32 ThreadId, double Time, uint32 CallstackId, uint64 Address, HeapId RootHeap);
 	void EditSwapOp(uint32 ThreadId, double Time, uint64 UnmaskedPageAddress, EMemoryTraceSwapOperation SwapOp, uint64 CompressedPageSize, uint32 CallstackId);
 
 	void EditHeapSpec(HeapId Id, HeapId ParentId, const FStringView& Name, EMemoryTraceHeapFlags Flags);
@@ -391,6 +392,9 @@ public:
 	//////////////////////////////////////////////////
 
 private:
+	FAllocationItem* AllocInternal(uint32 ThreadId, double Time, uint32 CallstackId, uint64 Address, uint64 Size, uint32 Alignment, HeapId RootHeap);
+	FAllocationItem* FreeInternal(uint32 ThreadId, double Time, uint32 CallstackId, uint64 Address, HeapId RootHeap);
+
 	bool ShouldIgnoreEvent(uint32 ThreadId, double Time, uint64 Address, HeapId RootHeapId, uint32 CallstackId);
 	void AdvanceTimelines(double Time);
 	void AddHeapSpec(HeapId Id, HeapId ParentId, const FStringView& Name, EMemoryTraceHeapFlags Flags);
