Trace GpuProfiler wait events. Fix frame boundary event.

diff --git a/Engine/Source/Developer/TraceServices/Private/Analyzers/GpuProfilerTraceAnalysis.cpp b/Engine/Source/Developer/TraceServices/Private/Analyzers/GpuProfilerTraceAnalysis.cpp
--- a/Engine/Source/Developer/TraceServices/Private/Analyzers/GpuProfilerTraceAnalysis.cpp
+++ b/Engine/Source/Developer/TraceServices/Private/Analyzers/GpuProfilerTraceAnalysis.cpp
@@ -31,6 +31,7 @@ void FGpuProfilerAnalyzer::OnAnalysisBegin(const FOnAnalysisContext& Context)
 	Builder.RouteEvent(RouteId_EventEndBreadcrumb, "GpuProfiler", "EventEndBreadcrumb");
 	Builder.RouteEvent(RouteId_EventBeginWork, "GpuProfiler", "EventBeginWork");
 	Builder.RouteEvent(RouteId_EventEndWork, "GpuProfiler", "EventEndWork");
+	Builder.RouteEvent(RouteId_EventWait, "GpuProfiler", "EventWait");
 	Builder.RouteEvent(RouteId_EventStats, "GpuProfiler", "EventStats");
 	Builder.RouteEvent(RouteId_SignalFence, "GpuProfiler", "SignalFence");
 	Builder.RouteEvent(RouteId_WaitFence, "GpuProfiler", "WaitFence");
@@ -233,6 +234,43 @@ bool FGpuProfilerAnalyzer::OnEvent(uint16 RouteId, EStyle Style, const FOnEventC
 		EndEvent(Queue, 1, Time, GpuWorkTimerId);
 		break;
 	}
+	case RouteId_EventWait:
+	{
+		uint32 QueueId = EventData.GetValue<uint32>("QueueId");
+		const uint64 StartTimestamp = EventData.GetValue<uint64>("StartTime");
+		const double StartTime = Context.EventTime.AsSeconds(StartTimestamp);
+
+		const uint64 EndTimestamp = EventData.GetValue<uint64>("EndTime");
+		const double EndTime = Context.EventTime.AsSeconds(EndTimestamp);
+
+		UE_LOG(LogTraceServices, VeryVerbose, TEXT("[GPU] EventWait QueueId=%u, StartTime=%f, EndTime=%f"), QueueId, StartTime, EndTime);
+
+		if (GpuWaitTimerId == ~0u)
+		{
+			const FString GpuWaitTimerName(TEXT("GpuWait"));
+			GpuWaitTimerId = GetOrAddTimerId(GpuWaitTimerName);
+		}
+
+		FQueue& Queue = GetOrAddQueue(QueueId);
+		TArray<uint8> CborData;
+		{
+			CborData.Reserve(256);
+			FMemoryWriter MemoryWriter(CborData, false, true);
+			FCborWriter CborWriter(&MemoryWriter, ECborEndianness::StandardCompliant);
+			CborWriter.WriteContainerStart(ECborCode::Map, 1); // 1 is the FieldCount
+			CborWriter.WriteValue("Frame", 5);
+			CborWriter.WriteValue((int64)Queue.FrameNumber);
+		}
+		uint32 MetadataTimerId;
+		{
+			FAnalysisSessionEditScope _(Session);
+			MetadataTimerId = TimingProfilerProvider.AddMetadata(GpuWaitTimerId, MoveTemp(CborData));
+		}
+
+		BeginEvent(Queue, 1, StartTime, MetadataTimerId);
+		EndEvent(Queue, 1, EndTime, GpuWaitTimerId);
+		break;
+	}
 	case RouteId_EventStats:
 	{
 		const uint32 QueueId = EventData.GetValue<uint32>("QueueId");
diff --git a/Engine/Source/Developer/TraceServices/Private/Analyzers/GpuProfilerTraceAnalysis.h b/Engine/Source/Developer/TraceServices/Private/Analyzers/GpuProfilerTraceAnalysis.h
--- a/Engine/Source/Developer/TraceServices/Private/Analyzers/GpuProfilerTraceAnalysis.h
+++ b/Engine/Source/Developer/TraceServices/Private/Analyzers/GpuProfilerTraceAnalysis.h
@@ -25,6 +25,7 @@ private:
 		RouteId_EventEndBreadcrumb,
 		RouteId_EventBeginWork,
 		RouteId_EventEndWork,
+		RouteId_EventWait,
 		RouteId_EventStats,
 		RouteId_SignalFence,
 		RouteId_WaitFence,
@@ -93,6 +94,7 @@ private:
 
 	uint32 Version = 0;
 	uint32 GpuWorkTimerId = ~0;
+	uint32 GpuWaitTimerId = ~0;
 	TMap<FString, uint32> BreadcrumbMap; // breadcrumb name --> GPU timer id
 	TMap<uint32, const TCHAR*> TimerMap; // GPU timer id --> persistent timer name
 	TMap<uint32, FQueue> Queues; // QueueId --> FQueue
diff --git a/Engine/Source/Runtime/RHI/Private/GPUProfiler.cpp b/Engine/Source/Runtime/RHI/Private/GPUProfiler.cpp
--- a/Engine/Source/Runtime/RHI/Private/GPUProfiler.cpp
+++ b/Engine/Source/Runtime/RHI/Private/GPUProfiler.cpp
@@ -2030,9 +2030,12 @@ namespace UE::RHI::GPUProfiler
 				FGpuProfilerTrace::WaitFence(Queue.Value, ResolvedWait.GPUTimestampTOP, WaitFence.Queue.Value, WaitFence.Value);
 
 				// Time the queue spent waiting for the fence to signal on another queue.
-				uint64 Wait = ResolvedWait.GPUTimestampTOP > WaitStart
-					? ResolvedWait.GPUTimestampTOP - WaitStart
-					: 0;
+				uint64 Wait = 0;
+				if (ResolvedWait.GPUTimestampTOP > WaitStart)
+				{
+					Wait = ResolvedWait.GPUTimestampTOP - WaitStart;
+					FGpuProfilerTrace::TraceWait(Queue.Value, WaitStart, ResolvedWait.GPUTimestampTOP);
+				}
 
 				// Bring the last GPU busy end time forwards to where the wait is resolved.
 				LastGPUCycles = ResolvedWait.GPUTimestampTOP;
@@ -2055,6 +2058,8 @@ namespace UE::RHI::GPUProfiler
 				check(!bBusy);
 				Timestamps.CPUFrameBoundary = Event.CPUTimestamp;
 
+				FGpuProfilerTrace::FrameBoundary(Queue.Value, Event.FrameNumber);
+
 			#if WITH_PROFILEGPU
 				if (Profile.bProfileFrame)
 				{
diff --git a/Engine/Source/Runtime/RHI/Private/GpuProfilerTrace.cpp b/Engine/Source/Runtime/RHI/Private/GpuProfilerTrace.cpp
--- a/Engine/Source/Runtime/RHI/Private/GpuProfilerTrace.cpp
+++ b/Engine/Source/Runtime/RHI/Private/GpuProfilerTrace.cpp
@@ -311,6 +311,12 @@ UE_TRACE_EVENT_BEGIN(GpuProfiler, EventEndWork)
 	UE_TRACE_EVENT_FIELD(uint64, GPUTimestampBOP)
 UE_TRACE_EVENT_END()
 
+UE_TRACE_EVENT_BEGIN(GpuProfiler, EventWait)
+	UE_TRACE_EVENT_FIELD(uint32, QueueId)
+	UE_TRACE_EVENT_FIELD(uint64, StartTime)
+	UE_TRACE_EVENT_FIELD(uint64, EndTime)
+UE_TRACE_EVENT_END()
+
 UE_TRACE_EVENT_BEGIN(GpuProfiler, EventStats)
 	UE_TRACE_EVENT_FIELD(uint32, QueueId)
 	UE_TRACE_EVENT_FIELD(uint32, NumDraws)
@@ -401,6 +407,14 @@ void FGpuProfilerTrace::EndWork(uint32 QueueId, uint64 GPUTimestampBOP)
 		<< EventEndWork.GPUTimestampBOP(GPUTimestampBOP);
 }
 
+void FGpuProfilerTrace::TraceWait(uint32 QueueId, uint64 StartTime, uint64 EndTime)
+{
+	UE_TRACE_LOG(GpuProfiler, EventWait, GpuChannel)
+		<< EventWait.QueueId(QueueId)
+		<< EventWait.StartTime(StartTime)
+		<< EventWait.EndTime(EndTime);
+}
+
 void FGpuProfilerTrace::Stats(uint32 QueueId, uint32 NumDraws, uint32 NumPrimitives)
 {
 	UE_TRACE_LOG(GpuProfiler, EventStats, GpuChannel)
diff --git a/Engine/Source/Runtime/RHI/Public/GpuProfilerTrace.h b/Engine/Source/Runtime/RHI/Public/GpuProfilerTrace.h
--- a/Engine/Source/Runtime/RHI/Public/GpuProfilerTrace.h
+++ b/Engine/Source/Runtime/RHI/Public/GpuProfilerTrace.h
@@ -113,6 +113,11 @@ struct FGpuProfilerTrace
 	 */
 	RHI_API static void EndWork(uint32 QueueId, uint64 GPUTimestampBOP);
 
+	/**
+	 * Trace a wait event on a Gpu Queue. 
+	 */
+	RHI_API static void TraceWait(uint32 QueueId, uint64 StartTime, uint64 EndTime);
+
 	/**
 	 * Trace Gpu stats. 
 	 */
@@ -137,6 +142,7 @@ struct FGpuProfilerTrace
 	RHI_API static void EndBreadcrumb(uint32 QueueId, uint64 GPUTimestampBOP, const TCHAR* Name) {}
 	RHI_API static void BeginWork(uint32 QueueId, uint64 GPUTimestampTOP, uint64 CPUTimestamp) {}
 	RHI_API static void EndWork(uint32 QueueId, uint64 GPUTimestampBOP) {}
+	RHI_API static void TraceWait(uint32 QueueId, uint64 StartTime, uint64 EndTime) {}
 	RHI_API static void Stats(uint32 QueueId, uint32 NumDraws, uint32 NumPrimitives) {}
 	RHI_API static void SignalFence(uint32 QueueId, uint64 ResolvedTimestamp, uint64 Value) {}
 	RHI_API static void WaitFence(uint32 QueueId, uint64 ResolvedTimestamp, uint32 QueueToWaitForId, uint64 Value) {}
