[Insights] TraceServices: Removed deprecated API in FProviderLock.

diff --git a/Engine/Source/Developer/TraceServices/Private/Common/ProviderLock.cpp b/Engine/Source/Developer/TraceServices/Private/Common/ProviderLock.cpp
--- a/Engine/Source/Developer/TraceServices/Private/Common/ProviderLock.cpp
+++ b/Engine/Source/Developer/TraceServices/Private/Common/ProviderLock.cpp
@@ -61,60 +61,4 @@ void FProviderLock::EndWrite(FProviderLock::FThreadLocalState& State)
 	}
 }
 
-void FProviderLock::ReadAccessCheck(const FProviderLock* ProviderLock, const int32& ReadProviderLockCount, const int32& WriteProviderLockCount) const
-{
-	checkf(ProviderLock == this && (ReadProviderLockCount > 0 || WriteProviderLockCount > 0),
-		TEXT("Trying to READ from provider outside of a READ scope"));
-}
-
-void FProviderLock::WriteAccessCheck(const int32& WriteProviderLockCount) const
-{
-	checkf(WriteProviderLockCount > 0,
-		TEXT("Trying to WRITE to provider outside of an EDIT/WRITE scope"));
-}
-
-void FProviderLock::BeginRead(FProviderLock*& ProviderLock, int32& ReadProviderLockCount, const int32& WriteProviderLockCount)
-{
-	check(!ProviderLock || ProviderLock == this);
-	checkf(WriteProviderLockCount == 0, TEXT("Trying to lock provider for READ while holding EDIT/WRITE access"));
-	if (ReadProviderLockCount++ == 0)
-	{
-		ProviderLock = this;
-		RWLock.ReadLock();
-	}
-}
-
-void FProviderLock::EndRead(FProviderLock*& ProviderLock, int32& ReadProviderLockCount)
-{
-	check(ProviderLock == this);
-	check(ReadProviderLockCount > 0);
-	if (--ReadProviderLockCount == 0)
-	{
-		RWLock.ReadUnlock();
-		ProviderLock = nullptr;
-	}
-}
-
-void FProviderLock::BeginWrite(FProviderLock*& ProviderLock, const int32& ReadProviderLockCount, int32& WriteProviderLockCount)
-{
-	check(!ProviderLock || ProviderLock == this);
-	checkf(ReadProviderLockCount == 0, TEXT("Trying to lock provider for EDIT/WRITE while holding READ access"));
-	if (WriteProviderLockCount++ == 0)
-	{
-		ProviderLock = this;
-		RWLock.WriteLock();
-	}
-}
-
-void FProviderLock::EndWrite(FProviderLock*& ProviderLock, int32& WriteProviderLockCount)
-{
-	check(ProviderLock == this);
-	check(WriteProviderLockCount > 0);
-	if (--WriteProviderLockCount == 0)
-	{
-		RWLock.WriteUnlock();
-		ProviderLock = nullptr;
-	}
-}
-
 } // namespace TraceServices
diff --git a/Engine/Source/Developer/TraceServices/Public/Common/ProviderLock.h b/Engine/Source/Developer/TraceServices/Public/Common/ProviderLock.h
--- a/Engine/Source/Developer/TraceServices/Public/Common/ProviderLock.h
+++ b/Engine/Source/Developer/TraceServices/Public/Common/ProviderLock.h
@@ -68,21 +68,6 @@ public:
 	TRACESERVICES_API void BeginWrite(FThreadLocalState& State);
 	TRACESERVICES_API void EndWrite(FThreadLocalState& State);
 
-	UE_DEPRECATED(5.3, "Please use the FThreadLocalState overload")
-	void ReadAccessCheck(const FProviderLock* CurrentProviderLock, const int32& CurrentReadProviderLockCount, const int32& CurrentWriteProviderLockCount) const;
-	UE_DEPRECATED(5.3, "Please use the FThreadLocalState overload")
-	void WriteAccessCheck(const int32& CurrentWriteProviderLockCount) const;
-
-	UE_DEPRECATED(5.3, "Please use the FThreadLocalState overload")
-	void BeginRead(FProviderLock*& CurrentProviderLock, int32& CurrentReadProviderLockCount, const int32& WriteProviderLockCount);
-	UE_DEPRECATED(5.3, "Please use the FThreadLocalState overload")
-	void EndRead(FProviderLock*& CurrentProviderLock, int32& CurrentReadProviderLockCount);
-
-	UE_DEPRECATED(5.3, "Please use the FThreadLocalState overload")
-	void BeginWrite(FProviderLock*& CurrentProviderLock, const int32& CurrentReadProviderLockCount, int32& WriteProviderLockCount);
-	UE_DEPRECATED(5.3, "Please use the FThreadLocalState overload")
-	void EndWrite(FProviderLock*& CurrentProviderLock, int32& WriteProviderLockCount);
-
 private:
 	FRWLock RWLock;
 };
