[Insights] Log View: Added "Unknown log message spec!" error messages for cases where a log message event is traced without a valid spec event (LogPoint).

diff --git a/Engine/Source/Developer/TraceServices/Private/Model/Log.cpp b/Engine/Source/Developer/TraceServices/Private/Model/Log.cpp
--- a/Engine/Source/Developer/TraceServices/Private/Model/Log.cpp
+++ b/Engine/Source/Developer/TraceServices/Private/Model/Log.cpp
@@ -135,44 +135,80 @@ FLogMessageInternal& FLogProvider::AppendMessageInternal(double Time)
 	return InternalMessage;
 }
 
+void FLogProvider::AppendUnknownMessageInternal(uint64 LogPoint, double Time, const FStringView Message)
+{
+	FLogMessageSpec& Spec = GetMessageSpec(~0ull);
+	FLogCategoryInfo& Category = GetCategory(~0ull);
+	Spec.Category = &Category;
+	Spec.Line = 0;
+	Spec.Verbosity = ELogVerbosity::Error;
+	Spec.FormatString = TEXT("%s");
+
+	FLogMessageInternal& InternalMessage = AppendMessageInternal(Time);
+	InternalMessage.Spec = &Spec;
+	InternalMessage.Message = Session.StoreString(*FString::Printf(TEXT("Unknown log message spec (LogPoint=0x%llX)! Message: \"%.*s\""), LogPoint, Message.Len(), Message.GetData()));
+}
+
 void FLogProvider::AppendMessage(uint64 LogPoint, double Time, const uint8* FormatArgs)
 {
 	Session.WriteAccessCheck();
 	FLogMessageSpec** FindSpec = SpecMap.Find(LogPoint);
-	if (FindSpec && (*FindSpec)->Verbosity != ELogVerbosity::SetColor)
+	if (FindSpec)
 	{
-		FLogMessageInternal& InternalMessage = AppendMessageInternal(Time);
-		InternalMessage.Spec = *FindSpec;
-		FFormatArgsHelper::Format(FormatBuffer, FormatBufferSize - 1, TempBuffer, FormatBufferSize - 1, InternalMessage.Spec->FormatString, FormatArgs);
-		InternalMessage.Message = Session.StoreString(FormatBuffer);
-		Session.UpdateDurationSeconds(Time);
+		if ((*FindSpec)->Verbosity != ELogVerbosity::SetColor)
+		{
+			FLogMessageInternal& InternalMessage = AppendMessageInternal(Time);
+			InternalMessage.Spec = *FindSpec;
+			FFormatArgsHelper::Format(FormatBuffer, FormatBufferSize - 1, TempBuffer, FormatBufferSize - 1, InternalMessage.Spec->FormatString, FormatArgs);
+			InternalMessage.Message = Session.StoreString(FormatBuffer);
+		}
 	}
+	else
+	{
+		FFormatArgsHelper::Format(FormatBuffer, FormatBufferSize - 1, TempBuffer, FormatBufferSize - 1, TEXT("%s"), FormatArgs);
+		AppendUnknownMessageInternal(LogPoint, Time, FStringView(FormatBuffer));
+	}
+	Session.UpdateDurationSeconds(Time);
 }
 
 void FLogProvider::AppendMessage(uint64 LogPoint, double Time, const TCHAR* Text)
 {
 	Session.WriteAccessCheck();
 	FLogMessageSpec** FindSpec = SpecMap.Find(LogPoint);
-	if (FindSpec && (*FindSpec)->Verbosity != ELogVerbosity::SetColor)
+	if (FindSpec)
 	{
-		FLogMessageInternal& InternalMessage = AppendMessageInternal(Time);
-		InternalMessage.Spec = *FindSpec;
-		InternalMessage.Message = Text;
-		Session.UpdateDurationSeconds(Time);
+		if ((*FindSpec)->Verbosity != ELogVerbosity::SetColor)
+		{
+			FLogMessageInternal& InternalMessage = AppendMessageInternal(Time);
+			InternalMessage.Spec = *FindSpec;
+			InternalMessage.Message = Text;
+		}
+		else
+		{
+			AppendUnknownMessageInternal(LogPoint, Time, FStringView(Text));
+		}
 	}
+	Session.UpdateDurationSeconds(Time);
 }
 
 void FLogProvider::AppendMessage(uint64 LogPoint, double Time, const FStringView Message)
 {
 	Session.WriteAccessCheck();
 	FLogMessageSpec** FindSpec = SpecMap.Find(LogPoint);
-	if (FindSpec && (*FindSpec)->Verbosity != ELogVerbosity::SetColor)
+	if (FindSpec)
+	{
+		if ((*FindSpec)->Verbosity != ELogVerbosity::SetColor)
+		{
+			FLogMessageInternal& InternalMessage = AppendMessageInternal(Time);
+			InternalMessage.Spec = *FindSpec;
+			InternalMessage.Message = Session.StoreString(Message);
+		}
+	}
+	else
 	{
-		FLogMessageInternal& InternalMessage = AppendMessageInternal(Time);
-		InternalMessage.Spec = *FindSpec;
-		InternalMessage.Message = Session.StoreString(Message);
-		Session.UpdateDurationSeconds(Time);
+		AppendUnknownMessageInternal(LogPoint, Time, Message);
 	}
+	Session.UpdateDurationSeconds(Time);
 }
 
 bool FLogProvider::ReadMessage(uint64 Index, TFunctionRef<void(const FLogMessageInfo&)> Callback) const
diff --git a/Engine/Source/Developer/TraceServices/Private/Model/LogPrivate.h b/Engine/Source/Developer/TraceServices/Private/Model/LogPrivate.h
--- a/Engine/Source/Developer/TraceServices/Private/Model/LogPrivate.h
+++ b/Engine/Source/Developer/TraceServices/Private/Model/LogPrivate.h
@@ -77,6 +77,7 @@ public:
 
 private:
 	FLogMessageInternal& AppendMessageInternal(double Time);
+	void AppendUnknownMessageInternal(uint64 LogPoint, double Time, const FStringView Message);
 	void ConstructMessage(const FLogMessageInternal& InternalMessage, uint64 Index, TFunctionRef<void(const FLogMessageInfo&)> Callback) const;
 
 	enum
