Strip off legacy thread id bias that was used to transition CPU event protocol in the very early days.

Once upon a time there was a CpuProfiler::EventBatch::ThreadId member. This was later removed when thread id became implicit as part of the trace transport protocol. EventBatch has also since been superceeded. Thus EventBatch, and it's prior ThreadId field, could be repurposed. Fake threads can be registered with trace to name event tracks neatly. But to do that the "anti-collision bias" added by FTraceAnalyzerUtils::GetThreadIdField() needs to be removed, and can be as it is fair to assume there next to zero traces out there with the very early ThreadId member.

#rnx

diff --git a/Engine/Source/Developer/TraceServices/Private/Analyzers/CpuProfilerTraceAnalysis.cpp b/Engine/Source/Developer/TraceServices/Private/Analyzers/CpuProfilerTraceAnalysis.cpp
--- a/Engine/Source/Developer/TraceServices/Private/Analyzers/CpuProfilerTraceAnalysis.cpp
+++ b/Engine/Source/Developer/TraceServices/Private/Analyzers/CpuProfilerTraceAnalysis.cpp
@@ -837,6 +837,15 @@ uint32 FCpuProfilerAnalyzer::GetTimerId(uint32 SpecId)
 
 FCpuProfilerAnalyzer::FThreadState& FCpuProfilerAnalyzer::GetThreadState(uint32 ThreadId)
 {
+	// Strip off legacy bias that was used to transition the CPU event protocol in the very early
+	// days. There used to be a ThreadId field in CpuProfiler::EventBatch which was supersceeded by
+	// thread id in transport packets. The bias was added to avoid thread ID collision. The below
+	// allows EventBatch to be brought out of retirement for adhob async/off-CPU introspection.
+	if (uint32 EventInlined = ThreadId & 0x7000'0000; EventInlined != 0)
+	{
+		ThreadId &= ~EventInlined;
+	}
+
 	FThreadState* ThreadState = ThreadStatesMap.FindRef(ThreadId);
 	if (!ThreadState)
 	{
