[Insights] Timing Regions / MiscTrace: Updated MiscTraceAnalysis and RegionProvider to support Regions with Categories.

- Refactored any functions that access regions or lanes directly into new `IRegionTimeline` interface (implemented in `FRegionTimeline`) and deprecated accessors on Region Provider
- Introduced new functions to access the old timeline and new timelines filtered by Category (`GetDefaultTimeline()`, `EnumerateTimelinesByCategory()`)

diff --git a/Engine/Source/Developer/TraceAnalysis/Public/Trace/Analyzer.h b/Engine/Source/Developer/TraceAnalysis/Public/Trace/Analyzer.h
--- a/Engine/Source/Developer/TraceAnalysis/Public/Trace/Analyzer.h
+++ b/Engine/Source/Developer/TraceAnalysis/Public/Trace/Analyzer.h
@@ -483,7 +483,7 @@ TEventRef<DefinitionType> IAnalyzer::FEventData::GetReferenceValue(const ANSICHA
 	const void* Value = GetReferenceValueImpl(FieldName, SizeAndType, TypeUid);
 	if (Value)
 	{
-		return MakeEventRef<DefinitionType>(CoerceValue<DefinitionType>(Value, SizeAndType), TypeUid);
+		return MakeEventRef<DefinitionType>(CoerceValue<DefinitionType>(Value, static_cast<int8>(SizeAndType)), TypeUid);
 	}
 	return MakeEventRef<DefinitionType>(0,0);
 }
diff --git a/Engine/Source/Developer/TraceInsights/Private/Insights/TimingProfiler/Tracks/RegionsTimingTrack.cpp b/Engine/Source/Developer/TraceInsights/Private/Insights/TimingProfiler/Tracks/RegionsTimingTrack.cpp
--- a/Engine/Source/Developer/TraceInsights/Private/Insights/TimingProfiler/Tracks/RegionsTimingTrack.cpp
+++ b/Engine/Source/Developer/TraceInsights/Private/Insights/TimingProfiler/Tracks/RegionsTimingTrack.cpp
@@ -89,7 +89,7 @@ void FTimingRegionsTrack::BuildDrawState(ITimingEventsTrackDrawStateBuilder& Bui
 
 	// We are counting only non-empty lanes, so we can collapse empty ones in the visualization.
 	int32 CurDepth = 0;
-	RegionProvider.EnumerateLanes([this, Viewport, &CurDepth, &Builder](const TraceServices::FRegionLane& Lane, const int32 Depth)
+	RegionProvider.GetDefaultTimeline().EnumerateLanes([this, Viewport, &CurDepth, &Builder](const TraceServices::FRegionLane& Lane, const int32 Depth)
 	{
 		bool RegionHadEvents = false;
 		Lane.EnumerateRegions(Viewport.GetStartTime(), Viewport.GetEndTime(), [&Builder, &RegionHadEvents, &CurDepth](const TraceServices::FTimeRegion& Region) -> bool
@@ -134,7 +134,7 @@ void FTimingRegionsTrack::BuildFilteredDrawState(ITimingEventsTrackDrawStateBuil
 			if (bFilterOnlyByEventType)
 			{
 				int32 CurDepth = 0;
-				RegionProvider.EnumerateLanes([this, Viewport, &CurDepth, &Builder, FilterEventType](const TraceServices::FRegionLane& Lane, const int32 Depth)
+				RegionProvider.GetDefaultTimeline().EnumerateLanes([this, Viewport, &CurDepth, &Builder, FilterEventType](const TraceServices::FRegionLane& Lane, const int32 Depth)
 					{
 						bool RegionHadEvents = false;
 						Lane.EnumerateRegions(Viewport.GetStartTime(), Viewport.GetEndTime(), [&Builder, &RegionHadEvents, &CurDepth, FilterEventType](const TraceServices::FTimeRegion& Region) -> bool
@@ -181,7 +181,7 @@ void FTimingRegionsTrack::BuildFilteredDrawState(ITimingEventsTrackDrawStateBuil
 			const FTimingTrackViewport& Viewport = Context.GetViewport();
 
 			int32 CurDepth = 0;
-			RegionProvider.EnumerateLanes([this, Viewport, &CurDepth, &Builder, &FilterContext](const TraceServices::FRegionLane& Lane, const int32 Depth)
+			RegionProvider.GetDefaultTimeline().EnumerateLanes([this, Viewport, &CurDepth, &Builder, &FilterContext](const TraceServices::FRegionLane& Lane, const int32 Depth)
 				{
 					bool RegionHadEvents = false;
 					Lane.EnumerateRegions(Viewport.GetStartTime(), Viewport.GetEndTime(), [&Builder, &RegionHadEvents, &CurDepth, &FilterContext, this](const TraceServices::FTimeRegion& Region) -> bool
@@ -253,7 +253,7 @@ bool FTimingRegionsTrack::FindRegionEvent(const FTimingEventSearchParameters& In
 		const TraceServices::IRegionProvider& RegionProvider = TraceServices::ReadRegionProvider(*Session);
 		TraceServices::FProviderReadScopeLock RegionProviderScopedLock(RegionProvider);
 
-		RegionProvider.EnumerateRegions(InContext.GetParameters().StartTime, InContext.GetParameters().EndTime, [&InContext](const TraceServices::FTimeRegion& Region)
+		RegionProvider.GetDefaultTimeline().EnumerateRegions(InContext.GetParameters().StartTime, InContext.GetParameters().EndTime, [&InContext](const TraceServices::FTimeRegion& Region)
 		{
 			InContext.Check(Region.BeginTime, Region.EndTime, Region.Depth, Region);
 
diff --git a/Engine/Source/Developer/TraceInsights/Private/Insights/TimingProfiler/ViewModels/TimingExporter.cpp b/Engine/Source/Developer/TraceInsights/Private/Insights/TimingProfiler/ViewModels/TimingExporter.cpp
--- a/Engine/Source/Developer/TraceInsights/Private/Insights/TimingProfiler/ViewModels/TimingExporter.cpp
+++ b/Engine/Source/Developer/TraceInsights/Private/Insights/TimingProfiler/ViewModels/TimingExporter.cpp
@@ -644,7 +644,7 @@ void FTimingExporter::GetRegions(const FString& InRegionNamePattern, TMap<FStrin
 
 		UE_LOG(TraceInsights, Log, TEXT("Looking for regions: '%s'"), *InRegionNamePattern);
 
-		RegionProvider.EnumerateRegions(0.0, std::numeric_limits<double>::max(),
+		RegionProvider.GetDefaultTimeline().EnumerateRegions(0.0, std::numeric_limits<double>::max(),
 			[&RegionCount, &RegionNameSpec, &OutRegionGroups](const TraceServices::FTimeRegion& InRegion) -> bool
 			{
 				if (!RegionNameSpec.Match(InRegion.Text))
diff --git a/Engine/Source/Developer/TraceServices/Private/Analyzers/MiscTraceAnalysis.cpp b/Engine/Source/Developer/TraceServices/Private/Analyzers/MiscTraceAnalysis.cpp
--- a/Engine/Source/Developer/TraceServices/Private/Analyzers/MiscTraceAnalysis.cpp
+++ b/Engine/Source/Developer/TraceServices/Private/Analyzers/MiscTraceAnalysis.cpp
@@ -8,12 +8,14 @@
 #include "Logging/MessageLog.h"
 #include "Model/BookmarksPrivate.h"
 #include "Model/Channel.h"
+#include "Model/DefinitionProvider.h"
 #include "Model/FramesPrivate.h"
 #include "Model/LogPrivate.h"
 #include "Model/RegionsPrivate.h"
 #include "Model/ScreenshotProviderPrivate.h"
 #include "Model/ThreadsPrivate.h"
 #include "TraceServices/Model/AnalysisSession.h"
+#include "TraceServices/Model/Strings.h"
 
 namespace TraceServices
 {
@@ -261,8 +263,16 @@ bool FMiscTraceAnalyzer::OnEvent(uint16 RouteId, EStyle Style, const FOnEventCon
 		uint64 Cycle = EventData.GetValue<uint64>("Cycle");
 		FString Name;
 		check(EventData.GetString("RegionName", Name));
-		FProviderEditScopeLock RegionProviderScopedLock(RegionProvider);
-		RegionProvider.AppendRegionBegin(*Name, Context.EventTime.AsSeconds(Cycle));
+			
+		// Category has been added after the initial implementation, so it's treated as optional
+		FString Category;
+		EventData.GetString("Category", Category);
+		double Time = Context.EventTime.AsSeconds(Cycle);
+		{
+			FProviderEditScopeLock RegionProviderScopedLock(RegionProvider);
+			RegionProvider.AppendRegionBegin(*Name, Time, Category.IsEmpty() ? nullptr : *Category);
+		}
+			Session.UpdateDurationSeconds(Time);
 		break;
 	}
 
@@ -271,8 +281,17 @@ bool FMiscTraceAnalyzer::OnEvent(uint16 RouteId, EStyle Style, const FOnEventCon
 		uint64 CycleAndId = EventData.GetValue<uint64>("CycleAndId");
 		FString Name;
 		check(EventData.GetString("RegionName", Name));
-		FProviderEditScopeLock RegionProviderScopedLock(RegionProvider);
-		RegionProvider.AppendRegionBeginWithId(*Name, CycleAndId, Context.EventTime.AsSeconds(CycleAndId));
+			
+		// Category has been added after the initial implementation, so it's treated as optional
+		FString Category;
+		EventData.GetString("Category", Category);
+		double Time = Context.EventTime.AsSeconds(CycleAndId);
+		{
+			FProviderEditScopeLock RegionProviderScopedLock(RegionProvider);
+			RegionProvider.AppendRegionBeginWithId(*Name, CycleAndId, Time, Category.IsEmpty() ? nullptr : *Category);
+		}
+		
+		Session.UpdateDurationSeconds(Time);
 		break;
 	}
 
@@ -281,8 +300,13 @@ bool FMiscTraceAnalyzer::OnEvent(uint16 RouteId, EStyle Style, const FOnEventCon
 		uint64 Cycle = EventData.GetValue<uint64>("Cycle");
 		FString Name;
 		EventData.GetString("RegionName", Name);
-		FProviderEditScopeLock RegionProviderScopedLock(RegionProvider);
-		RegionProvider.AppendRegionEnd(*Name, Context.EventTime.AsSeconds(Cycle));
+		double Time = Context.EventTime.AsSeconds(Cycle);
+		{
+			FProviderEditScopeLock RegionProviderScopedLock(RegionProvider);
+			RegionProvider.AppendRegionEnd(*Name, Time);
+		}
+
+		Session.UpdateDurationSeconds(Time);
 		break;
 	}
 
@@ -290,8 +314,13 @@ bool FMiscTraceAnalyzer::OnEvent(uint16 RouteId, EStyle Style, const FOnEventCon
 	{
 		uint64 Cycle = EventData.GetValue<uint64>("Cycle");
 		uint64 RegionId = EventData.GetValue<uint64>("RegionId", 0);
-		FProviderEditScopeLock RegionProviderScopedLock(RegionProvider);
-		RegionProvider.AppendRegionEndWithId(RegionId, Context.EventTime.AsSeconds(Cycle));
+		double Time = Context.EventTime.AsSeconds(Cycle);
+		{
+			FProviderEditScopeLock RegionProviderScopedLock(RegionProvider);
+			RegionProvider.AppendRegionEndWithId(RegionId, Time);
+		}
+
+		Session.UpdateDurationSeconds(Time);
 		break;
 	}
 	}
diff --git a/Engine/Source/Developer/TraceServices/Private/Model/Regions.cpp b/Engine/Source/Developer/TraceServices/Private/Model/Regions.cpp
--- a/Engine/Source/Developer/TraceServices/Private/Model/Regions.cpp
+++ b/Engine/Source/Developer/TraceServices/Private/Model/Regions.cpp
@@ -16,43 +16,177 @@ namespace TraceServices
 
 thread_local FProviderLock::FThreadLocalState GRegionsProviderLockState;
 
-FRegionProvider::FRegionProvider(IAnalysisSession& InSession)
-	: Session(InSession)
+// ####################### FRegionTimeline ##########################
+
+const TCHAR* FRegionTimeline::GetCategory() const
+{
+	Provider.ReadAccessCheck();
+	return Category;
+}
+	
+int32 FRegionTimeline::GetLaneCount() const
 {
+	Provider.ReadAccessCheck();
+	return Lanes.Num();
 }
 
-uint64 FRegionProvider::GetRegionCount() const
+const FRegionLane* FRegionTimeline::GetLane(int32 Index) const
 {
-	ReadAccessCheck();
+	Provider.ReadAccessCheck();
+	
+	if (Index >= 0 && Index < Lanes.Num())
+	{
+		return &(Lanes[Index]);
+	}
+	return nullptr;
+}
+
+bool FRegionTimeline::EnumerateRegions(double IntervalStart, double IntervalEnd,
+	TFunctionRef<bool(const FTimeRegion&)> Callback) const
+{
+	Provider.ReadAccessCheck();
+	
+	if (IntervalStart > IntervalEnd)
+	{
+		return false;
+	}
 
-	uint64 RegionCount = 0;
 	for (const FRegionLane& Lane : Lanes)
 	{
-		RegionCount += Lane.Num();
+		if (!Lane.EnumerateRegions(IntervalStart, IntervalEnd, Callback))
+		{
+			return false;
+		}
 	}
-	return RegionCount;
+
+	return true;
 }
 
-const FRegionLane* FRegionProvider::GetLane(int32 index) const
+void FRegionTimeline::EnumerateLanes(TFunctionRef<void(const FRegionLane&, const int32)> Callback) const
+{
+	Provider.ReadAccessCheck();
+	
+	for (int32 LaneIndex = 0; LaneIndex < Lanes.Num(); ++LaneIndex)
+	{
+		Callback(Lanes[LaneIndex], LaneIndex);
+	}
+}
+
+FTimeRegion* FRegionTimeline::InsertNewRegion(double BeginTime, const TCHAR* Name, uint64 Id, const TCHAR* RegionCategory)
+{
+	Provider.EditAccessCheck();
+	
+	FTimeRegion Region;
+	Region.BeginTime = BeginTime;
+	Region.Text = Name;
+	Region.Id = Id;
+	Region.Category = RegionCategory;
+	Region.Depth = CalculateRegionDepth(Region);
+
+	if (Region.Depth == Lanes.Num())
+	{
+		Lanes.Emplace(Session.GetLinearAllocator());
+	}
+	FRegionLane& NewLane = Lanes[Region.Depth];	
+	NewLane.Regions.EmplaceBack(Region);
+	FTimeRegion* NewOpenRegion = &(Lanes[Region.Depth].Regions.Last());
+
+	return NewOpenRegion;
+}
+
+int32 FRegionTimeline::CalculateRegionDepth(const FTimeRegion& Region) const
+{
+	constexpr int32 DepthLimit = 100;
+	int32 NewDepth = 0;
+
+	// Find first free lane/depth
+	while (NewDepth < DepthLimit)
+	{
+		if (!Lanes.IsValidIndex(NewDepth))
+		{
+			break;
+		}
+		
+		const FTimeRegion& LastRegion = Lanes[NewDepth].Regions.Last();
+		if (LastRegion.EndTime <= Region.BeginTime)
+		{
+			break;
+		}
+		NewDepth++;
+	}
+
+	ensureMsgf(NewDepth < DepthLimit, TEXT("Regions are nested too deep."));
+
+	return NewDepth;
+}	
+	
+// ##################### RegionProvider #####################	
+
+FRegionProvider::FRegionProvider(IAnalysisSession& InSession)
+	: Session(InSession), AllRegions(*this, InSession, nullptr)
+{
+}
+
+void FRegionProvider::EnumerateTimelinesByCategory(
+	TFunctionRef<void(const IRegionTimeline&, const TCHAR*)> Callback) const
+{
+	ReadAccessCheck();
+	
+	for (auto& KV: RegionsPerCategory)
+	{
+		Callback(KV.Value, KV.Key);
+	}
+}
+
+const IRegionTimeline& FRegionProvider::GetDefaultTimeline() const
+{
+	ReadAccessCheck();
+	
+	return AllRegions;
+}
+
+const IRegionTimeline* FRegionProvider::GetTimelineForCategory(const TCHAR* Category) const
 {
 	ReadAccessCheck();
 
-	if (index < Lanes.Num())
+	if (!Category)
 	{
-		return &(Lanes[index]);
+		return &AllRegions;
+	}
+	
+	if (const FRegionTimeline* Timeline = RegionsPerCategory.Find(Category))
+	{
+		return Timeline;
 	}
 	return nullptr;
 }
+	
+uint64 FRegionProvider::GetRegionCount() const
+{
+	ReadAccessCheck();
+
+	uint64 RegionCount = 0;
+	AllRegions.EnumerateLanes([&RegionCount](const FRegionLane& Lane, const int32 Index) {
+		RegionCount += Lane.Num();
+	});
+	return RegionCount;
+}
+
+const FRegionLane* FRegionProvider::GetLane(int32 index) const
+{
+	ReadAccessCheck();
+	return AllRegions.GetLane(index);
+}
 
-void FRegionProvider::AppendRegionBegin(const TCHAR* Name, double Time)
+void FRegionProvider::AppendRegionBegin(const TCHAR* Name, double Time, const TCHAR* Category)
 {
 	EditAccessCheck();
 
 	check(Name)
 
-	FTimeRegion** OpenRegion = OpenRegionsByName.Find(Name);
+	TArray<FTimeRegion*>* OpenRegions = OpenRegionsByName.Find(Name);
 
-	if (OpenRegion)
+	if (OpenRegions)
 	{
 		++NumWarnings;
 		if (NumWarnings <= MaxWarningMessages)
@@ -64,18 +198,20 @@ void FRegionProvider::AppendRegionBegin(const TCHAR* Name, double Time)
 		AppendRegionEnd(Name, Time);
 	}
 
-	FTimeRegion* NewRegion = InsertNewRegion(Time, Name, 0);
-	OpenRegionsByName.Add(NewRegion->Text, NewRegion);
+	TArray<FTimeRegion*> NewRegions = InsertNewRegion(Time, Name, 0, Category);
+	OpenRegionsByName.Add(NewRegions[0]->Text, NewRegions);
+
+	UpdateCounter++;
 }
 
-void FRegionProvider::AppendRegionBeginWithId(const TCHAR* Name, uint64 Id, double Time)
+void FRegionProvider::AppendRegionBeginWithId(const TCHAR* Name, uint64 Id, double Time, const TCHAR* Category)
 {
 	EditAccessCheck();
 
 	check(Name && Id)
-	FTimeRegion** OpenRegion = OpenRegionsById.Find(Id);
+	TArray<FTimeRegion*>* OpenRegions = OpenRegionsById.Find(Id);
 
-	if (OpenRegion)
+	if (OpenRegions)
 	{
 		++NumWarnings;
 		if (NumWarnings <= MaxWarningMessages)
@@ -87,8 +223,10 @@ void FRegionProvider::AppendRegionBeginWithId(const TCHAR* Name, uint64 Id, doub
 		AppendRegionEndWithId(Id, Time);
 	}
 
-	FTimeRegion* NewRegion = InsertNewRegion(Time, Name, Id);
-	OpenRegionsById.Add(Id, NewRegion);
+	TArray<FTimeRegion*> NewRegions = InsertNewRegion(Time, Name, Id, Category);
+	OpenRegionsById.Add(Id, NewRegions);
+	
+	UpdateCounter++;
 }
 
 void FRegionProvider::AppendRegionEnd(const TCHAR* Name, double Time)
@@ -96,9 +234,9 @@ void FRegionProvider::AppendRegionEnd(const TCHAR* Name, double Time)
 	EditAccessCheck();
 
 	check(Name)
-	FTimeRegion** OpenRegionPos = OpenRegionsByName.Find(Name);
+	TArray<FTimeRegion*>* OpenRegionsPos = OpenRegionsByName.Find(Name);
 
-	if (!OpenRegionPos)
+	if (!OpenRegionsPos)
 	{
 		++NumWarnings;
 		if (NumWarnings <= MaxWarningMessages)
@@ -107,33 +245,37 @@ void FRegionProvider::AppendRegionEnd(const TCHAR* Name, double Time)
 		}
 
 		AppendRegionBegin(Name, Time);
-		OpenRegionPos = OpenRegionsByName.Find(Name);
-		check(OpenRegionPos);
+		OpenRegionsPos = OpenRegionsByName.Find(Name);
+		check(OpenRegionsPos);
 	}
 
-	FTimeRegion* OpenRegion = *OpenRegionPos;
-	check(OpenRegion);
-
-	OpenRegion->EndTime = Time;
+	for (FTimeRegion* OpenRegion: *OpenRegionsPos )
+	{
+		check(OpenRegion);
+		OpenRegion->EndTime = Time;
+	}	
 
-	OpenRegionsByName.Remove(OpenRegion->Text);
+	OpenRegionsByName.Remove(Name);
+	
 	UpdateCounter++;
-
-	// Update session time
-	{
-		FAnalysisSessionEditScope _(Session);
-		Session.UpdateDurationSeconds(Time);
-	}
 }
 
 void FRegionProvider::AppendRegionEndWithId(uint64 Id, double Time)
 {
 	EditAccessCheck();
 
-	check(Id)
-	FTimeRegion** OpenRegionPos = OpenRegionsById.Find(Id);
+	if (!Id)
+	{
+		++NumWarnings;
+		if (NumWarnings <= MaxWarningMessages)
+		{
+			UE_LOG(LogTraceServices, Warning, TEXT("[Regions] A region end event with id 0 was encountered, ignoring (EndTime=%f."), Time)
+		}
+		return;
+	}
+	TArray<FTimeRegion*>* OpenRegionsPos = OpenRegionsById.Find(Id);
 
-	if (!OpenRegionPos)
+	if (!OpenRegionsPos)
 	{
 		++NumWarnings;
 		if (NumWarnings <= MaxWarningMessages)
@@ -145,23 +287,19 @@ void FRegionProvider::AppendRegionEndWithId(uint64 Id, double Time)
 		// Generates a display name if we're missing a begin and are closing by ID
 		FString GeneratedName = FString::Printf(TEXT("Unknown Region (missing begin, Id=%llu)"), Id);
 		AppendRegionBeginWithId(*GeneratedName, Id, Time);
-		OpenRegionPos = OpenRegionsById.Find(Id);
-		check(OpenRegionPos);
+		OpenRegionsPos = OpenRegionsById.Find(Id);
+		check(OpenRegionsPos);
 	}
 
-	FTimeRegion* OpenRegion = *OpenRegionPos;
-	check(OpenRegion);
-
-	OpenRegion->EndTime = Time;
-
-	OpenRegionsById.Remove(OpenRegion->Id);
-	UpdateCounter++;
-
-	// Update session time
+	for (FTimeRegion* OpenRegion: *OpenRegionsPos )
 	{
-		FAnalysisSessionEditScope _(Session);
-		Session.UpdateDurationSeconds(Time);
+		check(OpenRegion);
+		OpenRegion->EndTime = Time;
 	}
+
+	OpenRegionsById.Remove(Id);
+	
+	UpdateCounter++;
 }
 
 void FRegionProvider::OnAnalysisSessionEnded()
@@ -170,8 +308,9 @@ void FRegionProvider::OnAnalysisSessionEnded()
 
 	auto printOpenRegionMessage = [this](const auto& KV)
 	{
-		const FTimeRegion* Region = KV.Value;
-
+		const TArray<FTimeRegion*> Regions = KV.Value;
+		const auto Region = Regions[0];
+		
 		++NumWarnings;
 		if (NumWarnings <= MaxWarningMessages)
 		{
@@ -191,91 +330,46 @@ void FRegionProvider::OnAnalysisSessionEnded()
 	}
 
 	uint64 TotalRegionCount = GetRegionCount();
-	UE_LOG(LogTraceServices, Log, TEXT("[Regions] Analysis completed (%llu regions, %d lanes)."), TotalRegionCount, Lanes.Num());
+	UE_LOG(LogTraceServices, Log, TEXT("[Regions] Analysis completed (%llu regions, %d lanes)."), TotalRegionCount, AllRegions.Lanes.Num());
 }
 
-int32 FRegionProvider::CalculateRegionDepth(const FTimeRegion& Region) const
-{
-	constexpr int32 DepthLimit = 100;
-
-	int32 NewDepth = 0;
 
-	// Find first free lane/depth
-	while (NewDepth < DepthLimit)
-	{
-		if (!Lanes.IsValidIndex(NewDepth))
-		{
-			break;
-		}
-
-		const FTimeRegion& LastRegion = Lanes[NewDepth].Regions.Last();
-		if (LastRegion.EndTime <= Region.BeginTime)
-		{
-			break;
-		}
-		NewDepth++;
-	}
-
-	ensureMsgf(NewDepth < DepthLimit, TEXT("Regions are nested too deep."));
-
-	return NewDepth;
-}
-
-FTimeRegion* FRegionProvider::InsertNewRegion(double BeginTime, const TCHAR* Name, uint64 Id)
+TArray<FTimeRegion*> FRegionProvider::InsertNewRegion(double BeginTime, const TCHAR* Name, uint64 Id,
+                                                      const TCHAR* Category)
 {
-	FTimeRegion Region;
-	Region.BeginTime = BeginTime;
-	Region.Text = Session.StoreString(Name);
-	Region.Id = Id;
-	Region.Depth = CalculateRegionDepth(Region);
-
-	if (Region.Depth == Lanes.Num())
+	const TCHAR* StoredName = Session.StoreString(Name);
+	const TCHAR* StoredCategory = Category ? Session.StoreString(Category) : nullptr;
+	
+	TArray<FTimeRegion*> Regions;
+	Regions.Add(AllRegions.InsertNewRegion(BeginTime, StoredName, Id, StoredCategory));
+
+	// skip per category timeline if no category is set
+	const TCHAR* Key = !StoredCategory ? UncategorizedName : StoredCategory;
+	
+	FRegionTimeline* Timeline = RegionsPerCategory.Find(Key);
+	if (!Timeline)
 	{
-		Lanes.Emplace(Session.GetLinearAllocator());
+		FRegionTimeline NewTimeline{*this, Session, Key};
+		Timeline = &RegionsPerCategory.Emplace(Key, NewTimeline);
 	}
-
-	Lanes[Region.Depth].Regions.EmplaceBack(Region);
-	FTimeRegion* NewOpenRegion = &(Lanes[Region.Depth].Regions.Last());
-	UpdateCounter++;
-
-	// Update session time
-	{
-		FAnalysisSessionEditScope _(Session);
-		Session.UpdateDurationSeconds(BeginTime);
-	}
-	return NewOpenRegion;
+	Regions.Add(Timeline->InsertNewRegion(BeginTime, StoredName, Id, StoredCategory));
+	return Regions;
 }
 
 void FRegionProvider::EnumerateLanes(TFunctionRef<void(const FRegionLane&, int32)> Callback) const
 {
 	ReadAccessCheck();
-
-	for (int32 LaneIndex = 0; LaneIndex < Lanes.Num(); ++LaneIndex)
-	{
-		Callback(Lanes[LaneIndex], LaneIndex);
-	}
+	AllRegions.EnumerateLanes(Callback);
 }
 
 bool FRegionProvider::EnumerateRegions(double IntervalStart, double IntervalEnd, TFunctionRef<bool(const FTimeRegion&)> Callback) const
 {
 	ReadAccessCheck();
-
-	if (IntervalStart > IntervalEnd)
-	{
-		return false;
-	}
-
-	for (const FRegionLane& Lane : Lanes)
-	{
-		if (!Lane.EnumerateRegions(IntervalStart, IntervalEnd, Callback))
-		{
-			return false;
-		}
-	}
-
-	return true;
+	return AllRegions.EnumerateRegions(IntervalStart, IntervalEnd, Callback);
 }
 
+// ###################### FRegionLane #########################
+
 bool FRegionLane::EnumerateRegions(double IntervalStart, double IntervalEnd, TFunctionRef<bool(const FTimeRegion&)> Callback) const
 {
 	const FInt32Interval OverlapRange = GetElementRangeOverlappingGivenRange<FTimeRegion>(Regions, IntervalStart, IntervalEnd,
@@ -303,7 +397,6 @@ FName GetRegionProviderName()
 	static const FName Name("RegionProvider");
 	return Name;
 }
-
 const IRegionProvider& ReadRegionProvider(const IAnalysisSession& Session)
 {
 	return *Session.ReadProvider<IRegionProvider>(GetRegionProviderName());
diff --git a/Engine/Source/Developer/TraceServices/Private/Model/RegionsPrivate.h b/Engine/Source/Developer/TraceServices/Private/Model/RegionsPrivate.h
--- a/Engine/Source/Developer/TraceServices/Private/Model/RegionsPrivate.h
+++ b/Engine/Source/Developer/TraceServices/Private/Model/RegionsPrivate.h
@@ -9,12 +9,73 @@
 
 namespace TraceServices
 {
-
+class FRegionProvider;
+	
 extern thread_local FProviderLock::FThreadLocalState GRegionsProviderLockState;
 
 class FAnalysisSessionLock;
 class FStringStore;
 
+/**
+ *	A View onto a set of Timing Regions, sorted into a stack of individual lanes without overlaps for display,
+ *	optionally filtered by Category
+ */
+class TRACESERVICES_API FRegionTimeline : public IRegionTimeline
+{
+public:
+	virtual ~FRegionTimeline() = default;
+	
+	FRegionTimeline(FRegionProvider& InProvider, IAnalysisSession& InSession, const TCHAR* Category) : Category(Category), Session(InSession), Provider(InProvider) {}
+
+	/**
+	 * @return the category of this timeline if filtered or nullptr
+	 */
+	const TCHAR* GetCategory() const override;
+
+	/**
+	 * @return the number of lanes
+	 */
+	int32 GetLaneCount() const override;
+
+	/**
+	 * Direct access to a certain lane at a given index/depth.
+	 * The pointer is valid only in the current read scope.
+	 * @return a pointer to the lane at the specified depth index or nullptr if Index is invalid
+	 */
+	const FRegionLane* GetLane(int32 Index) const override;
+
+	/**
+	 * Enumerates all regions that overlap a certain time interval. Will enumerate by depth but does not expose lanes.
+	 * 
+	 * @param Callback a callback called for each region. Return false to abort iteration.
+	 * @returns true if the enumeration finished, false if it was aborted by the callback returning false
+	 */
+	bool EnumerateRegions(double IntervalStart, double IntervalEnd, TFunctionRef<bool(const FTimeRegion& /*Region*/)> Callback) const override;
+
+	/**
+	 * Will call Callback(Lane, Depth) for each lane in order.
+	 */
+	void EnumerateLanes(TFunctionRef<void(const FRegionLane& /*Lane*/, const int32 /*Depth*/)> Callback) const override;
+
+private:
+	friend class FRegionProvider;
+	/*
+	 * Inserts a new Region into the stack, returns the created Region to allow the update of the region end time when it closes
+	 */
+	FTimeRegion*  InsertNewRegion(double BeginTime, const TCHAR* Name, uint64 Id, const TCHAR* Category);
+
+private:
+	/// Update the depth member of a region to allow overlapping regions to be displayed on separate lanes.
+	int32 CalculateRegionDepth(const FTimeRegion& Item) const;
+	
+	// Closed regions
+	TArray<FRegionLane> Lanes;
+	const TCHAR* Category = nullptr;
+	TraceServices::IAnalysisSession& Session;
+	// the parent provider. Required to engage the ReadAccessLock or WriteAccessLock when a Lane is modified
+	TraceServices::FRegionProvider& Provider;
+};
+	
 class FRegionProvider
 	: public IRegionProvider
 	, public IEditableRegionProvider
@@ -29,9 +90,14 @@ public:
 	virtual void BeginRead() const override       { Lock.BeginRead(GRegionsProviderLockState); }
 	virtual void EndRead() const override         { Lock.EndRead(GRegionsProviderLockState); }
 	virtual void ReadAccessCheck() const override { Lock.ReadAccessCheck(GRegionsProviderLockState); }
+	
+	virtual void EnumerateTimelinesByCategory(TFunctionRef<void(const IRegionTimeline& /*Timeline*/, const TCHAR* /*Category*/)> Callback) const override;
+	virtual const IRegionTimeline& GetDefaultTimeline() const override;
+	virtual const IRegionTimeline* GetTimelineForCategory(const TCHAR* Category) const override;
+	virtual const TCHAR* GetUncategorizedRegionCategoryName() const override {return UncategorizedName;};
 
 	virtual uint64 GetRegionCount() const override;
-	virtual int32 GetLaneCount()  const override { ReadAccessCheck(); return Lanes.Num(); }
+	virtual int32 GetLaneCount()  const override { ReadAccessCheck(); return AllRegions.GetLaneCount(); }
 
 	virtual const FRegionLane* GetLane(int32 Index) const override;
 
@@ -39,7 +105,7 @@ public:
 	virtual void EnumerateLanes(TFunctionRef<void(const FRegionLane&, const int32)> Callback) const override;
 
 	virtual uint64 GetUpdateCounter() const override { ReadAccessCheck(); return UpdateCounter; }
-
+	
 	//////////////////////////////////////////////////
 	// Edit operations
 
@@ -47,20 +113,19 @@ public:
 	virtual void EndEdit() const override         { Lock.EndWrite(GRegionsProviderLockState); }
 	virtual void EditAccessCheck() const override { Lock.WriteAccessCheck(GRegionsProviderLockState); }
 
-	virtual void AppendRegionBegin(const TCHAR* Name, double Time) override;
-	virtual void AppendRegionBeginWithId(const TCHAR* Name, uint64 Id, double Time) override;
+	virtual void AppendRegionBegin(const TCHAR* Name, double Time, const TCHAR* Category = nullptr) override;
+	virtual void AppendRegionBeginWithId(const TCHAR* Name, uint64 Id, double Time, const TCHAR* Category = nullptr) override;
 	virtual void AppendRegionEnd(const TCHAR* Name, double Time) override;
 	virtual void AppendRegionEndWithId(uint64 Id, double Time) override;
 
 	virtual void OnAnalysisSessionEnded() override;
 
 	//////////////////////////////////////////////////
-
+	
 private:
-	/// Update the depth member of a region to allow overlapping regions to be displayed on separate lanes.
-	int32 CalculateRegionDepth(const FTimeRegion& Item) const;
+
 	/// Calculates depth, inserts a new region into the correct lane and updates session time.
-	FTimeRegion* InsertNewRegion(double BeginTime, const TCHAR* Name, uint64 Id);
+	TArray<FTimeRegion*> InsertNewRegion(double BeginTime, const TCHAR* Name, uint64 Id, const TCHAR* Category);
 
 private:
 	mutable FProviderLock Lock;
@@ -68,11 +133,15 @@ private:
 	IAnalysisSession& Session;
 
 	// Open regions inside lanes
-	TMap<FStringView, FTimeRegion*> OpenRegionsByName;
-	TMap<uint64, FTimeRegion*> OpenRegionsById;
-
-	// Closed regions
-	TArray<FRegionLane> Lanes;
+	TMap<FStringView, TArray<FTimeRegion*>> OpenRegionsByName;
+	TMap<uint64, TArray<FTimeRegion*>> OpenRegionsById;
+	
+	FRegionTimeline AllRegions;
+	/// maps CategoryName to an FRegionTimeline.
+	/// Also contains a special entry with key UncategorizedName for any regions without a category specified
+	TMap<const TCHAR*, FRegionTimeline, FDefaultSetAllocator, TStringPointerMapKeyFuncs_DEPRECATED<const TCHAR*, FRegionTimeline>> RegionsPerCategory;
+	// string maps don't support nullptr, so use this key internally for all regions where category is nullptr
+	const TCHAR* UncategorizedName = TEXT("Uncategorized");
 
 	// Counter incremented each time region data changes during analysis
 	uint64 UpdateCounter = -1;
diff --git a/Engine/Source/Developer/TraceServices/Public/TraceServices/Model/Regions.h b/Engine/Source/Developer/TraceServices/Public/TraceServices/Model/Regions.h
--- a/Engine/Source/Developer/TraceServices/Public/TraceServices/Model/Regions.h
+++ b/Engine/Source/Developer/TraceServices/Public/TraceServices/Model/Regions.h
@@ -16,14 +16,16 @@ struct TRACESERVICES_API FTimeRegion
 	double BeginTime = std::numeric_limits<double>::infinity();
 	double EndTime = std::numeric_limits<double>::infinity();
 	const TCHAR* Text = nullptr;
-	// Id will be zero if the region is identified by Name only
+	/// Id will be zero if the region is identified by Name only
 	uint64 Id = 0;
 	int32 Depth = -1;
+	/// this will be nullptr for regions without a Category
+	const TCHAR* Category = nullptr;
 };
 
 class TRACESERVICES_API FRegionLane
 {
-	friend class FRegionProvider;
+	friend class FRegionTimeline;
 
 public:
 	FRegionLane(ILinearAllocator& InAllocator) : Regions(InAllocator, 512) {}
@@ -41,43 +43,107 @@ private:
 	TPagedArray<FTimeRegion> Regions;
 };
 
+/*
+ *	Sorts a set of Timing Regions into a stack of individual lanes without overlaps for display
+ */
+class TRACESERVICES_API IRegionTimeline
+{
+public:
+	/**
+	 * @return the category of this timeline if filtered or nullptr
+	 */
+	virtual const TCHAR* GetCategory() const = 0;
+	
+	/**
+	 * @return the number of lanes
+	 */
+	virtual int32 GetLaneCount() const = 0;
+
+	/**
+	 * Direct access to a certain lane at a given index/depth.
+	 * The pointer is valid only in the current read scope.
+	 * @return a pointer to the lane at the specified depth index or nullptr if Index > GetLaneCount()-1
+	 */
+	virtual const FRegionLane* GetLane(int32 Index) const = 0;
+
+	/**
+	 * Enumerates all regions that overlap a certain time interval. Will enumerate by depth but does not expose lanes.
+	 * 
+	 * @param Callback a callback called for each region. Return false to abort iteration.
+	 * @returns true if the enumeration finished, false if it was aborted by the callback returning false
+	 */
+	virtual bool EnumerateRegions(double IntervalStart, double IntervalEnd, TFunctionRef<bool(const FTimeRegion&)> Callback) const = 0;
+	
+	/**
+	 * Will call Callback(Lane, Depth) for each lane in order.
+	 */
+	virtual void EnumerateLanes(TFunctionRef<void(const FRegionLane&, const int32)> Callback) const = 0;
+};
+	
 class TRACESERVICES_API IRegionProvider
 	: public IProvider
 {
 public:
 	virtual ~IRegionProvider() override = default;
 
+	/**
+	*	Enumerates all timelines, this includes the uncategorized timeline with uncategorized regions and individual timelines for each Category
+	*	@param Callback a callback called for each timeline. Second parameter contains Category or nullptr for the uncategorized timeline.
+	*/
+	virtual void EnumerateTimelinesByCategory(TFunctionRef<void(const IRegionTimeline& /*Region*/, const TCHAR* /*Category*/)> Callback) const = 0;
+
+	/**
+	*	@return the default timeline containing all regions without filtering 
+	*/
+	virtual const IRegionTimeline& GetDefaultTimeline() const = 0;
+	
+	/**
+	*	@return the timeline for a given Category or the timeline with uncategorized regions if Category is nullptr. Returns nullptr if the category was invalid/not found
+	*/
+	virtual const IRegionTimeline* GetTimelineForCategory(const TCHAR* Category) const = 0;
+
+	/**
+	 * @return returns the string used to store regions with no explicit category set ("__Uncategorized__")
+	 */
+	virtual const TCHAR* GetUncategorizedRegionCategoryName() const = 0;
+	
 	/**
 	 * @return the amount of currently known regions (including open-ended ones)
 	 */
 	virtual uint64 GetRegionCount() const = 0;
 
 	/**
-	 * @return the number of lanes
+	 * @return the number of lanes of the default timeline
 	 */
+	UE_DEPRECATED(5.6, "Use GetDefaultTimeline().GetLaneCount() or EnumerateTimelinesByCategory() to access timelines instead.")
 	virtual int32 GetLaneCount() const = 0;
 
 	/**
-	 * Direct access to a certain lane at a given index/depth.
+	 * Direct access to a certain lane at a given index/depth of the default timeline.
 	 * The pointer is valid only in the current read scope.
 	 * @return a pointer to the lane at the specified depth index or nullptr if Index > GetLaneCount()-1
 	 */
+	UE_DEPRECATED(5.6, "Use GetDefaultTimeline().GetLane() or EnumerateTimelinesByCategory() to access timelines instead.")
 	virtual const FRegionLane* GetLane(int32 Index) const = 0;
 
 	/**
 	 * Enumerates all regions that overlap a certain time interval. Will enumerate by depth but does not expose lanes.
+	 * 
 	 * @param Callback a callback called for each region. Return false to abort iteration.
 	 * @returns true if the enumeration finished, false if it was aborted by the callback returning false
 	 */
-	virtual bool EnumerateRegions(double IntervalStart, double IntervalEnd, TFunctionRef<bool(const FTimeRegion&)> Callback) const = 0;
+	UE_DEPRECATED(5.6, "Use GetDefaultTimeline().EnumerateRegions() or EnumerateTimelinesByCategory() to access timelines instead.")
+	virtual bool EnumerateRegions(double IntervalStart, double IntervalEnd, TFunctionRef<bool(const FTimeRegion& /*Region*/)> Callback) const = 0;
 
 	/**
 	 * Will call Callback(Lane, Depth) for each lane in order.
 	 */
-	virtual void EnumerateLanes(TFunctionRef<void(const FRegionLane&, const int32)> Callback) const = 0;
+	
+	UE_DEPRECATED(5.6, "Use GetDefaultTimeline().EnumerateLanes() or EnumerateTimelinesByCategory() to access timelines instead.")
+	virtual void EnumerateLanes(TFunctionRef<void(const FRegionLane& /*Lane*/, const int32 /*Depth*/)> Callback) const = 0;
 
 	/**
-	 * @return A monotonically increasing counter that changes each time new data is added to the provider.
+	 * @return a monotonically increasing counter that is updated each time new data is added to the provider.
 	 * This can be used to detect when to update any (UI-)state dependent on the provider during analysis.
 	 */
 	virtual uint64 GetUpdateCounter() const = 0;
@@ -98,8 +164,9 @@ public:
 	 *
 	 * @param Name		The string name of the region.
 	 * @param Time		The time in seconds of the begin event of this region.
+	 * @param Category	The category associated with this region, use nullptr if no category has been set
 	 */
-	virtual void AppendRegionBegin(const TCHAR* Name, double Time) = 0;
+	virtual void AppendRegionBegin(const TCHAR* Name, double Time, const TCHAR* Category = nullptr) = 0;
 
 	/**
 	 * Append a new begin event of a region from the trace session.
@@ -108,8 +175,9 @@ public:
 	 * @param Name		The string name of the region.
 	 * @param Id		The Id of the region. Used to uniquely identify regions with the same name.
 	 * @param Time		The time in seconds of the begin event of this region.
+	 * @param Category	The category associated with this region, use nullptr if no category has been set
 	 */
-	virtual void AppendRegionBeginWithId(const TCHAR* Name, uint64 Id, double Time) = 0;
+	virtual void AppendRegionBeginWithId(const TCHAR* Name, uint64 Id, double Time, const TCHAR* Category = nullptr) = 0;
 
 	/**
 	 * Append a new end event of a region from the trace session.
