Metadata for cpu events emitted by RHI Breadcrumbs.

diff --git a/Engine/Source/Developer/TraceServices/Private/Analyzers/CpuProfilerTraceAnalysis.cpp b/Engine/Source/Developer/TraceServices/Private/Analyzers/CpuProfilerTraceAnalysis.cpp
--- a/Engine/Source/Developer/TraceServices/Private/Analyzers/CpuProfilerTraceAnalysis.cpp
+++ b/Engine/Source/Developer/TraceServices/Private/Analyzers/CpuProfilerTraceAnalysis.cpp
@@ -3,11 +3,13 @@
 #include "CpuProfilerTraceAnalysis.h"
 
 #include "AnalysisServicePrivate.h"
+#include "CborReader.h"
 #include "CborWriter.h"
 #include "Common/Utils.h"
 #include "HAL/LowLevelMemTracker.h"
 #include "Model/ThreadsPrivate.h"
 #include "Model/MonotonicTimeline.h"
+#include "Serialization/MemoryReader.h"
 #include "Serialization/MemoryWriter.h"
 #include "TraceServices/Model/TimingProfiler.h"
 #include "TraceServices/Utils.h"
@@ -50,6 +52,10 @@ void FCpuProfilerAnalyzer::OnAnalysisBegin(const FOnAnalysisContext& Context)
 	Builder.RouteEvent(RouteId_EventBatchV2, "CpuProfiler", "EventBatchV2"); // added in UE 5.1
 	Builder.RouteEvent(RouteId_EventBatch, "CpuProfiler", "EventBatch"); // backward compatibility; removed in UE 5.1
 	Builder.RouteEvent(RouteId_EndCapture, "CpuProfiler", "EndCapture"); // backward compatibility; removed in UE 5.1
+	Builder.RouteEvent(RouteId_MetadataSpec, "CpuProfiler", "MetadataSpec");
+	Builder.RouteEvent(RouteId_Metadata, "CpuProfiler", "Metadata");
+	Builder.RouteEvent(RouteId_BeginMetadataEvent, "CpuProfiler", "BeginMetadataEvent");
+	Builder.RouteEvent(RouteId_EndMetadataEvent, "CpuProfiler", "EndMetadataEvent");
 	Builder.RouteLoggerEvents(RouteId_CpuScope, "Cpu", true);
 }
 
@@ -166,7 +172,7 @@ bool FCpuProfilerAnalyzer::OnEvent(uint16 RouteId, EStyle Style, const FOnEventC
 			Line = EventData.GetValue<uint32>("Line");
 		}
 
-		DefineTimer(SpecId, Session.StoreString(TimerName), FileName, Line);
+		DefineTimer(SpecId, Session.StoreString(TimerName), FileName, Line, true);
 		break;
 	}
 
@@ -277,6 +283,195 @@ bool FCpuProfilerAnalyzer::OnEvent(uint16 RouteId, EStyle Style, const FOnEventC
 			OnCpuScopeLeave(Context);
 		}
 		break;
+
+	case RouteId_MetadataSpec:
+	{
+		uint32 SpecId = Context.EventData.GetValue<uint32>("Id");
+
+		FString Name, NameFormat;
+		Context.EventData.GetString("Name", Name);
+		Context.EventData.GetString("NameFormat", NameFormat);
+		TArrayView<const uint8> FieldNames = EventData.GetArrayView<uint8>("FieldNames");
+
+		if (Name.Compare(NameFormat) == 0)
+		{
+			Name.Empty();
+		}
+
+		int Index = 0;
+		while (Index < NameFormat.Len() && Index < Name.Len())
+		{
+			if (Name[Index] != NameFormat[Index])
+			{
+				break;
+			}
+
+			++Index;
+		}
+
+		if (Index > 1)
+		{
+			NameFormat.MidInline(Index);
+		}
+
+		if (Name.IsEmpty())
+		{
+			const FString Separators = TEXT("% (=");
+			Index = 0;
+			bool bIsInFormatSpecifier = false;
+			while (Index < NameFormat.Len())
+			{
+				int32 SpecIndex = -1;
+				if (Separators.FindChar(NameFormat[Index], SpecIndex))
+				{
+					Name = NameFormat.Left(Index);
+					NameFormat.MidInline(Index);
+					break;
+				}
+
+				++Index;
+			}
+
+			if (Name.IsEmpty())
+			{
+				if (!NameFormat.IsEmpty())
+				{
+					Name = NameFormat;
+					NameFormat.Empty();
+				}
+				else
+				{
+					Name = Session.StoreString(TEXT("Unknown"));
+				}
+			}
+		}
+
+		FMetadataSpec Spec;
+		if (!NameFormat.IsEmpty())
+		{
+			Spec.Format = Session.StoreString(NameFormat);
+		}
+
+		if (FieldNames.Num() > 0)
+		{
+			FMemoryReaderView MemoryReader(FieldNames);
+			FCborReader CborReader(&MemoryReader, ECborEndianness::StandardCompliant);
+			FCborContext CborContext;
+
+			while (CborReader.ReadNext(CborContext))
+			{
+				if (CborContext.MajorType() == ECborCode::TextString)
+				{
+					FString Field = CborContext.AsString();
+					Spec.FieldNames.Add(Session.StoreString(Field));
+				}
+			}
+		}
+
+		uint32 TimerId = 0;
+		uint32* FoundTimerIdPtr = SpecIdToTimerIdMap.Find(SpecId);
+		if (FoundTimerIdPtr == nullptr)
+		{
+			TimerId = DefineTimer(SpecId, Session.StoreString(*Name), nullptr, 0, false);
+		}
+		else
+		{
+			TimerId = *FoundTimerIdPtr;
+			EditableTimingProfilerProvider.SetTimerName(TimerId, Name);
+		}
+
+		if (Spec.FieldNames.Num() > 0 || Spec.Format != nullptr)
+		{
+			uint32 MetadataSpecId = EditableTimingProfilerProvider.AddMetadataSpec(MoveTemp(Spec));
+			EditableTimingProfilerProvider.SetMetadataSpec(TimerId, MetadataSpecId);
+		}
+
+		break;
+	}
+	case RouteId_Metadata:
+	{
+		uint32 MetadataId = Context.EventData.GetValue<uint32>("Id");
+		uint32 SpecId = Context.EventData.GetValue<uint32>("SpecId");
+		TArray<uint8> Metadata(EventData.GetArrayView<uint8>("Metadata"));
+
+		// We don't know if the Metadata or the BeginMetadataEvent events arrive first so handle both cases.
+		uint32* TimerIdPtr = MetadataIdToTimerIdMap.Find(MetadataId);
+		if (TimerIdPtr == nullptr)
+		{
+			uint32 TimerId = GetTimerId(SpecId);
+			TimerId = EditableTimingProfilerProvider.AddMetadata(TimerId, MoveTemp(Metadata));
+			MetadataIdToTimerIdMap.Add(MetadataId, TimerId);
+		}
+		else
+		{
+			uint32 TimerId = *TimerIdPtr;
+			// Replace the placeholder metadata added if we received a BeginMetadataEvent first.
+			EditableTimingProfilerProvider.SetMetadata(TimerId, MoveTemp(Metadata));
+		}
+
+		break;
+	}
+	case RouteId_BeginMetadataEvent:
+	{
+		const uint32 ThreadId = Context.ThreadInfo.GetId();
+		FThreadState& ThreadState = GetThreadState(ThreadId);
+		if (ThreadState.bShouldIgnorePendingEvents)
+		{
+			break;
+		}
+
+		uint32 MetadataId = Context.EventData.GetValue<uint32>("MetadataId");
+		uint32 SpecId = Context.EventData.GetValue<uint32>("SpecId");
+		uint64 Cycle = EventData.GetValue<uint64>("Cycle");
+		double Timestamp = Context.EventTime.AsSeconds(Cycle);
+
+		// We don't know if the Metadata or the BeginMetadataEvent events arrive first so handle both cases.
+		uint32 TimerId = 0;
+		uint32* TimerIdPtr = MetadataIdToTimerIdMap.Find(MetadataId);
+		if (TimerIdPtr == nullptr)
+		{
+			// Add an empty placeholder metadata so we obtain a MetadataId to use as the TimerId.
+			TimerId = GetTimerId(SpecId);
+			TimerId = EditableTimingProfilerProvider.AddMetadata(TimerId, TArray<uint8>());
+			MetadataIdToTimerIdMap.Add(MetadataId, TimerId);
+		}
+		else
+		{
+			TimerId = *TimerIdPtr;
+		}
+
+		check(ThreadState.LastCycle <= Cycle);
+		check(ThreadState.LastPendingEventTime <= Timestamp);
+		ThreadState.LastPendingEventTime = Timestamp;
+
+		ThreadState.PendingEvents.Add({ Cycle, Timestamp, TimerId });
+		++ThreadState.MetadataEventsDepth;
+		break;
+	}
+	case RouteId_EndMetadataEvent:
+	{
+		uint32 ThreadId = Context.ThreadInfo.GetId();
+		FThreadState& ThreadState = GetThreadState(ThreadId);
+
+		if (ThreadState.bShouldIgnorePendingEvents)
+		{
+			break;
+		}
+
+		uint64 Cycle = EventData.GetValue<uint64>("Cycle");
+		double Timestamp = Context.EventTime.AsSeconds(Cycle);
+
+		if (ThreadState.MetadataEventsDepth > 0)
+		{
+			check(ThreadState.LastCycle <= Cycle);
+			check(ThreadState.LastPendingEventTime <= Timestamp);
+			ThreadState.LastPendingEventTime = Timestamp;
+
+			ThreadState.PendingEvents.Add({ ~Cycle, Timestamp, 0 });
+			--ThreadState.MetadataEventsDepth;
+		}
+		break;
+	}
 	}
 
 	return true;
@@ -696,7 +891,7 @@ void FCpuProfilerAnalyzer::OnCpuScopeEnter(const FOnEventContext& Context)
 	{
 		FString ScopeName;
 		ScopeName += Context.EventData.GetTypeInfo().GetName();
-		TimerId = DefineTimer(SpecId, Session.StoreString(*ScopeName), nullptr, 0);
+		TimerId = DefineTimer(SpecId, Session.StoreString(*ScopeName), nullptr, 0, false);
 	}
 
 	TArray<uint8> CborData;
@@ -705,6 +900,11 @@ void FCpuProfilerAnalyzer::OnCpuScopeEnter(const FOnEventContext& Context)
 
 	uint64 Cycle = Context.EventTime.AsCycle64();
 	double Time = Context.EventTime.AsSeconds();
+
+	check(ThreadState.LastCycle <= Cycle);
+	check(ThreadState.LastPendingEventTime <= Time);
+	ThreadState.LastPendingEventTime = Time;
+
 	ThreadState.PendingEvents.Add({ Cycle, Time, TimerId });
 }
 
@@ -727,17 +927,22 @@ void FCpuProfilerAnalyzer::OnCpuScopeLeave(const FOnEventContext& Context)
 
 	uint64 Cycle = Context.EventTime.AsCycle64();
 	double Time = Context.EventTime.AsSeconds();
+	
+	check(ThreadState.LastCycle <= Cycle);
+	check(ThreadState.LastPendingEventTime <= Time);
+	ThreadState.LastPendingEventTime = Time;
+
 	ThreadState.PendingEvents.Add({ ~Cycle, Time, 0 });
 }
 
 ////////////////////////////////////////////////////////////////////////////////////////////////////
 
-uint32 FCpuProfilerAnalyzer::DefineTimer(uint32 SpecId, const TCHAR* Name, const TCHAR* File, uint32 Line)
+uint32 FCpuProfilerAnalyzer::DefineTimer(uint32 SpecId, const TCHAR* Name, const TCHAR* File, uint32 Line, bool bMergeByName)
 {
 	// Expected: Name is already a pointer in a string store.
 
-	// The CPU scoped events (timers) will be merged by name.
-	// Ex.: If there are multiple timers defined in code with same name,
+	// The CPU scoped events (timers) can be merged by name.
+	// Ex.: If bMergeByName == true and there are multiple timers defined in code with same name,
 	//      those will appear in Insights as a single timer.
 
 	// Check if a timer with same name was already defined.
@@ -752,7 +957,7 @@ uint32 FCpuProfilerAnalyzer::DefineTimer(uint32 SpecId, const TCHAR* Name, const
 
 		// Check if SpecId is already mapped to timer.
 		const uint32* FindTimerId = SpecIdToTimerIdMap.Find(SpecId);
-		if (FindTimerId)
+		if (FindTimerId && bMergeByName)
 		{
 			// Yes, SpecId was already mapped to a timer (ex. as an <unknown> timer).
 			TimerId = *FindTimerId;
@@ -766,8 +971,16 @@ uint32 FCpuProfilerAnalyzer::DefineTimer(uint32 SpecId, const TCHAR* Name, const
 		}
 		else
 		{
-			// Use the previously defined timer with same name.
-			TimerId = FindTimerIdByName->Id;
+			if (bMergeByName)
+			{
+				// Use the previously defined timer with same name.
+				TimerId = FindTimerIdByName->Id;
+			}
+			else
+			{
+				// Define a new CPU timer.
+				TimerId = EditableTimingProfilerProvider.AddCpuTimer(Name, File, Line);
+			}
 
 			// Map this SpecId to the previously defined timer with same name.
 			SpecIdToTimerIdMap.Add(SpecId, TimerId);
diff --git a/Engine/Source/Developer/TraceServices/Private/Analyzers/CpuProfilerTraceAnalysis.h b/Engine/Source/Developer/TraceServices/Private/Analyzers/CpuProfilerTraceAnalysis.h
--- a/Engine/Source/Developer/TraceServices/Private/Analyzers/CpuProfilerTraceAnalysis.h
+++ b/Engine/Source/Developer/TraceServices/Private/Analyzers/CpuProfilerTraceAnalysis.h
@@ -47,6 +47,8 @@ private:
 		IEditableTimeline<FTimingProfilerEvent>* Timeline = nullptr;
 		uint64 LastCycle = 0;
 		bool bShouldIgnorePendingEvents = false; // becomes true when we detect first pending event with incorrect timestamp (i.e < LastCycle)
+		double LastPendingEventTime = 0;
+		int32 MetadataEventsDepth = 0;
 	};
 
 	struct FTimerInfo
@@ -63,6 +65,10 @@ private:
 		RouteId_EventBatch, // backward compatibility
 		RouteId_EndCapture, // backward compatibility
 		RouteId_CpuScope,
+		RouteId_MetadataSpec,
+		RouteId_Metadata,
+		RouteId_BeginMetadataEvent,
+		RouteId_EndMetadataEvent,
 	};
 
 	void ProcessBuffer(const FEventTime& EventTime, FThreadState& ThreadState, const uint8* BufferPtr, uint32 BufferSize);
@@ -72,7 +78,7 @@ private:
 	void EndOpenEvents(FThreadState& ThreadState, double Timestamp);
 	void OnCpuScopeEnter(const FOnEventContext& Context);
 	void OnCpuScopeLeave(const FOnEventContext& Context);
-	uint32 DefineTimer(uint32 SpecId, const TCHAR* TimerName, const TCHAR* File, uint32 Line); // returns the TimerId
+	uint32 DefineTimer(uint32 SpecId, const TCHAR* TimerName, const TCHAR* File, uint32 Line, bool bMergeByName); // returns the TimerId
 	uint32 DefineNewTimerChecked(uint32 SpecId, const TCHAR* TimerName, const TCHAR* File = nullptr, uint32 Line = 0); // returns the TimerId
 	uint32 GetTimerId(uint32 SpecId);
 	FThreadState& GetThreadState(uint32 ThreadId);
@@ -85,6 +91,7 @@ private:
 	TMap<uint32, FThreadState*> ThreadStatesMap;
 	TMap<uint32, uint32> SpecIdToTimerIdMap; // SpecId --> TimerId
 	TMap<uint64, FTimerInfo> ScopeNameToTimerIdMap; // (uint64)Name --> FTimerInfo
+	TMap<uint32, uint32> MetadataIdToTimerIdMap; // MetadataId --> TimerId
 
 	uint32 CoroutineTimerId = ~0;
 	uint32 CoroutineUnknownTimerId = ~0;
diff --git a/Engine/Source/Developer/TraceServices/Private/Model/TimingProfiler.cpp b/Engine/Source/Developer/TraceServices/Private/Model/TimingProfiler.cpp
--- a/Engine/Source/Developer/TraceServices/Private/Model/TimingProfiler.cpp
+++ b/Engine/Source/Developer/TraceServices/Private/Model/TimingProfiler.cpp
@@ -155,6 +155,25 @@ TArrayView<const uint8> FTimingProfilerProvider::GetMetadata(uint32 MetadataTime
 	return Metadata.Payload;
 }
 
+void FTimingProfilerProvider::SetMetadata(uint32 MetadataTimerId, TArray<uint8>&& Metadata)
+{
+	Session.WriteAccessCheck();
+
+	if (int32(MetadataTimerId) >= 0)
+	{
+		return;
+	}
+
+	MetadataTimerId = ~MetadataTimerId;
+	if (MetadataTimerId >= uint32(Metadatas.Num()))
+	{
+		return;
+	}
+
+	FMetadata& MetadataToReplace = Metadatas[MetadataTimerId];
+	MetadataToReplace.Payload = Metadata;
+}
+
 TArrayView<uint8> FTimingProfilerProvider::GetEditableMetadata(uint32 MetadataTimerId)
 {
 	Session.WriteAccessCheck();
diff --git a/Engine/Source/Developer/TraceServices/Private/Model/TimingProfilerPrivate.h b/Engine/Source/Developer/TraceServices/Private/Model/TimingProfilerPrivate.h
--- a/Engine/Source/Developer/TraceServices/Private/Model/TimingProfilerPrivate.h
+++ b/Engine/Source/Developer/TraceServices/Private/Model/TimingProfilerPrivate.h
@@ -69,6 +69,7 @@ public:
 	virtual void SetTimerName(uint32 TimerId, FStringView Name) override;
 	virtual void SetTimerNameAndLocation(uint32 TimerId, FStringView Name, const TCHAR* File, uint32 Line) override;
 	virtual uint32 AddMetadata(uint32 OriginalTimerId, TArray<uint8>&& Metadata) override;
+	virtual void SetMetadata(uint32 MetadataTimerId, TArray<uint8>&& Metadata) override;
 	virtual TArrayView<uint8> GetEditableMetadata(uint32 MetadataTimerId) override;
 
 	virtual uint32 AddMetadataSpec(FMetadataSpec&& Metadata) override;
diff --git a/Engine/Source/Developer/TraceServices/Public/TraceServices/Model/TimingProfiler.h b/Engine/Source/Developer/TraceServices/Public/TraceServices/Model/TimingProfiler.h
--- a/Engine/Source/Developer/TraceServices/Public/TraceServices/Model/TimingProfiler.h
+++ b/Engine/Source/Developer/TraceServices/Public/TraceServices/Model/TimingProfiler.h
@@ -366,6 +366,15 @@ public:
 	*/
 	virtual uint32 AddMetadata(uint32 OriginalTimerId, TArray<uint8>&& Metadata) = 0;
 
+	/*
+	* Sets metadata for the specified MetadataTimerId. The MetadataTimerId must be a value returned by AddMetadata. The function is meant to be used to replace the metadata at a MetadataTimerId.
+	*
+	* @param MetadataTimerId	The identity of the metadata to replace.
+	* @param Metadata			The metadata.
+	*
+	*/
+	virtual void SetMetadata(uint32 MetadataTimerId, TArray<uint8>&& Metadata) {}
+
 	/*
 	* Gets metadata by id.
 	*
diff --git a/Engine/Source/Runtime/Core/Private/ProfilingDebugging/CpuProfilerTrace.cpp b/Engine/Source/Runtime/Core/Private/ProfilingDebugging/CpuProfilerTrace.cpp
--- a/Engine/Source/Runtime/Core/Private/ProfilingDebugging/CpuProfilerTrace.cpp
+++ b/Engine/Source/Runtime/Core/Private/ProfilingDebugging/CpuProfilerTrace.cpp
@@ -32,6 +32,29 @@ UE_TRACE_EVENT_BEGIN(CpuProfiler, EventSpec, NoSync|Important)
 	#endif
 UE_TRACE_EVENT_END()
 
+UE_TRACE_EVENT_BEGIN(CpuProfiler, MetadataSpec, NoSync | Important)
+	UE_TRACE_EVENT_FIELD(uint32, Id)
+	UE_TRACE_EVENT_FIELD(UE::Trace::AnsiString, Name)
+	UE_TRACE_EVENT_FIELD(UE::Trace::WideString, NameFormat)
+	UE_TRACE_EVENT_FIELD(uint8[], FieldNames)
+UE_TRACE_EVENT_END()
+
+UE_TRACE_EVENT_BEGIN(CpuProfiler, Metadata, NoSync)
+	UE_TRACE_EVENT_FIELD(uint32, Id)
+	UE_TRACE_EVENT_FIELD(uint32, SpecId)
+	UE_TRACE_EVENT_FIELD(uint8[], Metadata)
+UE_TRACE_EVENT_END()
+
+UE_TRACE_EVENT_BEGIN(CpuProfiler, BeginMetadataEvent, NoSync)
+	UE_TRACE_EVENT_FIELD(uint64, Cycle)
+	UE_TRACE_EVENT_FIELD(uint32, MetadataId)
+	UE_TRACE_EVENT_FIELD(uint32, SpecId)
+UE_TRACE_EVENT_END()
+
+UE_TRACE_EVENT_BEGIN(CpuProfiler, EndMetadataEvent, NoSync)
+	UE_TRACE_EVENT_FIELD(uint64, Cycle)
+UE_TRACE_EVENT_END()
+
 UE_TRACE_EVENT_BEGIN(CpuProfiler, EventBatchV2, NoSync)
 	UE_TRACE_EVENT_FIELD(uint8[], Data)
 UE_TRACE_EVENT_END()
@@ -110,6 +133,7 @@ struct FCpuProfilerTraceInternal
 	};
 
 	uint32 static GetNextSpecId();
+	uint32 static GetNextMetadataId();
 	FORCENOINLINE static FThreadBuffer* CreateThreadBuffer();
 	FORCENOINLINE static void FlushThreadBuffer(FThreadBuffer* ThreadBuffer);
 
@@ -438,6 +462,12 @@ uint32 FCpuProfilerTraceInternal::GetNextSpecId()
 	return (NextSpecId++) + 1;
 }
 
+uint32 FCpuProfilerTraceInternal::GetNextMetadataId()
+{
+	static TAtomic<uint32> NextMetadataId;
+	return (NextMetadataId++) + 1;
+}
+
 uint32 FCpuProfilerTrace::OutputEventType(const TCHAR* Name, const ANSICHAR* File, uint32 Line)
 {
 	if (AutoRTFM::IsClosed())
@@ -466,6 +496,87 @@ uint32 FCpuProfilerTrace::OutputEventType(const TCHAR* Name, const ANSICHAR* Fil
 	return SpecId;
 }
 
+void FCpuProfilerTrace::OutputEventMetadataSpec(uint32 SpecId, const TCHAR* StaticName, const TCHAR* NameFormat, const TArray<uint8>& FieldNames)
+{
+	if (AutoRTFM::IsClosed())
+	{
+		return AutoRTFM::Open([&]
+			{
+				return OutputEventMetadataSpec(SpecId, StaticName, NameFormat, FieldNames);
+			});
+	}
+
+	uint16 NameLen = uint16(FCString::Strlen(StaticName));
+	uint32 DataSize = NameLen * sizeof(ANSICHAR); // EventSpec.Name is traced as UE::Trace::AnsiString
+	uint16 NameFormatLen = uint16(FCString::Strlen(NameFormat));
+	DataSize += NameFormatLen * sizeof(TCHAR);
+	DataSize += FieldNames.Num() * sizeof(uint8);
+
+	UE_TRACE_LOG(CpuProfiler, MetadataSpec, CpuChannel, DataSize)
+		<< MetadataSpec.Id(SpecId)
+		<< MetadataSpec.Name(StaticName, NameLen)
+		<< MetadataSpec.NameFormat(NameFormat, NameFormatLen)
+		<< MetadataSpec.FieldNames(FieldNames.GetData(), FieldNames.Num());
+}
+
+uint32 FCpuProfilerTrace::OutputMetadata(uint32 SpecId, const TArray<uint8>& InMetadata)
+{
+	if (AutoRTFM::IsClosed())
+	{
+		return AutoRTFM::Open([&]
+			{
+				return OutputMetadata(SpecId, InMetadata);
+			});
+	}
+
+	if (!UE_TRACE_CHANNELEXPR_IS_ENABLED(CpuChannel))
+	{
+		return 0;
+	}
+
+	uint32 MetadataId = FCpuProfilerTraceInternal::GetNextMetadataId();
+
+	UE_TRACE_LOG(CpuProfiler, Metadata, CpuChannel)
+		<< Metadata.Id(MetadataId)
+		<< Metadata.SpecId(SpecId)
+		<< Metadata.Metadata(InMetadata.GetData(), InMetadata.Num());
+
+	return MetadataId;
+}
+
+void FCpuProfilerTrace::OutputBeginEventWithMetadata(uint32 SpecId, uint32 MetadataId)
+{
+	if (AutoRTFM::IsClosed())
+	{
+		AutoRTFM::Open([&]
+			{
+				OutputBeginEventWithMetadata(SpecId, MetadataId);
+			});
+
+		AutoRTFM::PushOnAbortHandler(&FCpuProfilerTrace::OnAbortKey, []
+			{
+				OutputEndEventWithMetadata();
+			});
+
+		return;
+	}
+
+	uint64 Cycle = FPlatformTime::Cycles64();
+
+	UE_TRACE_LOG(CpuProfiler, BeginMetadataEvent, CpuChannel)
+		<< BeginMetadataEvent.Cycle(Cycle)
+		<< BeginMetadataEvent.MetadataId(MetadataId)
+		<< BeginMetadataEvent.SpecId(SpecId);
+}
+
+void FCpuProfilerTrace::OutputEndEventWithMetadata()
+{
+	uint64 Cycle = FPlatformTime::Cycles64();
+
+	UE_TRACE_LOG(CpuProfiler, EndMetadataEvent, CpuChannel)
+		<< EndMetadataEvent.Cycle(Cycle);
+}
+
 uint32 FCpuProfilerTrace::OutputEventType(const ANSICHAR* Name, const ANSICHAR* File, uint32 Line)
 {
 	if (AutoRTFM::IsClosed())
diff --git a/Engine/Source/Runtime/Core/Public/ProfilingDebugging/CpuProfilerTrace.h b/Engine/Source/Runtime/Core/Public/ProfilingDebugging/CpuProfilerTrace.h
--- a/Engine/Source/Runtime/Core/Public/ProfilingDebugging/CpuProfilerTrace.h
+++ b/Engine/Source/Runtime/Core/Public/ProfilingDebugging/CpuProfilerTrace.h
@@ -3,6 +3,7 @@
 #pragma once
 
 #include "CoreTypes.h"
+#include "Containers/ContainersFwd.h"
 #include "HAL/PlatformAtomics.h"
 #include "HAL/PreprocessorHelpers.h"
 #include "Misc/Build.h"
@@ -52,6 +53,25 @@ struct FCpuProfilerTrace
 	FORCENOINLINE CORE_API static uint32 OutputEventType(const TCHAR* Name, const ANSICHAR* File = nullptr, uint32 Line = 0);
 	FORCENOINLINE CORE_API static uint32 OutputEventType(FNameParam Name, const ANSICHAR* File = nullptr, uint32 Line = 0);
 
+	/*
+	 * Output CPU event metadata spec.
+	 * The trace event emitted by this function is an "important event" (so all events emitted will add to the Trace System's cache).
+	 * It is the responsibility of the caller code to ensure this function is not abused.
+	 * @param SpecId - The SpecId of the event to associate this metadata spec with
+	 * @param StaticName - The static name of the event
+	 * @param NameFormat - The format string for the event
+	 * @param FieldNames - Cbor data representing the serialized field names. These should be descriptive names for the metadata values passed to OutputBeginEventWithMetadata.
+	 */
+	FORCENOINLINE CORE_API static void OutputEventMetadataSpec(uint32 SpecId, const TCHAR* StaticName, const TCHAR* NameFormat, const TArray<uint8>& FieldNames);
+
+	/*
+	 * Output CPU event metadata. This metadata can than be associated with multiple timing events emitted with OutputBeginEventWithMetadata.
+	 * @param SpecId - The SpecId of the event to associate this metadata with.
+	 * @param CborData - The serialized metadata associated with the event.
+	 * @returns A MetadataId that can be used with OutputBeginEventWithMetadata. Returns 0 if the cpu channel is off.
+	 */
+	FORCENOINLINE CORE_API static uint32 OutputMetadata(uint32 SpecId, const TArray<uint8>& CborData);
+
 	/*
 	 * Output CPU event definition (spec) for a dynamic event.
 	 * The Name will be cached and the trace event will only be emitted once (for each unique Name; even if File or Line changes).
@@ -82,6 +102,18 @@ struct FCpuProfilerTrace
 	 */
 	CORE_API static void OutputBeginEvent(uint32 SpecId);
 
+	/*
+	 * Output begin event marker with metadata for a given spec. Must always be matched with an end event.
+	 * @param SpecId - The event definition id.
+	 * @param MetadataId - The MetadataId returned from OutputMetadata.
+	 */
+	CORE_API static void OutputBeginEventWithMetadata(uint32 SpecId, uint32 MetadataId);
+
+	/*
+	 * Output end event marker for an event traced with OutputBeginEventWithMetadata.
+	 */
+	CORE_API static void OutputEndEventWithMetadata();
+
 	/*
 	 * Output begin event marker for a dynamic event name. This is more expensive than statically known event
 	 * names using \ref OutputBeginEvent. Must always be matched with an end event.
diff --git a/Engine/Source/Runtime/RHI/Private/GpuProfilerTrace.cpp b/Engine/Source/Runtime/RHI/Private/GpuProfilerTrace.cpp
--- a/Engine/Source/Runtime/RHI/Private/GpuProfilerTrace.cpp
+++ b/Engine/Source/Runtime/RHI/Private/GpuProfilerTrace.cpp
@@ -287,7 +287,7 @@ UE_TRACE_EVENT_END()
 
 #if WITH_RHI_BREADCRUMBS
 
-UE_TRACE_EVENT_BEGIN(GpuProfiler, EventBreadcrumbSpec)
+UE_TRACE_EVENT_BEGIN(GpuProfiler, EventBreadcrumbSpec, NoSync | Important)
 	UE_TRACE_EVENT_FIELD(uint32, SpecId)
 	UE_TRACE_EVENT_FIELD(UE::Trace::WideString, StaticName)
 	UE_TRACE_EVENT_FIELD(UE::Trace::WideString, NameFormat)
@@ -428,8 +428,6 @@ void FGpuProfilerTrace::WaitFence(uint32 QueueId, uint64 ResolvedTimestamp, uint
 }
 
 std::atomic<uint32> FGpuProfilerTrace::NextSpecId = 1;
-TMap<const void*, uint32> FGpuProfilerTrace::DynamicNamesMap;
-FCriticalSection FGpuProfilerTrace::Cs;
 
 uint32 FGpuProfilerTrace::InternalBreadcrumbSpec(const TCHAR* StaticName, const TCHAR* NameFormat, const TArray<uint8>& FieldNames)
 {
@@ -440,8 +438,12 @@ uint32 FGpuProfilerTrace::InternalBreadcrumbSpec(const TCHAR* StaticName, const
 
 	uint32 SpecId = NextSpecId.fetch_add(1);
 
+	uint32 DataSize = FCString::Strlen(StaticName) * sizeof(TCHAR);
+	DataSize += FCString::Strlen(NameFormat) * sizeof(TCHAR);
+	DataSize += FieldNames.Num() * sizeof(uint8);
+
 #if WITH_RHI_BREADCRUMBS
-	UE_TRACE_LOG(GpuProfiler, EventBreadcrumbSpec, GpuChannel)
+	UE_TRACE_LOG(GpuProfiler, EventBreadcrumbSpec, GpuChannel, DataSize)
 		<< EventBreadcrumbSpec.SpecId(SpecId)
 		<< EventBreadcrumbSpec.StaticName(StaticName)
 		<< EventBreadcrumbSpec.NameFormat(NameFormat)
diff --git a/Engine/Source/Runtime/RHI/Public/GpuProfilerTrace.h b/Engine/Source/Runtime/RHI/Public/GpuProfilerTrace.h
--- a/Engine/Source/Runtime/RHI/Public/GpuProfilerTrace.h
+++ b/Engine/Source/Runtime/RHI/Public/GpuProfilerTrace.h
@@ -142,8 +142,6 @@ private:
 	RHI_API static uint32 InternalBreadcrumbSpec(const TCHAR* StaticName, const TCHAR* NameFormat, const TArray<uint8>& FieldNames);
 
 	static std::atomic<uint32> NextSpecId;
-	static TMap<const void*, uint32> DynamicNamesMap;
-	static FCriticalSection Cs;
 
 public:
 
diff --git a/Engine/Source/Runtime/RHI/Public/RHIBreadcrumbs.h b/Engine/Source/Runtime/RHI/Public/RHIBreadcrumbs.h
--- a/Engine/Source/Runtime/RHI/Public/RHIBreadcrumbs.h
+++ b/Engine/Source/Runtime/RHI/Public/RHIBreadcrumbs.h
@@ -6,6 +6,7 @@
 #include "Misc/MemStack.h"
 #include "HAL/Platform.h"
 #include "Templates/SharedPointer.h"
+#include "ProfilingDebugging/CpuProfilerTrace.h"
 #include "ProfilingDebugging/CsvProfiler.h"
 #include "GenericPlatform/GenericPlatformCrashContext.h"
 
@@ -239,7 +240,8 @@ namespace UE::RHI::GPUProfiler
 
 		FRHIBreadcrumbData const& Data;
 	#if RHI_BREADCRUMBS_EMIT_CPU
-		uint32 const CPUTraceMarkerID = 0;
+		uint32 TraceCpuSpecId = 0;
+		uint32 TraceCpuMetadataId = 0;
 	#endif
 
 		uint32 const ID = 0;
@@ -265,8 +267,8 @@ namespace UE::RHI::GPUProfiler
 		virtual void TraceBeginGPU(uint32 QueueId, uint64 GPUTimestampTOP) const = 0;
 		virtual void TraceEndGPU(uint32 QueueId, uint64 GPUTimestampTOP) const = 0;
 
-		inline void BeginCPU() const;
-		inline void EndCPU() const;
+		inline void TraceBeginCPU() const;
+		inline void TraceEndCPU() const;
 
 		// Calls BeginCPU() on all the breadcrumb nodes between the root and the specified node.
 		// Only valid to call from the bottom-of-pipe, after the dispatch thread has fixed up the breadcrumb tree.
@@ -321,22 +323,6 @@ namespace UE::RHI::GPUProfiler
 		// Constructor for the sentinel value
 		FRHIBreadcrumbNode(FRHIBreadcrumbData const& Data);
 		friend struct FRHIBreadcrumbList;
-
-		void CreateTraceMarkers()
-		{
-		#if RHI_BREADCRUMBS_EMIT_CPU
-			if (TRACE_CPUPROFILER_EVENT_MANUAL_IS_ENABLED())
-			{
-				FRHIBreadcrumbNode::FBuffer Buffer;
-				TCHAR const* Str = GetTCHAR(Buffer);
-				const_cast<uint32&>(CPUTraceMarkerID) = FCpuProfilerTrace::OutputDynamicEventType(Str
-			#if RHI_BREADCRUMBS_EMIT_LOCATION
-					, Data.File, Data.Line
-			#endif
-				);
-			}
-		#endif
-		}
 	};
 
 	// Typedef for backwards compatibility. The FRHIBreadcrumb and FRHIBreadcrumbNode have been merged into one type.
@@ -655,24 +641,37 @@ namespace UE::RHI::GPUProfiler
 		}
 	}
 
-	inline void FRHIBreadcrumbNode::BeginCPU() const
+	inline void FRHIBreadcrumbNode::TraceBeginCPU() const
 	{
-	#if RHI_BREADCRUMBS_EMIT_CPU
-		if (CPUTraceMarkerID)
+#if RHI_BREADCRUMBS_EMIT_CPU
+		if (TraceCpuSpecId)
 		{
-			FCpuProfilerTrace::OutputBeginEvent(CPUTraceMarkerID);
+			if (TraceCpuMetadataId > 0)
+			{
+				FCpuProfilerTrace::OutputBeginEventWithMetadata(TraceCpuSpecId, TraceCpuMetadataId);
+			}
+			else
+			{
+				FCpuProfilerTrace::OutputBeginEvent(TraceCpuSpecId);
+			}
 		}
-	#endif
+#endif
 	}
-
-	inline void FRHIBreadcrumbNode::EndCPU() const
+	inline void FRHIBreadcrumbNode::TraceEndCPU() const
 	{
-	#if RHI_BREADCRUMBS_EMIT_CPU
-		if (CPUTraceMarkerID)
+#if RHI_BREADCRUMBS_EMIT_CPU
+		if (TraceCpuSpecId)
 		{
-			FCpuProfilerTrace::OutputEndEvent();
+			if (TraceCpuMetadataId > 0)
+			{
+				FCpuProfilerTrace::OutputEndEventWithMetadata();
+			}
+			else
+			{
+				FCpuProfilerTrace::OutputEndEvent();
+			}
 		}
-	#endif
+#endif
 	}
 
 	inline void FRHIBreadcrumbNode::WalkIn(FRHIBreadcrumbNode const* Node)
@@ -686,7 +685,7 @@ namespace UE::RHI::GPUProfiler
 					return;
 
 				Recurse(Recurse, Current->GetParent());
-				Current->BeginCPU();
+				Current->TraceBeginCPU();
 			};
 			Recurse(Recurse, Node);
 		}
@@ -707,7 +706,7 @@ namespace UE::RHI::GPUProfiler
 					Recurse(Recurse, Current->GetParent());
 				}
 
-				Current->BeginCPU();
+				Current->TraceBeginCPU();
 			};
 			Recurse(Recurse, Leaf);
 		}
@@ -721,7 +720,7 @@ namespace UE::RHI::GPUProfiler
 		{
 			while (Node && Node != Sentinel)
 			{
-				Node->EndCPU();
+				Node->TraceEndCPU();
 				Node = Node->GetParent();
 			}
 		}
@@ -737,7 +736,7 @@ namespace UE::RHI::GPUProfiler
 		{
 			while (true)
 			{
-				Leaf->EndCPU();
+				Leaf->TraceEndCPU();
 				if (Leaf == Root)
 					break;
 
@@ -1078,6 +1077,8 @@ namespace UE::RHI::GPUProfiler
 		{
 			using TFormatString = TFormatString<FormatStringSize, sizeof...(TValues) != 0>;
 
+			mutable FCriticalSection Cs;
+
 			template <typename TTuple, size_t... Indices>
 			static auto ExtractFieldNames(TTuple&& Tuple, std::index_sequence<Indices...>)
 			{
@@ -1102,8 +1103,11 @@ namespace UE::RHI::GPUProfiler
 			// Array of the string literal field names
 			std::array<TCHAR const*, NumFields> const FieldNames;
 
-			// Id of the spec of the gpu events associated with this breadcrumb.
+			// Id of the spec of the GPU events associated with this breadcrumb.
 			mutable uint32 TraceGpuSpecId = 0;
+			
+			// Id of the spec of the CPU events associated with this breadcrumb.
+			mutable std::atomic<uint32> TraceCpuSpecId = 0;
 
 			template <typename... TInnerValues>
 			TRHIBreadcrumbDesc(FRHIBreadcrumbData&& BaseData, TFormatString FormatString, TInnerValues&&... Values)
@@ -1142,6 +1146,36 @@ namespace UE::RHI::GPUProfiler
 				return TraceGpuSpecId;
 			}
 
+			uint32 GetTraceCpuSpec() const
+			{
+#if CPUPROFILERTRACE_ENABLED
+				if (TraceCpuSpecId == 0 && UE::RHI::GPUProfiler::FGpuProfilerTrace::IsAvailable())
+				{
+					FScopeLock Lock(&Cs);
+
+					if (TraceCpuSpecId == 0)
+					{
+						TraceCpuSpecId = FCpuProfilerTrace::OutputEventType(StaticName
+#if RHI_BREADCRUMBS_EMIT_LOCATION
+							, File, Line
+#endif
+						);
+
+						if constexpr (FormatStringSize > 0)
+						{
+							UE::RHI::GPUProfiler::FMetadataSerializer Serializer;
+							for (const TCHAR* FieldName : FieldNames)
+							{
+								Serializer.AppendValue(FieldName);
+							}
+							FCpuProfilerTrace::OutputEventMetadataSpec(TraceCpuSpecId, StaticName, TFormatString::FormatString, Serializer.GetData());
+						}
+					}
+				}
+#endif
+				return TraceCpuSpecId;
+			}
+
 		private:
 			template<typename T>
 			void SerializeValue(const T& Item, UE::RHI::GPUProfiler::FMetadataSerializer& Serializer) const
@@ -1170,7 +1204,18 @@ namespace UE::RHI::GPUProfiler
 				, TDesc::TValuesTuple(Forward<TValues>(Values)...)
 				, Desc(Desc)
 			{
-				CreateTraceMarkers();
+#if RHI_BREADCRUMBS_EMIT_CPU
+				TraceCpuSpecId = Desc.GetTraceCpuSpec();
+				if (TraceCpuSpecId)
+				{
+					if (Desc.NumValues > 0)
+					{
+						UE::RHI::GPUProfiler::FMetadataSerializer Serializer;
+						Desc.SerializeValues(Serializer, *this);
+						TraceCpuMetadataId = FCpuProfilerTrace::OutputMetadata(TraceCpuSpecId, Serializer.GetData());
+					}
+				}
+#endif
 			}
 
 			TCHAR const* GetTCHAR(FBuffer& Buffer) const override
diff --git a/Engine/Source/Runtime/RHI/Public/RHICommandList.h b/Engine/Source/Runtime/RHI/Public/RHICommandList.h
--- a/Engine/Source/Runtime/RHI/Public/RHICommandList.h
+++ b/Engine/Source/Runtime/RHI/Public/RHICommandList.h
@@ -3389,7 +3389,7 @@ public:
 		if (IsTopOfPipe())
 		{
 			// Recording thread
-			Breadcrumb->BeginCPU();
+			Breadcrumb->TraceBeginCPU();
 			PersistentState.LocalBreadcrumb = Breadcrumb;
 
 			if (bLink)
@@ -3411,7 +3411,7 @@ public:
 			if (bLink)
 			{
 				ExecutingCmdList.CPUBreadcrumbState.Current = Breadcrumb;
-				Breadcrumb->BeginCPU();
+				Breadcrumb->TraceBeginCPU();
 			}
 		});
 	}
@@ -3424,7 +3424,7 @@ public:
 		if (IsTopOfPipe())
 		{
 			// Recording thread
-			Breadcrumb->EndCPU();
+			Breadcrumb->TraceEndCPU();
 			PersistentState.LocalBreadcrumb = Breadcrumb->GetParent();
 
 			if (bLink)
@@ -3444,7 +3444,7 @@ public:
 				ExecutingCmdList.CPUBreadcrumbState.Current = Breadcrumb->GetParent();
 				check(ExecutingCmdList.CPUBreadcrumbState.Current != FRHIBreadcrumbNode::Sentinel);
 
-				Breadcrumb->EndCPU();
+				Breadcrumb->TraceEndCPU();
 			}
 		});
 	}
diff --git a/Engine/Source/Runtime/RenderCore/Public/RenderGraphEvent.inl b/Engine/Source/Runtime/RenderCore/Public/RenderGraphEvent.inl
--- a/Engine/Source/Runtime/RenderCore/Public/RenderGraphEvent.inl
+++ b/Engine/Source/Runtime/RenderCore/Public/RenderGraphEvent.inl
@@ -112,7 +112,7 @@ inline void FRDGScope_GPU::EndGPU(FRHIComputeCommandList& RHICmdList)
 		{
 			Node->SetParent(State.CurrentBreadcrumbRef);
 			State.CurrentBreadcrumbRef = Node;
-			Node->BeginCPU();
+			Node->TraceBeginCPU();
 
 			if (!State.ScopeState.bImmediate)
 			{
@@ -127,7 +127,7 @@ inline void FRDGScope_GPU::EndGPU(FRHIComputeCommandList& RHICmdList)
 	#if WITH_RHI_BREADCRUMBS
 		if (Node)
 		{
-			Node->EndCPU();
+			Node->TraceEndCPU();
 			State.CurrentBreadcrumbRef = Node->GetParent();
 		}
 	#endif
