Revert 41482232 as there is old data around.

#rnx

diff --git a/Engine/Source/Developer/TraceServices/Private/Analyzers/CpuProfilerTraceAnalysis.cpp b/Engine/Source/Developer/TraceServices/Private/Analyzers/CpuProfilerTraceAnalysis.cpp
--- a/Engine/Source/Developer/TraceServices/Private/Analyzers/CpuProfilerTraceAnalysis.cpp
+++ b/Engine/Source/Developer/TraceServices/Private/Analyzers/CpuProfilerTraceAnalysis.cpp
@@ -1050,15 +1050,6 @@ uint32 FCpuProfilerAnalyzer::GetTimerId(uint32 SpecId)
 
 FCpuProfilerAnalyzer::FThreadState& FCpuProfilerAnalyzer::GetThreadState(uint32 ThreadId)
 {
-	// Strip off legacy bias that was used to transition the CPU event protocol in the very early
-	// days. There used to be a ThreadId field in CpuProfiler::EventBatch which was supersceeded by
-	// thread id in transport packets. The bias was added to avoid thread ID collision. The below
-	// allows EventBatch to be brought out of retirement for adhob async/off-CPU introspection.
-	if (uint32 EventInlined = ThreadId & 0x7000'0000; EventInlined != 0)
-	{
-		ThreadId &= ~EventInlined;
-	}
-
 	FThreadState* ThreadState = ThreadStatesMap.FindRef(ThreadId);
 	if (!ThreadState)
 	{
