Integrate dynamic metadata cpu events in the main cpu event stream.

diff --git a/Engine/Source/Developer/TraceServices/Private/Analyzers/CpuProfilerTraceAnalysis.cpp b/Engine/Source/Developer/TraceServices/Private/Analyzers/CpuProfilerTraceAnalysis.cpp
--- a/Engine/Source/Developer/TraceServices/Private/Analyzers/CpuProfilerTraceAnalysis.cpp
+++ b/Engine/Source/Developer/TraceServices/Private/Analyzers/CpuProfilerTraceAnalysis.cpp
@@ -49,13 +49,12 @@ void FCpuProfilerAnalyzer::OnAnalysisBegin(const FOnAnalysisContext& Context)
 
 	Builder.RouteEvent(RouteId_EventSpec, "CpuProfiler", "EventSpec");
 	Builder.RouteEvent(RouteId_EndThread, "CpuProfiler", "EndThread");
-	Builder.RouteEvent(RouteId_EventBatchV2, "CpuProfiler", "EventBatchV2"); // added in UE 5.1
+	Builder.RouteEvent(RouteId_EventBatchV3, "CpuProfiler", "EventBatchV3"); // added in UE 5.6
+	Builder.RouteEvent(RouteId_EventBatchV2, "CpuProfiler", "EventBatchV2"); // backward compatibility, added in UE 5.1, removed in 5.6
 	Builder.RouteEvent(RouteId_EventBatch, "CpuProfiler", "EventBatch"); // backward compatibility; removed in UE 5.1
 	Builder.RouteEvent(RouteId_EndCapture, "CpuProfiler", "EndCapture"); // backward compatibility; removed in UE 5.1
 	Builder.RouteEvent(RouteId_MetadataSpec, "CpuProfiler", "MetadataSpec");
 	Builder.RouteEvent(RouteId_Metadata, "CpuProfiler", "Metadata");
-	Builder.RouteEvent(RouteId_BeginMetadataEvent, "CpuProfiler", "BeginMetadataEvent");
-	Builder.RouteEvent(RouteId_EndMetadataEvent, "CpuProfiler", "EndMetadataEvent");
 	Builder.RouteLoggerEvents(RouteId_CpuScope, "Cpu", true);
 }
 
@@ -230,6 +229,31 @@ bool FCpuProfilerAnalyzer::OnEvent(uint16 RouteId, EStyle Style, const FOnEventC
 		TotalEventSize += BufferSize;
 		break;
 	}
+	case RouteId_EventBatchV3:
+	{
+		const uint32 ThreadId = Context.ThreadInfo.GetId();
+		FThreadState& ThreadState = GetThreadState(ThreadId);
+
+		if (ThreadState.LastCycle == ~0ull)
+		{
+			// Ignore timing events received after EndThread.
+			break;
+		}
+
+		TArrayView<const uint8> DataView = Context.EventData.GetArrayView<uint8>("Data");
+		const uint32 BufferSize = DataView.Num();
+		const uint8* BufferPtr = DataView.GetData();
+
+		ProcessBufferV3(Context.EventTime, ThreadState, BufferPtr, BufferSize);
+		if (ThreadState.LastCycle != 0)
+		{
+			double Timestamp = Context.EventTime.AsSeconds(ThreadState.LastCycle);
+			Session.UpdateDurationSeconds(Timestamp);
+		}
+
+		TotalEventSize += BufferSize;
+		break;
+	}
 
 	case RouteId_EventBatch: // backward compatibility
 	case RouteId_EndCapture: // backward compatibility
@@ -404,72 +428,12 @@ bool FCpuProfilerAnalyzer::OnEvent(uint16 RouteId, EStyle Style, const FOnEventC
 		}
 		else
 		{
-			uint32 TimerId = *TimerIdPtr;
-			// Replace the placeholder metadata added if we received a BeginMetadataEvent first.
-			EditableTimingProfilerProvider.SetMetadata(TimerId, MoveTemp(Metadata));
-		}
-
-		break;
-	}
-	case RouteId_BeginMetadataEvent:
-	{
-		const uint32 ThreadId = Context.ThreadInfo.GetId();
-		FThreadState& ThreadState = GetThreadState(ThreadId);
-		if (ThreadState.bShouldIgnorePendingEvents)
-		{
-			break;
-		}
-
-		uint32 MetadataId = Context.EventData.GetValue<uint32>("MetadataId");
-		uint32 SpecId = Context.EventData.GetValue<uint32>("SpecId");
-		uint64 Cycle = EventData.GetValue<uint64>("Cycle");
-		double Timestamp = Context.EventTime.AsSeconds(Cycle);
-
-		// We don't know if the Metadata or the BeginMetadataEvent events arrive first so handle both cases.
-		uint32 TimerId = 0;
-		uint32* TimerIdPtr = MetadataIdToTimerIdMap.Find(MetadataId);
-		if (TimerIdPtr == nullptr)
-		{
-			// Add an empty placeholder metadata so we obtain a MetadataId to use as the TimerId.
-			TimerId = GetTimerId(SpecId);
-			TimerId = EditableTimingProfilerProvider.AddMetadata(TimerId, TArray<uint8>());
-			MetadataIdToTimerIdMap.Add(MetadataId, TimerId);
-		}
-		else
-		{
-			TimerId = *TimerIdPtr;
-		}
-
-		check(ThreadState.LastCycle <= Cycle);
-		check(ThreadState.LastPendingEventTime <= Timestamp);
-		ThreadState.LastPendingEventTime = Timestamp;
-
-		ThreadState.PendingEvents.Add({ Cycle, Timestamp, TimerId });
-		++ThreadState.MetadataEventsDepth;
-		break;
-	}
-	case RouteId_EndMetadataEvent:
-	{
-		uint32 ThreadId = Context.ThreadInfo.GetId();
-		FThreadState& ThreadState = GetThreadState(ThreadId);
-
-		if (ThreadState.bShouldIgnorePendingEvents)
-		{
-			break;
+			uint32 MetadataTimerId = *TimerIdPtr;
+			uint32 TimerId = GetTimerId(SpecId);
+			// Replace the placeholder metadata added if we received an event with this metadata first.
+			EditableTimingProfilerProvider.SetMetadata(MetadataTimerId, MoveTemp(Metadata), TimerId);
 		}
 
-		uint64 Cycle = EventData.GetValue<uint64>("Cycle");
-		double Timestamp = Context.EventTime.AsSeconds(Cycle);
-
-		if (ThreadState.MetadataEventsDepth > 0)
-		{
-			check(ThreadState.LastCycle <= Cycle);
-			check(ThreadState.LastPendingEventTime <= Timestamp);
-			ThreadState.LastPendingEventTime = Timestamp;
-
-			ThreadState.PendingEvents.Add({ ~Cycle, Timestamp, 0 });
-			--ThreadState.MetadataEventsDepth;
-		}
 		break;
 	}
 	}
@@ -765,6 +729,234 @@ void FCpuProfilerAnalyzer::ProcessBufferV2(const FEventTime& EventTime, FThreadS
 
 ////////////////////////////////////////////////////////////////////////////////////////////////////
 
+void FCpuProfilerAnalyzer::ProcessBufferV3(const FEventTime& EventTime, FThreadState& ThreadState, const uint8* BufferPtr, uint32 BufferSize)
+{
+	uint64 LastCycle = ThreadState.LastCycle;
+
+	CPUPROFILER_DEBUG_LOGF(TEXT("[%u] ProcessBuffer %llu (%.9f)\n"), ThreadState.ThreadId, LastCycle, EventTime.AsSeconds(LastCycle));
+
+	check(EventTime.GetTimestamp() == 0);
+	const uint64 BaseCycle = EventTime.AsCycle64();
+
+	int32 RemainingPending = ThreadState.PendingEvents.Num();
+	const FPendingEvent* PendingCursor = ThreadState.PendingEvents.GetData();
+
+	const uint8* BufferEnd = BufferPtr + BufferSize;
+	while (BufferPtr < BufferEnd)
+	{
+		uint64 DecodedCycle = FTraceAnalyzerUtils::Decode7bit(BufferPtr);
+		uint64 ActualCycle = (DecodedCycle >> 2);
+
+		// ActualCycle larger or equal to LastCycle means we have a new
+		// base value.
+		if (ActualCycle < LastCycle)
+		{
+			ActualCycle += LastCycle;
+		}
+
+		// If we late connect we will be joining the cycle stream mid-flow and
+		// will have missed out on it's base timestamp. Reconstruct it here.
+		if (ActualCycle < BaseCycle)
+		{
+			ActualCycle += BaseCycle;
+		}
+
+		// Dispatch pending events that are younger than the one we've just decoded.
+		DispatchPendingEvents(LastCycle, ActualCycle, ThreadState, PendingCursor, RemainingPending, (DecodedCycle & 1ull) != 0);
+
+		double ActualTime = EventTime.AsSeconds(ActualCycle);
+
+		if (DecodedCycle & 2ull)
+		{
+			constexpr uint32 CoroutineSpecId = (1u << 31u) - 1u;
+			constexpr uint32 CoroutineUnknownSpecId = (1u << 31u) - 2u;
+
+			if (DecodedCycle & 1ull)
+			{
+				uint64 CoroutineId = FTraceAnalyzerUtils::Decode7bit(BufferPtr);
+				uint32 TimerScopeDepth = IntCastChecked<uint32>(FTraceAnalyzerUtils::Decode7bit(BufferPtr));
+
+				// Begins a "CoroTask" scoped timer.
+				{
+					if (CoroutineTimerId == ~0u)
+					{
+						CoroutineTimerId = DefineNewTimerChecked(CoroutineSpecId, TEXT("Coroutine"));
+					}
+
+					TArray<uint8> CborData;
+					{
+						CborData.Reserve(256);
+						FMemoryWriter MemoryWriter(CborData, false, true);
+						FCborWriter CborWriter(&MemoryWriter, ECborEndianness::StandardCompliant);
+						CborWriter.WriteContainerStart(ECborCode::Map, 2); // 2 is the FieldCount
+						CborWriter.WriteValue("Id", 2);
+						CborWriter.WriteValue(CoroutineId);
+						CborWriter.WriteValue("C", 1); // continuation?
+						CborWriter.WriteValue(false);
+					}
+					uint32 MetadataTimerId = EditableTimingProfilerProvider.AddMetadata(CoroutineTimerId, MoveTemp(CborData));
+
+					FEventScopeState& ScopeState = ThreadState.ScopeStack.AddDefaulted_GetRef();
+					ScopeState.StartCycle = ActualCycle;
+					ScopeState.EventTypeId = MetadataTimerId;
+
+					CPUPROFILER_DEBUG_LOGF(TEXT("[%u] *B=%llu (%.9f)\n"), ThreadState.ThreadId, ActualCycle, ActualTime);
+					FTimingProfilerEvent Event;
+					Event.TimerIndex = MetadataTimerId;
+					ThreadState.Timeline->AppendBeginEvent(ActualTime, Event);
+					CPUPROFILER_DEBUG_BEGIN_EVENT(ActualTime, Event);
+				}
+
+				// Begins the CPU scoped timers (suspended in previous coroutine execution).
+				{
+					if (CoroutineUnknownTimerId == ~0u)
+					{
+						CoroutineUnknownTimerId = DefineNewTimerChecked(CoroutineUnknownSpecId, TEXT("<unknown>"));
+					}
+
+					//TODO: Restore the saved stack of CPU scoped timers for this CoroutineId.
+					for (uint32 i = 0; i < TimerScopeDepth; ++i)
+					{
+						FEventScopeState& ScopeState = ThreadState.ScopeStack.AddDefaulted_GetRef();
+						ScopeState.StartCycle = ActualCycle;
+						ScopeState.EventTypeId = CoroutineUnknownTimerId;
+
+						CPUPROFILER_DEBUG_LOGF(TEXT("[%u] +B=%llu (%.9f)\n"), ThreadState.ThreadId, ActualCycle, ActualTime);
+						FTimingProfilerEvent Event;
+						Event.TimerIndex = CoroutineUnknownTimerId;
+						ThreadState.Timeline->AppendBeginEvent(ActualTime, Event);
+						CPUPROFILER_DEBUG_BEGIN_EVENT(ActualTime, Event);
+					}
+				}
+			}
+			else
+			{
+				uint32 TimerScopeDepth = IntCastChecked<uint32>(FTraceAnalyzerUtils::Decode7bit(BufferPtr));
+
+				if (TimerScopeDepth != 0)
+				{
+					//TODO: Save current stack of CPU scoped timers (using id from metadata of CoroTask timer?)
+
+					// Ends (suspends) the CPU scoped timers.
+					for (uint32 i = 0; i < TimerScopeDepth; ++i)
+					{
+						// If we receive mismatched end events ignore them for now.
+						// This can happen for example because tracing connects to the store after events were traced. Those events can be lost.
+						if (ThreadState.ScopeStack.Num() > 0)
+						{
+							ThreadState.ScopeStack.Pop();
+							CPUPROFILER_DEBUG_LOGF(TEXT("[%u] +E=%llu (%.9f)\n"), ThreadState.ThreadId, ActualCycle, ActualTime);
+							ThreadState.Timeline->AppendEndEvent(ActualTime);
+							CPUPROFILER_DEBUG_END_EVENT(ActualTime);
+						}
+					}
+
+					// Update the "continuation" (suspended or destroyed) metadata flag.
+					if (ThreadState.ScopeStack.Num() > 0)
+					{
+						uint32 MetadataTimerId = ThreadState.ScopeStack.Top().EventTypeId;
+						TArrayView<uint8> Metadata = EditableTimingProfilerProvider.GetEditableMetadata(MetadataTimerId);
+						// Change the last byte in metadata to "true".
+						Metadata.GetData()[Metadata.Num() - 1] = (uint8)(ECborCode::Prim | ECborCode::True);
+					}
+				}
+
+				// Ends the "CoroTask" scoped timer.
+				{
+					// If we receive mismatched end events ignore them for now.
+					// This can happen for example because tracing connects to the store after events were traced. Those events can be lost.
+					if (ThreadState.ScopeStack.Num() > 0)
+					{
+						ThreadState.ScopeStack.Pop();
+						CPUPROFILER_DEBUG_LOGF(TEXT("[%u] *E=%llu (%.9f)\n"), ThreadState.ThreadId, ActualCycle, ActualTime);
+						ThreadState.Timeline->AppendEndEvent(ActualTime);
+						CPUPROFILER_DEBUG_END_EVENT(ActualTime);
+					}
+				}
+			}
+		}
+		else
+		{
+			if (DecodedCycle & 1ull)
+			{
+				uint32 TimerId = 0;
+				uint32 SpecId = IntCastChecked<uint32>(FTraceAnalyzerUtils::Decode7bit(BufferPtr));
+				if (SpecId & 1u) // The last bit is set if this is a metadata id.
+				{
+					// Get the actual Metadata Id.
+					uint32 MetadataId = SpecId >> 1;
+					uint32* TimerIdPtr = MetadataIdToTimerIdMap.Find(MetadataId);
+					if (TimerIdPtr == nullptr)
+					{
+						constexpr uint32 MetadataUnknownSpecId = (1u << 31u) - 3u;
+						if (MetadataUnknownTimerId == ~0u)
+						{
+							MetadataUnknownTimerId = DefineNewTimerChecked(MetadataUnknownSpecId, TEXT("<unknown>"));
+						}
+
+						// Add an empty placeholder metadata so we obtain a MetadataId to use as the TimerId. Will be replaced with the actual metadata if the metadata event arrives later.
+						TimerId = EditableTimingProfilerProvider.AddMetadata(MetadataUnknownTimerId, TArray<uint8>());
+						MetadataIdToTimerIdMap.Add(MetadataId, TimerId);
+					}
+					else
+					{
+						TimerId = *TimerIdPtr;
+					}
+				}
+				else
+				{
+					// Get the actual Spec Id.
+					SpecId = SpecId >> 1;
+					TimerId = GetTimerId(SpecId);
+				}
+
+
+				FEventScopeState& ScopeState = ThreadState.ScopeStack.AddDefaulted_GetRef();
+				ScopeState.StartCycle = ActualCycle;
+				ScopeState.EventTypeId = TimerId;
+
+				CPUPROFILER_DEBUG_LOGF(TEXT("[%u]  B=%llu (%.9f)\n"), ThreadState.ThreadId, ActualCycle, ActualTime);
+				FTimingProfilerEvent Event;
+				Event.TimerIndex = TimerId;
+				ThreadState.Timeline->AppendBeginEvent(ActualTime, Event);
+				CPUPROFILER_DEBUG_BEGIN_EVENT(ActualTime, Event);
+			}
+			else
+			{
+				// If we receive mismatched end events ignore them for now.
+				// This can happen for example because tracing connects to the store after events were traced. Those events can be lost.
+				if (ThreadState.ScopeStack.Num() > 0)
+				{
+					ThreadState.ScopeStack.Pop();
+					CPUPROFILER_DEBUG_LOGF(TEXT("[%u]  E=%llu (%.9f)\n"), ThreadState.ThreadId, ActualCycle, ActualTime);
+					ThreadState.Timeline->AppendEndEvent(ActualTime);
+					CPUPROFILER_DEBUG_END_EVENT(ActualTime);
+				}
+			}
+		}
+
+		check(ActualCycle > 0);
+		LastCycle = ActualCycle;
+	}
+	check(BufferPtr == BufferEnd);
+
+	if (RemainingPending == 0)
+	{
+		//CPUPROFILER_DEBUG_LOGF(TEXT("[%u] MetaEvents: %d added\n"), ThreadState.ThreadId, ThreadState.PendingEvents.Num());
+		ThreadState.PendingEvents.Reset();
+	}
+	else
+	{
+		const int32 NumEventsToRemove = ThreadState.PendingEvents.Num() - RemainingPending;
+		CPUPROFILER_DEBUG_LOGF(TEXT("[%u] MetaEvents: %d added, %d still pending\n"), ThreadState.ThreadId, NumEventsToRemove, RemainingPending);
+		ThreadState.PendingEvents.RemoveAt(0, NumEventsToRemove);
+	}
+
+	ThreadState.LastCycle = LastCycle;
+}
+
+////////////////////////////////////////////////////////////////////////////////////////////////////
+
 void FCpuProfilerAnalyzer::DispatchPendingEvents(
 	uint64& LastCycle,
 	uint64 CurrentCycle,
diff --git a/Engine/Source/Developer/TraceServices/Private/Analyzers/CpuProfilerTraceAnalysis.h b/Engine/Source/Developer/TraceServices/Private/Analyzers/CpuProfilerTraceAnalysis.h
--- a/Engine/Source/Developer/TraceServices/Private/Analyzers/CpuProfilerTraceAnalysis.h
+++ b/Engine/Source/Developer/TraceServices/Private/Analyzers/CpuProfilerTraceAnalysis.h
@@ -61,18 +61,18 @@ private:
 	{
 		RouteId_EventSpec,
 		RouteId_EndThread,
-		RouteId_EventBatchV2,
+		RouteId_EventBatchV3,
+		RouteId_EventBatchV2, // backward compatibility
 		RouteId_EventBatch, // backward compatibility
 		RouteId_EndCapture, // backward compatibility
 		RouteId_CpuScope,
 		RouteId_MetadataSpec,
 		RouteId_Metadata,
-		RouteId_BeginMetadataEvent,
-		RouteId_EndMetadataEvent,
 	};
 
 	void ProcessBuffer(const FEventTime& EventTime, FThreadState& ThreadState, const uint8* BufferPtr, uint32 BufferSize);
 	void ProcessBufferV2(const FEventTime& EventTime, FThreadState& ThreadState, const uint8* BufferPtr, uint32 BufferSize);
+	void ProcessBufferV3(const FEventTime& EventTime, FThreadState& ThreadState, const uint8* BufferPtr, uint32 BufferSize);
 	void DispatchPendingEvents(uint64& LastCycle, uint64 CurrentCycle, FThreadState& ThreadState, const FPendingEvent*& PendingCursor, int32& RemainingPending, bool bIsBeginEvent);
 	void DispatchRemainingPendingEvents(FThreadState& ThreadState);
 	void EndOpenEvents(FThreadState& ThreadState, double Timestamp);
@@ -95,6 +95,7 @@ private:
 
 	uint32 CoroutineTimerId = ~0;
 	uint32 CoroutineUnknownTimerId = ~0;
+	uint32 MetadataUnknownTimerId = ~0;
 	uint64 TotalEventSize = 0;
 	uint64 TotalScopeCount = 0;
 	double BytesPerScope = 0.0;
diff --git a/Engine/Source/Developer/TraceServices/Private/Model/TimingProfiler.cpp b/Engine/Source/Developer/TraceServices/Private/Model/TimingProfiler.cpp
--- a/Engine/Source/Developer/TraceServices/Private/Model/TimingProfiler.cpp
+++ b/Engine/Source/Developer/TraceServices/Private/Model/TimingProfiler.cpp
@@ -174,6 +174,26 @@ void FTimingProfilerProvider::SetMetadata(uint32 MetadataTimerId, TArray<uint8>&
 	MetadataToReplace.Payload = Metadata;
 }
 
+void FTimingProfilerProvider::SetMetadata(uint32 MetadataTimerId, TArray<uint8>&& Metadata, uint32 NewTimerId)
+{
+	Session.WriteAccessCheck();
+
+	if (int32(MetadataTimerId) >= 0)
+	{
+		return;
+	}
+
+	MetadataTimerId = ~MetadataTimerId;
+	if (MetadataTimerId >= uint32(Metadatas.Num()))
+	{
+		return;
+	}
+
+	FMetadata& MetadataToReplace = Metadatas[MetadataTimerId];
+	MetadataToReplace.TimerId = NewTimerId;
+	MetadataToReplace.Payload = Metadata;
+}
+
 TArrayView<uint8> FTimingProfilerProvider::GetEditableMetadata(uint32 MetadataTimerId)
 {
 	Session.WriteAccessCheck();
diff --git a/Engine/Source/Developer/TraceServices/Private/Model/TimingProfilerPrivate.h b/Engine/Source/Developer/TraceServices/Private/Model/TimingProfilerPrivate.h
--- a/Engine/Source/Developer/TraceServices/Private/Model/TimingProfilerPrivate.h
+++ b/Engine/Source/Developer/TraceServices/Private/Model/TimingProfilerPrivate.h
@@ -70,6 +70,7 @@ public:
 	virtual void SetTimerNameAndLocation(uint32 TimerId, FStringView Name, const TCHAR* File, uint32 Line) override;
 	virtual uint32 AddMetadata(uint32 OriginalTimerId, TArray<uint8>&& Metadata) override;
 	virtual void SetMetadata(uint32 MetadataTimerId, TArray<uint8>&& Metadata) override;
+	virtual void SetMetadata(uint32 MetadataTimerId, TArray<uint8>&& Metadata, uint32 NewTimerId) override;
 	virtual TArrayView<uint8> GetEditableMetadata(uint32 MetadataTimerId) override;
 
 	virtual uint32 AddMetadataSpec(FMetadataSpec&& Metadata) override;
diff --git a/Engine/Source/Developer/TraceServices/Public/TraceServices/Model/TimingProfiler.h b/Engine/Source/Developer/TraceServices/Public/TraceServices/Model/TimingProfiler.h
--- a/Engine/Source/Developer/TraceServices/Public/TraceServices/Model/TimingProfiler.h
+++ b/Engine/Source/Developer/TraceServices/Public/TraceServices/Model/TimingProfiler.h
@@ -374,6 +374,16 @@ public:
 	*
 	*/
 	virtual void SetMetadata(uint32 MetadataTimerId, TArray<uint8>&& Metadata) {}
+	
+	/*
+	* Sets metadata for the specified MetadataTimerId and replaces the OriginialTimerId with a new value. The MetadataTimerId must be a value returned by AddMetadata. The function is meant to be used to replace the metadata at a MetadataTimerId.
+	*
+	* @param MetadataTimerId	The identity of the metadata to replace.
+	* @param Metadata			The metadata.
+	* @param NewTimerId			The new TimerId. Represents the identity of the timer the metadata is attached to.
+	*
+	*/
+	virtual void SetMetadata(uint32 MetadataTimerId, TArray<uint8>&& Metadata, uint32 NewTimerId) {}
 
 	/*
 	* Gets metadata by id.
diff --git a/Engine/Source/Runtime/Core/Private/ProfilingDebugging/CpuProfilerTrace.cpp b/Engine/Source/Runtime/Core/Private/ProfilingDebugging/CpuProfilerTrace.cpp
--- a/Engine/Source/Runtime/Core/Private/ProfilingDebugging/CpuProfilerTrace.cpp
+++ b/Engine/Source/Runtime/Core/Private/ProfilingDebugging/CpuProfilerTrace.cpp
@@ -45,17 +45,7 @@ UE_TRACE_EVENT_BEGIN(CpuProfiler, Metadata, NoSync)
 	UE_TRACE_EVENT_FIELD(uint8[], Metadata)
 UE_TRACE_EVENT_END()
 
-UE_TRACE_EVENT_BEGIN(CpuProfiler, BeginMetadataEvent, NoSync)
-	UE_TRACE_EVENT_FIELD(uint64, Cycle)
-	UE_TRACE_EVENT_FIELD(uint32, MetadataId)
-	UE_TRACE_EVENT_FIELD(uint32, SpecId)
-UE_TRACE_EVENT_END()
-
-UE_TRACE_EVENT_BEGIN(CpuProfiler, EndMetadataEvent, NoSync)
-	UE_TRACE_EVENT_FIELD(uint64, Cycle)
-UE_TRACE_EVENT_END()
-
-UE_TRACE_EVENT_BEGIN(CpuProfiler, EventBatchV2, NoSync)
+UE_TRACE_EVENT_BEGIN(CpuProfiler, EventBatchV3, NoSync)
 	UE_TRACE_EVENT_FIELD(uint8[], Data)
 UE_TRACE_EVENT_END()
 
@@ -161,8 +151,8 @@ FCpuProfilerTraceInternal::FThreadBuffer* FCpuProfilerTraceInternal::CreateThrea
 
 void FCpuProfilerTraceInternal::FlushThreadBuffer(FThreadBuffer* InThreadBuffer)
 {
-	UE_TRACE_LOG(CpuProfiler, EventBatchV2, true)
-		<< EventBatchV2.Data(InThreadBuffer->Buffer, InThreadBuffer->BufferSize);
+	UE_TRACE_LOG(CpuProfiler, EventBatchV3, true)
+		<< EventBatchV3.Data(InThreadBuffer->Buffer, InThreadBuffer->BufferSize);
 	InThreadBuffer->BufferSize = 0;
 	InThreadBuffer->LastCycle = 0;
 }
@@ -175,13 +165,13 @@ void FCpuProfilerTraceInternal::FlushThreadBuffer(FThreadBuffer* InThreadBuffer)
 		ThreadBuffer = FCpuProfilerTraceInternal::CreateThreadBuffer(); \
 	} \
 
-#define CPUPROFILERTRACE_OUTPUTBEGINEVENT_EPILOGUE() \
+#define CPUPROFILERTRACE_OUTPUTBEGINEVENT_EPILOGUE(InSpecId) \
 	uint64 Cycle = FPlatformTime::Cycles64(); \
 	uint64 CycleDiff = Cycle - ThreadBuffer->LastCycle; \
 	ThreadBuffer->LastCycle = Cycle; \
 	uint8* BufferPtr = ThreadBuffer->Buffer + ThreadBuffer->BufferSize; \
 	FTraceUtils::Encode7bit((CycleDiff << 2) | 1ull, BufferPtr); \
-	FTraceUtils::Encode7bit(SpecId, BufferPtr); \
+	FTraceUtils::Encode7bit(InSpecId, BufferPtr); \
 	ThreadBuffer->BufferSize = (uint16)(BufferPtr - ThreadBuffer->Buffer); \
 	if (ThreadBuffer->BufferSize >= FCpuProfilerTraceInternal::FullBufferThreshold) \
 	{ \
@@ -253,7 +243,7 @@ void FCpuProfilerTrace::OutputBeginEvent(uint32 SpecId)
 	}
 
 	CPUPROFILERTRACE_OUTPUTBEGINEVENT_PROLOGUE();
-	CPUPROFILERTRACE_OUTPUTBEGINEVENT_EPILOGUE();
+	CPUPROFILERTRACE_OUTPUTBEGINEVENT_EPILOGUE(SpecId << 1 | 0);
 }
 
 void FCpuProfilerTrace::OutputBeginDynamicEvent(const ANSICHAR* Name, const ANSICHAR* File, uint32 Line)
@@ -275,7 +265,7 @@ void FCpuProfilerTrace::OutputBeginDynamicEvent(const ANSICHAR* Name, const ANSI
 
 	CPUPROFILERTRACE_OUTPUTBEGINEVENT_PROLOGUE();
 	uint32 SpecId = OutputDynamicEventType(Name, File, Line);
-	CPUPROFILERTRACE_OUTPUTBEGINEVENT_EPILOGUE();
+	CPUPROFILERTRACE_OUTPUTBEGINEVENT_EPILOGUE(SpecId << 1 | 0);
 }
 
 void FCpuProfilerTrace::OutputBeginDynamicEvent(const TCHAR* Name, const ANSICHAR* File, uint32 Line)
@@ -297,7 +287,7 @@ void FCpuProfilerTrace::OutputBeginDynamicEvent(const TCHAR* Name, const ANSICHA
 
 	CPUPROFILERTRACE_OUTPUTBEGINEVENT_PROLOGUE();
 	uint32 SpecId = OutputDynamicEventType(Name, File, Line);
-	CPUPROFILERTRACE_OUTPUTBEGINEVENT_EPILOGUE();
+	CPUPROFILERTRACE_OUTPUTBEGINEVENT_EPILOGUE(SpecId << 1 | 0);
 }
 
 void FCpuProfilerTrace::OutputBeginDynamicEvent(FNameParam Name, const ANSICHAR* File, uint32 Line)
@@ -319,7 +309,7 @@ void FCpuProfilerTrace::OutputBeginDynamicEvent(FNameParam Name, const ANSICHAR*
 
 	CPUPROFILERTRACE_OUTPUTBEGINEVENT_PROLOGUE();
 	uint32 SpecId = OutputDynamicEventType(Name, File, Line);
-	CPUPROFILERTRACE_OUTPUTBEGINEVENT_EPILOGUE();
+	CPUPROFILERTRACE_OUTPUTBEGINEVENT_EPILOGUE(SpecId << 1 | 0);
 }
 
 void FCpuProfilerTrace::OutputBeginDynamicEventWithId(FNameParam Id, const ANSICHAR* Name, const ANSICHAR* File, uint32 Line)
@@ -341,7 +331,7 @@ void FCpuProfilerTrace::OutputBeginDynamicEventWithId(FNameParam Id, const ANSIC
 
 	CPUPROFILERTRACE_OUTPUTBEGINEVENT_PROLOGUE();
 	uint32 SpecId = OutputDynamicEventTypeWithId(Id, Name, File, Line);
-	CPUPROFILERTRACE_OUTPUTBEGINEVENT_EPILOGUE();
+	CPUPROFILERTRACE_OUTPUTBEGINEVENT_EPILOGUE(SpecId << 1 | 0);
 }
 
 void FCpuProfilerTrace::OutputBeginDynamicEventWithId(FNameParam Id, const TCHAR* Name, const ANSICHAR* File, uint32 Line)
@@ -363,7 +353,7 @@ void FCpuProfilerTrace::OutputBeginDynamicEventWithId(FNameParam Id, const TCHAR
 
 	CPUPROFILERTRACE_OUTPUTBEGINEVENT_PROLOGUE();
 	uint32 SpecId = OutputDynamicEventTypeWithId(Id, Name, File, Line);
-	CPUPROFILERTRACE_OUTPUTBEGINEVENT_EPILOGUE();
+	CPUPROFILERTRACE_OUTPUTBEGINEVENT_EPILOGUE(SpecId << 1 | 0);
 }
 
 void FCpuProfilerTrace::OutputResumeEvent(uint64 SpecId, uint32& TimerScopeDepth)
@@ -416,6 +406,27 @@ void FCpuProfilerTrace::OutputSuspendEvent()
 	}
 }
 
+void FCpuProfilerTrace::OutputBeginEventWithMetadata(uint32 MetadataId)
+{
+	if (AutoRTFM::IsClosed())
+	{
+		AutoRTFM::Open([&]
+			{
+				OutputBeginEventWithMetadata(MetadataId);
+			});
+
+		AutoRTFM::PushOnAbortHandler(&FCpuProfilerTrace::OnAbortKey, []
+			{
+				OutputEndEventWithMetadata();
+			});
+
+		return;
+	}
+
+	CPUPROFILERTRACE_OUTPUTBEGINEVENT_PROLOGUE();
+	CPUPROFILERTRACE_OUTPUTBEGINEVENT_EPILOGUE(MetadataId << 1 | 1);
+}
+
 #undef CPUPROFILERTRACE_OUTPUTBEGINEVENT_PROLOGUE
 #undef CPUPROFILERTRACE_OUTPUTBEGINEVENT_EPILOGUE
 
@@ -544,37 +555,9 @@ uint32 FCpuProfilerTrace::OutputMetadata(uint32 SpecId, const TArray<uint8>& InM
 	return MetadataId;
 }
 
-void FCpuProfilerTrace::OutputBeginEventWithMetadata(uint32 SpecId, uint32 MetadataId)
-{
-	if (AutoRTFM::IsClosed())
-	{
-		AutoRTFM::Open([&]
-			{
-				OutputBeginEventWithMetadata(SpecId, MetadataId);
-			});
-
-		AutoRTFM::PushOnAbortHandler(&FCpuProfilerTrace::OnAbortKey, []
-			{
-				OutputEndEventWithMetadata();
-			});
-
-		return;
-	}
-
-	uint64 Cycle = FPlatformTime::Cycles64();
-
-	UE_TRACE_LOG(CpuProfiler, BeginMetadataEvent, CpuChannel)
-		<< BeginMetadataEvent.Cycle(Cycle)
-		<< BeginMetadataEvent.MetadataId(MetadataId)
-		<< BeginMetadataEvent.SpecId(SpecId);
-}
-
 void FCpuProfilerTrace::OutputEndEventWithMetadata()
 {
-	uint64 Cycle = FPlatformTime::Cycles64();
-
-	UE_TRACE_LOG(CpuProfiler, EndMetadataEvent, CpuChannel)
-		<< EndMetadataEvent.Cycle(Cycle);
+	OutputEndEvent();
 }
 
 uint32 FCpuProfilerTrace::OutputEventType(const ANSICHAR* Name, const ANSICHAR* File, uint32 Line)
diff --git a/Engine/Source/Runtime/Core/Public/ProfilingDebugging/CpuProfilerTrace.h b/Engine/Source/Runtime/Core/Public/ProfilingDebugging/CpuProfilerTrace.h
--- a/Engine/Source/Runtime/Core/Public/ProfilingDebugging/CpuProfilerTrace.h
+++ b/Engine/Source/Runtime/Core/Public/ProfilingDebugging/CpuProfilerTrace.h
@@ -104,10 +104,9 @@ struct FCpuProfilerTrace
 
 	/*
 	 * Output begin event marker with metadata for a given spec. Must always be matched with an end event.
-	 * @param SpecId - The event definition id.
 	 * @param MetadataId - The MetadataId returned from OutputMetadata.
 	 */
-	CORE_API static void OutputBeginEventWithMetadata(uint32 SpecId, uint32 MetadataId);
+	CORE_API static void OutputBeginEventWithMetadata(uint32 MetadataId);
 
 	/*
 	 * Output end event marker for an event traced with OutputBeginEventWithMetadata.
diff --git a/Engine/Source/Runtime/RHI/Public/RHIBreadcrumbs.h b/Engine/Source/Runtime/RHI/Public/RHIBreadcrumbs.h
--- a/Engine/Source/Runtime/RHI/Public/RHIBreadcrumbs.h
+++ b/Engine/Source/Runtime/RHI/Public/RHIBreadcrumbs.h
@@ -648,7 +648,7 @@ namespace UE::RHI::GPUProfiler
 		{
 			if (TraceCpuMetadataId > 0)
 			{
-				FCpuProfilerTrace::OutputBeginEventWithMetadata(TraceCpuSpecId, TraceCpuMetadataId);
+				FCpuProfilerTrace::OutputBeginEventWithMetadata(TraceCpuMetadataId);
 			}
 			else
 			{
