[Insights]
* CpuProfiler: Merged analysis code for EventBatchV2 and EventBatchV3, for easier maintenance.
* GpuProfiler: Added "Analysis completed" log message for the new GPU Profiler.

diff --git a/Engine/Source/Developer/TraceServices/Private/Analyzers/CpuProfilerTraceAnalysis.cpp b/Engine/Source/Developer/TraceServices/Private/Analyzers/CpuProfilerTraceAnalysis.cpp
--- a/Engine/Source/Developer/TraceServices/Private/Analyzers/CpuProfilerTraceAnalysis.cpp
+++ b/Engine/Source/Developer/TraceServices/Private/Analyzers/CpuProfilerTraceAnalysis.cpp
@@ -94,7 +94,8 @@ void FCpuProfilerAnalyzer::OnAnalysisEnd()
 		UE_LOG(LogTraceServices, Warning, TEXT("[CpuProfiler] Possible incorrect use of FCpuProfilerTrace::OutputEventType()!"));
 	}
 
-	UE_LOG(LogTraceServices, Log, TEXT("[CpuProfiler] Analysis completed (%d timers, %d spec ids, %llu scopes, %llu bytes, %.1f bytes/scope)."),
+	UE_LOG(LogTraceServices, Log, TEXT("[CpuProfiler] Analysis completed (%d threads, %d timers, %d spec ids, %llu scopes, %llu bytes, %.1f bytes/scope)."),
+		ThreadStatesMap.Num(),
 		ScopeNameToTimerIdMap.Num(),
 		SpecIdToTimerIdMap.Num(),
 		TotalScopeCount,
@@ -204,32 +205,8 @@ bool FCpuProfilerAnalyzer::OnEvent(uint16 RouteId, EStyle Style, const FOnEventC
 		break;
 	}
 
-	case RouteId_EventBatchV2:
-	{
-		const uint32 ThreadId = Context.ThreadInfo.GetId();
-		FThreadState& ThreadState = GetThreadState(ThreadId);
-
-		if (ThreadState.LastCycle == ~0ull)
-		{
-			// Ignore timing events received after EndThread.
-			break;
-		}
-
-		TArrayView<const uint8> DataView = Context.EventData.GetArrayView<uint8>("Data");
-		const uint32 BufferSize = DataView.Num();
-		const uint8* BufferPtr = DataView.GetData();
-
-		ProcessBufferV2(Context.EventTime, ThreadState, BufferPtr, BufferSize);
-		if (ThreadState.LastCycle != 0)
-		{
-			double Timestamp = Context.EventTime.AsSeconds(ThreadState.LastCycle);
-			Session.UpdateDurationSeconds(Timestamp);
-		}
-
-		TotalEventSize += BufferSize;
-		break;
-	}
 	case RouteId_EventBatchV3:
+	case RouteId_EventBatchV2: // backward compatibility
 	{
 		const uint32 ThreadId = Context.ThreadInfo.GetId();
 		FThreadState& ThreadState = GetThreadState(ThreadId);
@@ -244,7 +221,7 @@ bool FCpuProfilerAnalyzer::OnEvent(uint16 RouteId, EStyle Style, const FOnEventC
 		const uint32 BufferSize = DataView.Num();
 		const uint8* BufferPtr = DataView.GetData();
 
-		ProcessBufferV3(Context.EventTime, ThreadState, BufferPtr, BufferSize);
+		ProcessBufferV2(Context.EventTime, ThreadState, BufferPtr, BufferSize, RouteId == RouteId_EventBatchV3 ? 3 : 2);
 		if (ThreadState.LastCycle != 0)
 		{
 			double Timestamp = Context.EventTime.AsSeconds(ThreadState.LastCycle);
@@ -530,7 +507,7 @@ void FCpuProfilerAnalyzer::ProcessBuffer(const FEventTime& EventTime, FThreadSta
 
 ////////////////////////////////////////////////////////////////////////////////////////////////////
 
-void FCpuProfilerAnalyzer::ProcessBufferV2(const FEventTime& EventTime, FThreadState& ThreadState, const uint8* BufferPtr, uint32 BufferSize)
+void FCpuProfilerAnalyzer::ProcessBufferV2(const FEventTime& EventTime, FThreadState& ThreadState, const uint8* BufferPtr, uint32 BufferSize, int32 Version)
 {
 	uint64 LastCycle = ThreadState.LastCycle;
 
@@ -657,8 +634,11 @@ void FCpuProfilerAnalyzer::ProcessBufferV2(const FEventTime& EventTime, FThreadS
 					{
 						uint32 MetadataTimerId = ThreadState.ScopeStack.Top().EventTypeId;
 						TArrayView<uint8> Metadata = EditableTimingProfilerProvider.GetEditableMetadata(MetadataTimerId);
-						// Change the last byte in metadata to "true".
-						Metadata.GetData()[Metadata.Num() - 1] = (uint8)(ECborCode::Prim | ECborCode::True);
+						if (ensure(Metadata.Num() > 0))
+						{
+							// Change the last byte in metadata to "true".
+							Metadata.GetData()[Metadata.Num() - 1] = (uint8)(ECborCode::Prim | ECborCode::True);
+						}
 					}
 				}
 
@@ -681,236 +661,44 @@ void FCpuProfilerAnalyzer::ProcessBufferV2(const FEventTime& EventTime, FThreadS
 			if (DecodedCycle & 1ull)
 			{
 				uint32 SpecId = IntCastChecked<uint32>(FTraceAnalyzerUtils::Decode7bit(BufferPtr));
-				uint32 TimerId = GetTimerId(SpecId);
-
-				FEventScopeState& ScopeState = ThreadState.ScopeStack.AddDefaulted_GetRef();
-				ScopeState.StartCycle = ActualCycle;
-				ScopeState.EventTypeId = TimerId;
-
-				CPUPROFILER_DEBUG_LOGF(TEXT("[%u]  B=%llu (%.9f)\n"), ThreadState.ThreadId, ActualCycle, ActualTime);
-				FTimingProfilerEvent Event;
-				Event.TimerIndex = TimerId;
-				ThreadState.Timeline->AppendBeginEvent(ActualTime, Event);
-				CPUPROFILER_DEBUG_BEGIN_EVENT(ActualTime, Event);
-			}
-			else
-			{
-				// If we receive mismatched end events ignore them for now.
-				// This can happen for example because tracing connects to the store after events were traced. Those events can be lost.
-				if (ThreadState.ScopeStack.Num() > 0)
-				{
-					ThreadState.ScopeStack.Pop();
-					CPUPROFILER_DEBUG_LOGF(TEXT("[%u]  E=%llu (%.9f)\n"), ThreadState.ThreadId, ActualCycle, ActualTime);
-					ThreadState.Timeline->AppendEndEvent(ActualTime);
-					CPUPROFILER_DEBUG_END_EVENT(ActualTime);
-				}
-			}
-		}
-
-		check(ActualCycle > 0);
-		LastCycle = ActualCycle;
-	}
-	check(BufferPtr == BufferEnd);
-
-	if (RemainingPending == 0)
-	{
-		//CPUPROFILER_DEBUG_LOGF(TEXT("[%u] MetaEvents: %d added\n"), ThreadState.ThreadId, ThreadState.PendingEvents.Num());
-		ThreadState.PendingEvents.Reset();
-	}
-	else
-	{
-		const int32 NumEventsToRemove = ThreadState.PendingEvents.Num() - RemainingPending;
-		CPUPROFILER_DEBUG_LOGF(TEXT("[%u] MetaEvents: %d added, %d still pending\n"), ThreadState.ThreadId, NumEventsToRemove, RemainingPending);
-		ThreadState.PendingEvents.RemoveAt(0, NumEventsToRemove);
-	}
-
-	ThreadState.LastCycle = LastCycle;
-}
-
-////////////////////////////////////////////////////////////////////////////////////////////////////
-
-void FCpuProfilerAnalyzer::ProcessBufferV3(const FEventTime& EventTime, FThreadState& ThreadState, const uint8* BufferPtr, uint32 BufferSize)
-{
-	uint64 LastCycle = ThreadState.LastCycle;
-
-	CPUPROFILER_DEBUG_LOGF(TEXT("[%u] ProcessBuffer %llu (%.9f)\n"), ThreadState.ThreadId, LastCycle, EventTime.AsSeconds(LastCycle));
-
-	check(EventTime.GetTimestamp() == 0);
-	const uint64 BaseCycle = EventTime.AsCycle64();
-
-	int32 RemainingPending = ThreadState.PendingEvents.Num();
-	const FPendingEvent* PendingCursor = ThreadState.PendingEvents.GetData();
-
-	const uint8* BufferEnd = BufferPtr + BufferSize;
-	while (BufferPtr < BufferEnd)
-	{
-		uint64 DecodedCycle = FTraceAnalyzerUtils::Decode7bit(BufferPtr);
-		uint64 ActualCycle = (DecodedCycle >> 2);
-
-		// ActualCycle larger or equal to LastCycle means we have a new
-		// base value.
-		if (ActualCycle < LastCycle)
-		{
-			ActualCycle += LastCycle;
-		}
-
-		// If we late connect we will be joining the cycle stream mid-flow and
-		// will have missed out on it's base timestamp. Reconstruct it here.
-		if (ActualCycle < BaseCycle)
-		{
-			ActualCycle += BaseCycle;
-		}
-
-		// Dispatch pending events that are younger than the one we've just decoded.
-		DispatchPendingEvents(LastCycle, ActualCycle, ThreadState, PendingCursor, RemainingPending, (DecodedCycle & 1ull) != 0);
-
-		double ActualTime = EventTime.AsSeconds(ActualCycle);
-
-		if (DecodedCycle & 2ull)
-		{
-			constexpr uint32 CoroutineSpecId = (1u << 31u) - 1u;
-			constexpr uint32 CoroutineUnknownSpecId = (1u << 31u) - 2u;
 
-			if (DecodedCycle & 1ull)
-			{
-				uint64 CoroutineId = FTraceAnalyzerUtils::Decode7bit(BufferPtr);
-				uint32 TimerScopeDepth = IntCastChecked<uint32>(FTraceAnalyzerUtils::Decode7bit(BufferPtr));
-
-				// Begins a "CoroTask" scoped timer.
-				{
-					if (CoroutineTimerId == ~0u)
-					{
-						CoroutineTimerId = DefineNewTimerChecked(CoroutineSpecId, TEXT("Coroutine"));
-					}
-
-					TArray<uint8> CborData;
-					{
-						CborData.Reserve(256);
-						FMemoryWriter MemoryWriter(CborData, false, true);
-						FCborWriter CborWriter(&MemoryWriter, ECborEndianness::StandardCompliant);
-						CborWriter.WriteContainerStart(ECborCode::Map, 2); // 2 is the FieldCount
-						CborWriter.WriteValue("Id", 2);
-						CborWriter.WriteValue(CoroutineId);
-						CborWriter.WriteValue("C", 1); // continuation?
-						CborWriter.WriteValue(false);
-					}
-					uint32 MetadataTimerId = EditableTimingProfilerProvider.AddMetadata(CoroutineTimerId, MoveTemp(CborData));
-
-					FEventScopeState& ScopeState = ThreadState.ScopeStack.AddDefaulted_GetRef();
-					ScopeState.StartCycle = ActualCycle;
-					ScopeState.EventTypeId = MetadataTimerId;
-
-					CPUPROFILER_DEBUG_LOGF(TEXT("[%u] *B=%llu (%.9f)\n"), ThreadState.ThreadId, ActualCycle, ActualTime);
-					FTimingProfilerEvent Event;
-					Event.TimerIndex = MetadataTimerId;
-					ThreadState.Timeline->AppendBeginEvent(ActualTime, Event);
-					CPUPROFILER_DEBUG_BEGIN_EVENT(ActualTime, Event);
-				}
-
-				// Begins the CPU scoped timers (suspended in previous coroutine execution).
-				{
-					if (CoroutineUnknownTimerId == ~0u)
-					{
-						CoroutineUnknownTimerId = DefineNewTimerChecked(CoroutineUnknownSpecId, TEXT("<unknown>"));
-					}
-
-					//TODO: Restore the saved stack of CPU scoped timers for this CoroutineId.
-					for (uint32 i = 0; i < TimerScopeDepth; ++i)
-					{
-						FEventScopeState& ScopeState = ThreadState.ScopeStack.AddDefaulted_GetRef();
-						ScopeState.StartCycle = ActualCycle;
-						ScopeState.EventTypeId = CoroutineUnknownTimerId;
-
-						CPUPROFILER_DEBUG_LOGF(TEXT("[%u] +B=%llu (%.9f)\n"), ThreadState.ThreadId, ActualCycle, ActualTime);
-						FTimingProfilerEvent Event;
-						Event.TimerIndex = CoroutineUnknownTimerId;
-						ThreadState.Timeline->AppendBeginEvent(ActualTime, Event);
-						CPUPROFILER_DEBUG_BEGIN_EVENT(ActualTime, Event);
-					}
-				}
-			}
-			else
-			{
-				uint32 TimerScopeDepth = IntCastChecked<uint32>(FTraceAnalyzerUtils::Decode7bit(BufferPtr));
-
-				if (TimerScopeDepth != 0)
+				uint32 TimerId = 0;
+				if (Version == 3)
 				{
-					//TODO: Save current stack of CPU scoped timers (using id from metadata of CoroTask timer?)
-
-					// Ends (suspends) the CPU scoped timers.
-					for (uint32 i = 0; i < TimerScopeDepth; ++i)
+					if (SpecId & 1u) // The last bit is set if this is a metadata id.
 					{
-						// If we receive mismatched end events ignore them for now.
-						// This can happen for example because tracing connects to the store after events were traced. Those events can be lost.
-						if (ThreadState.ScopeStack.Num() > 0)
+						// Get the actual Metadata Id.
+						uint32 MetadataId = SpecId >> 1;
+						uint32* TimerIdPtr = MetadataIdToTimerIdMap.Find(MetadataId);
+						if (TimerIdPtr == nullptr)
 						{
-							ThreadState.ScopeStack.Pop();
-							CPUPROFILER_DEBUG_LOGF(TEXT("[%u] +E=%llu (%.9f)\n"), ThreadState.ThreadId, ActualCycle, ActualTime);
-							ThreadState.Timeline->AppendEndEvent(ActualTime);
-							CPUPROFILER_DEBUG_END_EVENT(ActualTime);
+							constexpr uint32 MetadataUnknownSpecId = (1u << 31u) - 3u;
+							if (MetadataUnknownTimerId == ~0u)
+							{
+								MetadataUnknownTimerId = DefineNewTimerChecked(MetadataUnknownSpecId, TEXT("<unknown>"));
+							}
+
+							// Add an empty placeholder metadata so we obtain a MetadataId to use as the TimerId. Will be replaced with the actual metadata if the metadata event arrives later.
+							TimerId = EditableTimingProfilerProvider.AddMetadata(MetadataUnknownTimerId, TArray<uint8>());
+							MetadataIdToTimerIdMap.Add(MetadataId, TimerId);
 						}
-					}
-
-					// Update the "continuation" (suspended or destroyed) metadata flag.
-					if (ThreadState.ScopeStack.Num() > 0)
-					{
-						uint32 MetadataTimerId = ThreadState.ScopeStack.Top().EventTypeId;
-						TArrayView<uint8> Metadata = EditableTimingProfilerProvider.GetEditableMetadata(MetadataTimerId);
-						// Change the last byte in metadata to "true".
-						Metadata.GetData()[Metadata.Num() - 1] = (uint8)(ECborCode::Prim | ECborCode::True);
-					}
-				}
-
-				// Ends the "CoroTask" scoped timer.
-				{
-					// If we receive mismatched end events ignore them for now.
-					// This can happen for example because tracing connects to the store after events were traced. Those events can be lost.
-					if (ThreadState.ScopeStack.Num() > 0)
-					{
-						ThreadState.ScopeStack.Pop();
-						CPUPROFILER_DEBUG_LOGF(TEXT("[%u] *E=%llu (%.9f)\n"), ThreadState.ThreadId, ActualCycle, ActualTime);
-						ThreadState.Timeline->AppendEndEvent(ActualTime);
-						CPUPROFILER_DEBUG_END_EVENT(ActualTime);
-					}
-				}
-			}
-		}
-		else
-		{
-			if (DecodedCycle & 1ull)
-			{
-				uint32 TimerId = 0;
-				uint32 SpecId = IntCastChecked<uint32>(FTraceAnalyzerUtils::Decode7bit(BufferPtr));
-				if (SpecId & 1u) // The last bit is set if this is a metadata id.
-				{
-					// Get the actual Metadata Id.
-					uint32 MetadataId = SpecId >> 1;
-					uint32* TimerIdPtr = MetadataIdToTimerIdMap.Find(MetadataId);
-					if (TimerIdPtr == nullptr)
-					{
-						constexpr uint32 MetadataUnknownSpecId = (1u << 31u) - 3u;
-						if (MetadataUnknownTimerId == ~0u)
+						else
 						{
-							MetadataUnknownTimerId = DefineNewTimerChecked(MetadataUnknownSpecId, TEXT("<unknown>"));
+							TimerId = *TimerIdPtr;
 						}
-
-						// Add an empty placeholder metadata so we obtain a MetadataId to use as the TimerId. Will be replaced with the actual metadata if the metadata event arrives later.
-						TimerId = EditableTimingProfilerProvider.AddMetadata(MetadataUnknownTimerId, TArray<uint8>());
-						MetadataIdToTimerIdMap.Add(MetadataId, TimerId);
 					}
 					else
 					{
-						TimerId = *TimerIdPtr;
+						// Get the actual Spec Id.
+						SpecId = SpecId >> 1;
+						TimerId = GetTimerId(SpecId);
 					}
 				}
 				else
 				{
-					// Get the actual Spec Id.
-					SpecId = SpecId >> 1;
 					TimerId = GetTimerId(SpecId);
 				}
 
-
 				FEventScopeState& ScopeState = ThreadState.ScopeStack.AddDefaulted_GetRef();
 				ScopeState.StartCycle = ActualCycle;
 				ScopeState.EventTypeId = TimerId;
@@ -1119,7 +907,7 @@ void FCpuProfilerAnalyzer::OnCpuScopeLeave(const FOnEventContext& Context)
 
 	uint64 Cycle = Context.EventTime.AsCycle64();
 	double Time = Context.EventTime.AsSeconds();
-	
+
 	check(ThreadState.LastCycle <= Cycle);
 	check(ThreadState.LastPendingEventTime <= Time);
 	ThreadState.LastPendingEventTime = Time;
diff --git a/Engine/Source/Developer/TraceServices/Private/Analyzers/CpuProfilerTraceAnalysis.h b/Engine/Source/Developer/TraceServices/Private/Analyzers/CpuProfilerTraceAnalysis.h
--- a/Engine/Source/Developer/TraceServices/Private/Analyzers/CpuProfilerTraceAnalysis.h
+++ b/Engine/Source/Developer/TraceServices/Private/Analyzers/CpuProfilerTraceAnalysis.h
@@ -71,8 +71,7 @@ private:
 	};
 
 	void ProcessBuffer(const FEventTime& EventTime, FThreadState& ThreadState, const uint8* BufferPtr, uint32 BufferSize);
-	void ProcessBufferV2(const FEventTime& EventTime, FThreadState& ThreadState, const uint8* BufferPtr, uint32 BufferSize);
-	void ProcessBufferV3(const FEventTime& EventTime, FThreadState& ThreadState, const uint8* BufferPtr, uint32 BufferSize);
+	void ProcessBufferV2(const FEventTime& EventTime, FThreadState& ThreadState, const uint8* BufferPtr, uint32 BufferSize, int32 Version);
 	void DispatchPendingEvents(uint64& LastCycle, uint64 CurrentCycle, FThreadState& ThreadState, const FPendingEvent*& PendingCursor, int32& RemainingPending, bool bIsBeginEvent);
 	void DispatchRemainingPendingEvents(FThreadState& ThreadState);
 	void EndOpenEvents(FThreadState& ThreadState, double Timestamp);
diff --git a/Engine/Source/Developer/TraceServices/Private/Analyzers/GpuProfilerTraceAnalysis.cpp b/Engine/Source/Developer/TraceServices/Private/Analyzers/GpuProfilerTraceAnalysis.cpp
--- a/Engine/Source/Developer/TraceServices/Private/Analyzers/GpuProfilerTraceAnalysis.cpp
+++ b/Engine/Source/Developer/TraceServices/Private/Analyzers/GpuProfilerTraceAnalysis.cpp
@@ -11,6 +11,8 @@
 #include "Serialization/MemoryReader.h"
 #include "Serialization/MemoryWriter.h"
 
+#define UE_INSIGHTS_MERGE_UNACCOUNTED_TIMERS 0
+
 namespace TraceServices
 {
 
@@ -48,7 +50,9 @@ void FGpuProfilerAnalyzer::OnAnalysisBegin(const FOnAnalysisContext& Context)
 
 void FGpuProfilerAnalyzer::OnAnalysisEnd()
 {
-	// Old Gpu Profiler errors.
+	//////////////////////////////////////////////////
+	// Old GPU Profiler
+
 	if (NumFramesWithErrors > 0)
 	{
 		UE_LOG(LogTraceServices, Error, TEXT("[GpuProfiler] Frames with errors: %u"), NumFramesWithErrors);
@@ -59,7 +63,9 @@ void FGpuProfilerAnalyzer::OnAnalysisEnd()
 		UE_LOG(LogTraceServices, Log, TEXT("[GpuProfiler] Analysis completed (%u frames, %d timers)."), NumFrames, EventTypeMap.Num());
 	}
 
-	// New Gpu Profiler errors.
+	//////////////////////////////////////////////////
+	// New GPU Profiler
+
 	if (ErrorData.NumInterleavedEvents > 0)
 	{
 		UE_LOG(LogTraceServices, Error, TEXT("[GpuProfiler] Number of interleaved events: %u, max error %f"), ErrorData.NumInterleavedEvents, ErrorData.InterleavedEventsMaxDelta);
@@ -79,6 +85,12 @@ void FGpuProfilerAnalyzer::OnAnalysisEnd()
 	{
 		UE_LOG(LogTraceServices, Error, TEXT("[GpuProfiler] Number of negative duration events: %u, max error %f"), ErrorData.NumNegativeDurationEvents, ErrorData.NegativeDurationEventsMaxDelta);
 	}
+
+	if (Queues.Num() > 0 || TimerMap.Num() > 0)
+	{
+		UE_LOG(LogTraceServices, Log, TEXT("[GpuProfiler] Analysis completed (%u queues, %d timers, %d breadcrumb specs, %d breadcrumb names)."),
+			Queues.Num(), TimerMap.Num(), BreadcrumbSpecMap.Num(), BreadcrumbMap.Num());
+	}
 }
 
 bool FGpuProfilerAnalyzer::OnEvent(uint16 RouteId, EStyle Style, const FOnEventContext& Context)
@@ -228,7 +240,7 @@ bool FGpuProfilerAnalyzer::OnEvent(uint16 RouteId, EStyle Style, const FOnEventC
 				}
 			}
 		}
-		
+
 		FAnalysisSessionEditScope _(Session);
 		uint32 TimerId = 0;
 		uint32* FoundSpecIdPtr = BreadcrumbSpecMap.Find(SpecId);
@@ -255,7 +267,7 @@ bool FGpuProfilerAnalyzer::OnEvent(uint16 RouteId, EStyle Style, const FOnEventC
 		const uint32 SpecId = EventData.GetValue<uint32>("SpecId");
 		const uint32 QueueId = EventData.GetValue<uint32>("QueueId");
 		const uint64 GPUTimestampTOP = EventData.GetValue<uint64>("GPUTimestampTOP");
-		
+
 		// Ignore events for which the timestamp could not be determined.
 		if (GPUTimestampTOP == 0)
 		{
@@ -270,7 +282,7 @@ bool FGpuProfilerAnalyzer::OnEvent(uint16 RouteId, EStyle Style, const FOnEventC
 		uint32 TimerId = 0;
 		uint32* TimerIdPtr;
 		TimerIdPtr = BreadcrumbSpecMap.Find(SpecId);
-		
+
 		{
 			FAnalysisSessionEditScope _(Session);
 			if (TimerIdPtr)
@@ -454,7 +466,20 @@ bool FGpuProfilerAnalyzer::OnEvent(uint16 RouteId, EStyle Style, const FOnEventC
 		if (!TimerIndexPtr)
 		{
 			FAnalysisSessionEditScope _(Session);
-			EventTypeMap.Add(EventType, TimingProfilerProvider.AddGpuTimer(FStringView(NameTChar.Get(), NameTChar.Length())));
+			uint32 TimerId;
+#if UE_INSIGHTS_MERGE_UNACCOUNTED_TIMERS
+			FStringView Unaccounted = TEXTVIEW("Unaccounted -");
+			if (NameTChar.Length() > Unaccounted.Len() &&
+				FCString::Strncmp(Unaccounted.GetData(), NameTChar.Get(), Unaccounted.Len()) == 0)
+			{
+				TimerId = TimingProfilerProvider.AddGpuTimer(TEXTVIEW("Unaccounted"));
+			}
+			else
+#endif // UE_INSIGHTS_MERGE_UNACCOUNTED_TIMERS
+			{
+				TimerId = TimingProfilerProvider.AddGpuTimer(FStringView(NameTChar.Get(), NameTChar.Length()));
+			}
+			EventTypeMap.Add(EventType, TimerId);
 		}
 		else
 		{
@@ -735,3 +760,5 @@ void FGpuProfilerAnalyzer::InitCountersDesc(FQueue& FoundQueue, uint8 Gpu, uint8
 }
 
 } // namespace TraceServices
+
+#undef UE_INSIGHTS_MERGE_UNACCOUNTED_TIMERS
