Unreal Insights - Net Traces: discard net trace events before the init event, as these are just events which were in the cache buffer from just before the trace started

diff --git a/Engine/Source/Developer/TraceServices/Private/Analyzers/NetTraceAnalyzer.cpp b/Engine/Source/Developer/TraceServices/Private/Analyzers/NetTraceAnalyzer.cpp
--- a/Engine/Source/Developer/TraceServices/Private/Analyzers/NetTraceAnalyzer.cpp
+++ b/Engine/Source/Developer/TraceServices/Private/Analyzers/NetTraceAnalyzer.cpp
@@ -93,10 +93,11 @@ bool FNetTraceAnalyzer::OnEvent(uint16 RouteId, EStyle Style, const FOnEventCont
 
 	FAnalysisSessionEditScope _(Session);
 
-	// check that we always get the InitEvent before processing any other events
-	if (!ensure(RouteId == RouteId_InitEvent || NetTraceVersion > 0))
+	// sometimes when a trace starts, the cache buffer contains some events which get written before the init event can be written,
+	// so we need to discard those
+	if ((NetTraceVersion == 0) && (RouteId != RouteId_InitEvent))
 	{
-		return false;
+		return true;
 	}
 
 	const auto& EventData = Context.EventData;
