Fix for truncated symbols with Insights when the symbol contains spaces.

diff --git a/Engine/Source/Developer/TraceServices/Private/Model/PsymResolver.cpp b/Engine/Source/Developer/TraceServices/Private/Model/PsymResolver.cpp
--- a/Engine/Source/Developer/TraceServices/Private/Model/PsymResolver.cpp
+++ b/Engine/Source/Developer/TraceServices/Private/Model/PsymResolver.cpp
@@ -360,6 +360,8 @@ void FPsymResolver::LoadModuleSymbols(const FModule* Module, const TCHAR* Path,
 				}
 				else if (Command == TEXT("FUNC"))
 				{
+					constexpr int32 MaxSymbolLength = 255;
+
 					FString Address = NextToken(Line, Index);
 					if (Address == "m")
 					{
@@ -368,7 +370,13 @@ void FPsymResolver::LoadModuleSymbols(const FModule* Module, const TCHAR* Path,
 
 					FString Size = NextToken(Line, Index);
 					FString ParamSize = NextToken(Line, Index);
-					FString Name = NextToken(Line, Index);
+
+					// find first non-whitespace char
+					while (Index < Line.Len() && IsWhiteSpace(Line[Index]))
+					{
+						Index++;
+					}
+					FString Name = Index == Line.Len() ? FString(TEXT("")) : FString(Line.RightChop(Index).Left(MaxSymbolLength));
 
  					++Entry.Module->Stats.Discovered;
 					PsymSymbols.Add(FPsymSymbol(FParse::HexNumber64(*Address) + BaseAddress, FParse::HexNumber(*Size), StringAllocator.Store(*Name)));
