[Insights] Added more rules to ignore top symbols in callstacks.

diff --git a/Engine/Source/Developer/TraceServices/Private/Model/ModuleProvider.cpp b/Engine/Source/Developer/TraceServices/Private/Model/ModuleProvider.cpp
--- a/Engine/Source/Developer/TraceServices/Private/Model/ModuleProvider.cpp
+++ b/Engine/Source/Developer/TraceServices/Private/Model/ModuleProvider.cpp
@@ -495,10 +495,14 @@ FResolvedSymbolFilter::FResolvedSymbolFilter()
 	IgnoreSymbolsByFunctionName.Add(TEXTVIEW("FCachedOSPageAllocator::"));
 	IgnoreSymbolsByFunctionName.Add(TEXTVIEW("FMallocBinned"));
 	IgnoreSymbolsByFunctionName.Add(TEXTVIEW("FD3D12Adapter::TraceMemoryAllocation"));
+	IgnoreSymbolsByFunctionName.Add(TEXTVIEW("AnsiMalloc"));
+	IgnoreSymbolsByFunctionName.Add(TEXTVIEW("AnsiRealloc"));
+	IgnoreSymbolsByFunctionName.Add(TEXTVIEW("AnsiFree"));
 
 	IgnoreSymbolsByFilePath.Add(TEXTVIEW("/Containers/"));
 	IgnoreSymbolsByFilePath.Add(TEXTVIEW("/ConcurrentLinearAllocator"));
 	IgnoreSymbolsByFilePath.Add(TEXTVIEW("/D3D12PoolAllocator"));
+	IgnoreSymbolsByFilePath.Add(TEXTVIEW("/AutoRTFM/"));
 }
 
 ////////////////////////////////////////////////////////////////////////////////////////////////////
