Fixed switched comments for enumeration functions in Log

diff --git a/Engine/Source/Developer/TraceServices/Public/TraceServices/Model/Log.h b/Engine/Source/Developer/TraceServices/Public/TraceServices/Model/Log.h
--- a/Engine/Source/Developer/TraceServices/Public/TraceServices/Model/Log.h
+++ b/Engine/Source/Developer/TraceServices/Public/TraceServices/Model/Log.h
@@ -66,7 +66,7 @@ public:
 	virtual void EnumerateMessagesByIndex(uint64 StartIndex, uint64 EndIndex, TFunctionRef<void(const FLogMessageInfo&)> Callback) const = 0;
 
 	/**
-	 * Enumerates messages with timestamp in the specified time interval [EndTime, StartTime].
+	 * Enumerates messages with timestamp in the specified time interval [StartTime, EndTime].
 	 *
 	 * @param StartTime The inclusive start time of the interval.
 	 * @param EndTime The inclusive end time of the interval.
@@ -77,7 +77,7 @@ public:
 	virtual void EnumerateMessages(double StartTime, double EndTime, TFunctionRef<void(const FLogMessageInfo&)> Callback) const = 0;
 
 	/**
-	 * Enumerates messages in reverse with timestamp in the specified time interval [StartTime, EndTime] until either it finds the last message
+	 * Enumerates messages in reverse with timestamp in the specified time interval [EndTime, StartTime] until either it finds the last message
 	 * or is stopped on callback.
 	 *
 	 * @param StartTime The inclusive start time of the interval.
