[Insights] Removed deprecated IAnalysisSession::GetLog() API.

diff --git a/Engine/Source/Developer/TraceServices/Private/AnalysisServicePrivate.h b/Engine/Source/Developer/TraceServices/Private/AnalysisServicePrivate.h
--- a/Engine/Source/Developer/TraceServices/Private/AnalysisServicePrivate.h
+++ b/Engine/Source/Developer/TraceServices/Private/AnalysisServicePrivate.h
@@ -59,7 +59,6 @@ public:
 	virtual void AddMetadata(FName InName, double InValue) override;
 	virtual void AddMetadata(FName InName, FString InValue) override;
 
-	virtual FMessageLog* GetLog() const override {return nullptr;};
 	virtual uint32 GetNumPendingMessages() const override;
 	virtual TArray<FAnalysisMessage> DrainPendingMessages() override;
 
diff --git a/Engine/Source/Developer/TraceServices/Public/TraceServices/Model/AnalysisSession.h b/Engine/Source/Developer/TraceServices/Public/TraceServices/Model/AnalysisSession.h
--- a/Engine/Source/Developer/TraceServices/Public/TraceServices/Model/AnalysisSession.h
+++ b/Engine/Source/Developer/TraceServices/Public/TraceServices/Model/AnalysisSession.h
@@ -9,7 +9,6 @@
 #include "Templates/SharedPointer.h"
 
 class FName;
-class FMessageLog;
 
 namespace UE {
 namespace Trace {
@@ -96,8 +95,6 @@ public:
 	virtual void AddMetadata(FName InName, double InValue) = 0;
 	virtual void AddMetadata(FName InName, FString InValue) = 0;
 
-	UE_DEPRECATED(5.3, "No longer used.")
-	virtual FMessageLog* GetLog() const = 0;
 	/** Get the number of pending messages. Can be called without lock */
 	virtual uint32 GetNumPendingMessages() const = 0;
 	/** Moves pending messages to the return array. */
