[Insights]
* Counters: Fixed crash (check) when analysis detects a stats counter declared multiple times.
* Counters: Implemented better detection for counters declared multiple times. Added log warnings/errors if a counter is re-declared with different properties.

diff --git a/Engine/Source/Developer/TraceServices/Private/Analyzers/StatsTraceAnalysis.cpp b/Engine/Source/Developer/TraceServices/Private/Analyzers/StatsTraceAnalysis.cpp
--- a/Engine/Source/Developer/TraceServices/Private/Analyzers/StatsTraceAnalysis.cpp
+++ b/Engine/Source/Developer/TraceServices/Private/Analyzers/StatsTraceAnalysis.cpp
@@ -48,6 +48,7 @@ bool FStatsAnalyzer::OnEvent(uint16 RouteId, EStyle Style, const FOnEventContext
 	{
 		case RouteId_Spec:
 		{
+			bool bIsAlreadyDeclared = false;
 			uint32 StatId = EventData.GetValue<uint32>("Id");
 			IEditableCounter* EditableCounter = EditableCountersMap.FindRef(StatId);
 			if (!EditableCounter)
@@ -55,6 +56,13 @@ bool FStatsAnalyzer::OnEvent(uint16 RouteId, EStyle Style, const FOnEventContext
 				EditableCounter = EditableCounterProvider.CreateEditableCounter();
 				EditableCountersMap.Add(StatId, EditableCounter);
 			}
+			else
+			{
+				bIsAlreadyDeclared = true;
+				UE_LOG(LogTraceServices, Warning, TEXT("Stats counter %u already declared."), StatId);
+			}
+
+			const ICounter* Counter = EditableCounterProvider.GetCounter(EditableCounter);
 
 			FString Name;
 			FString Description;
@@ -75,45 +83,94 @@ bool FStatsAnalyzer::OnEvent(uint16 RouteId, EStyle Style, const FOnEventContext
 				UE_LOG(LogTraceServices, Warning, TEXT("Invalid counter name for Stats counter %u."), StatId);
 				Name = FString::Printf(TEXT("<noname stats counter %u>"), StatId);
 			}
-			EditableCounter->SetName(Session.StoreString(*Name));
+			//if (!Name.IsEmpty())
+			{
+				if (!bIsAlreadyDeclared)
+				{
+					EditableCounter->SetName(Session.StoreString(*Name));
+				}
+				else if (Counter && FCString::Strcmp(*Name, Counter->GetName()) != 0)
+				{
+					UE_LOG(LogTraceServices, Error, TEXT("Stats counter %u redeclared with different name."), StatId);
+				}
+			}
 
 			if (!Group.IsEmpty())
 			{
-				EditableCounter->SetGroup(Session.StoreString(*Group));
+				if (!bIsAlreadyDeclared)
+				{
+					EditableCounter->SetGroup(Session.StoreString(*Group));
+				}
+				else if (Counter && FCString::Strcmp(*Group, Counter->GetGroup()) != 0)
+				{
+					UE_LOG(LogTraceServices, Error, TEXT("Stats counter %u redeclared with different group name."), StatId);
+				}
+			}
+
+			if (!Description.IsEmpty())
+			{
+				if (!bIsAlreadyDeclared)
+				{
+					EditableCounter->SetDescription(Session.StoreString(*Description));
+				}
+				else if (Counter && FCString::Strcmp(*Description, Counter->GetDescription()) != 0)
+				{
+					UE_LOG(LogTraceServices, Error, TEXT("Stats counter %u redeclared with different description."), StatId);
+				}
 			}
-			EditableCounter->SetDescription(Session.StoreString(*Description));
 
 			bool bIsFloatingPoint = EventData.GetValue<bool>("IsFloatingPoint");
-			EditableCounter->SetIsFloatingPoint(bIsFloatingPoint);
-			
+			if (!bIsAlreadyDeclared)
+			{
+				EditableCounter->SetIsFloatingPoint(bIsFloatingPoint);
+			}
+			else if (Counter && bIsFloatingPoint != Counter->IsFloatingPoint())
+			{
+				UE_LOG(LogTraceServices, Error, TEXT("Stats counter %u redeclared with different IsFloatingPoint flag."), StatId);
+			}
+
 			bool bIsResetEveryFrame = EventData.GetValue<bool>("ShouldClearEveryFrame");
-			EditableCounter->SetIsResetEveryFrame(bIsResetEveryFrame);
-			if (bIsResetEveryFrame)
+			if (!bIsAlreadyDeclared)
 			{
-				if (bIsFloatingPoint)
+				EditableCounter->SetIsResetEveryFrame(bIsResetEveryFrame);
+				if (bIsResetEveryFrame)
 				{
-					IEditableCounter* ResetEveryFrameEditableCounter = FloatResetEveryFrameCountersMap.FindRef(StatId);
-					if (!ResetEveryFrameEditableCounter)
+					if (bIsFloatingPoint)
 					{
-						FloatResetEveryFrameCountersMap.Add(StatId, EditableCounter);
+						IEditableCounter* ResetEveryFrameEditableCounter = FloatResetEveryFrameCountersMap.FindRef(StatId);
+						if (!ResetEveryFrameEditableCounter)
+						{
+							FloatResetEveryFrameCountersMap.Add(StatId, EditableCounter);
+						}
 					}
-				}
-				else
-				{
-					IEditableCounter* ResetEveryFrameEditableCounter = Int64ResetEveryFrameCountersMap.FindRef(StatId);
-					if (!ResetEveryFrameEditableCounter)
+					else
 					{
-						Int64ResetEveryFrameCountersMap.Add(StatId, EditableCounter);
+						IEditableCounter* ResetEveryFrameEditableCounter = Int64ResetEveryFrameCountersMap.FindRef(StatId);
+						if (!ResetEveryFrameEditableCounter)
+						{
+							Int64ResetEveryFrameCountersMap.Add(StatId, EditableCounter);
+						}
 					}
 				}
 			}
+			else if (Counter && bIsResetEveryFrame != Counter->IsResetEveryFrame())
+			{
+				UE_LOG(LogTraceServices, Error, TEXT("Stats counter %u redeclared with different IsResetEveryFrame flag."), StatId);
+			}
 
 			ECounterDisplayHint DisplayHint = CounterDisplayHint_None;
 			if (EventData.GetValue<bool>("IsMemory"))
 			{
 				DisplayHint = CounterDisplayHint_Memory;
 			}
-			EditableCounter->SetDisplayHint(DisplayHint);
+			if (!bIsAlreadyDeclared)
+			{
+				EditableCounter->SetDisplayHint(DisplayHint);
+			}
+			else if (Counter && DisplayHint != Counter->GetDisplayHint())
+			{
+				UE_LOG(LogTraceServices, Error, TEXT("Stats counter %u redeclared with different display hint."), StatId);
+			}
 			break;
 		}
 
diff --git a/Engine/Source/Developer/TraceServices/Private/Model/Counters.cpp b/Engine/Source/Developer/TraceServices/Private/Model/Counters.cpp
--- a/Engine/Source/Developer/TraceServices/Private/Model/Counters.cpp
+++ b/Engine/Source/Developer/TraceServices/Private/Model/Counters.cpp
@@ -17,10 +17,14 @@ FCounter::FCounter(ILinearAllocator& Allocator, const TArray64<double>& InFrameS
 
 void FCounter::SetIsFloatingPoint(bool bInIsFloatingPoint)
 {
-	check(!ModCount);
 	bIsFloatingPoint = bInIsFloatingPoint;
 }
 
+void FCounter::SetIsResetEveryFrame(bool bInIsResetEveryFrame)
+{
+	bIsResetEveryFrame = bInIsResetEveryFrame;
+}
+
 template<typename CounterType, typename EnumerationType>
 static void EnumerateCounterValuesInternal(const TCounterData<CounterType>& CounterData, const TArray64<double>& FrameStartTimes, bool bResetEveryFrame, double IntervalStart, double IntervalEnd, bool bIncludeExternalBounds, TFunctionRef<void(double, EnumerationType)> Callback)
 {
@@ -172,7 +176,6 @@ void FCounter::AddValue(double Time, int64 Value)
 	{
 		IntCounterData.InsertOp(Time, ECounterOpType::Add, Value);
 	}
-	++ModCount;
 }
 
 void FCounter::AddValue(double Time, double Value)
@@ -185,7 +188,6 @@ void FCounter::AddValue(double Time, double Value)
 	{
 		IntCounterData.InsertOp(Time, ECounterOpType::Add, int64(Value));
 	}
-	++ModCount;
 }
 
 void FCounter::SetValue(double Time, int64 Value)
@@ -198,7 +200,6 @@ void FCounter::SetValue(double Time, int64 Value)
 	{
 		IntCounterData.InsertOp(Time, ECounterOpType::Set, Value);
 	}
-	++ModCount;
 }
 
 void FCounter::SetValue(double Time, double Value)
@@ -211,7 +212,6 @@ void FCounter::SetValue(double Time, double Value)
 	{
 		IntCounterData.InsertOp(Time, ECounterOpType::Set, int64(Value));
 	}
-	++ModCount;
 }
 
 FCounterProvider::FCounterProvider(IAnalysisSession& InSession, IFrameProvider& InFrameProvider)
diff --git a/Engine/Source/Developer/TraceServices/Private/Model/CountersPrivate.h b/Engine/Source/Developer/TraceServices/Private/Model/CountersPrivate.h
--- a/Engine/Source/Developer/TraceServices/Private/Model/CountersPrivate.h
+++ b/Engine/Source/Developer/TraceServices/Private/Model/CountersPrivate.h
@@ -210,22 +210,26 @@ class FCounter
 {
 public:
 	FCounter(ILinearAllocator& Allocator, const TArray64<double>& FrameStartTimes);
+
+	// ICounter
 	virtual const TCHAR* GetName() const override { return Name; }
-	virtual void SetName(const TCHAR* InName) override { Name = InName; }
 	virtual const TCHAR* GetGroup() const override { return Group; }
-	virtual void SetGroup(const TCHAR* InGroup) override { Group = InGroup; }
 	virtual const TCHAR* GetDescription() const override { return Description; }
-	virtual void SetDescription(const TCHAR* InDescription) override { Description = InDescription; }
 	virtual bool IsFloatingPoint() const override { return bIsFloatingPoint; }
-	virtual void SetIsFloatingPoint(bool bInIsFloatingPoint) override;
 	virtual bool IsResetEveryFrame() const override { return bIsResetEveryFrame; }
-	virtual void SetIsResetEveryFrame(bool bInIsResetEveryFrame) override { bIsResetEveryFrame = bInIsResetEveryFrame; }
 	virtual ECounterDisplayHint GetDisplayHint() const { return DisplayHint; }
-	virtual void SetDisplayHint(ECounterDisplayHint InDisplayHint) override { DisplayHint = InDisplayHint; }
 	virtual void EnumerateValues(double IntervalStart, double IntervalEnd, bool bIncludeExternalBounds, TFunctionRef<void(double, int64)> Callback) const override;
 	virtual void EnumerateFloatValues(double IntervalStart, double IntervalEnd, bool bIncludeExternalBounds, TFunctionRef<void(double, double)> Callback) const override;
 	virtual void EnumerateOps(double IntervalStart, double IntervalEnd, bool bIncludeExternalBounds, TFunctionRef<void(double, ECounterOpType, int64)> Callback) const override;
 	virtual void EnumerateFloatOps(double IntervalStart, double IntervalEnd, bool bIncludeExternalBounds, TFunctionRef<void(double, ECounterOpType, double)> Callback) const override;
+
+	// IEditableCounter
+	virtual void SetName(const TCHAR* InName) override { Name = InName; }
+	virtual void SetGroup(const TCHAR* InGroup) override { Group = InGroup; }
+	virtual void SetDescription(const TCHAR* InDescription) override { Description = InDescription; }
+	virtual void SetIsFloatingPoint(bool bInIsFloatingPoint) override;
+	virtual void SetIsResetEveryFrame(bool bInIsResetEveryFrame) override;
+	virtual void SetDisplayHint(ECounterDisplayHint InDisplayHint) override { DisplayHint = InDisplayHint; }
 	virtual void AddValue(double Time, int64 Value) override;
 	virtual void AddValue(double Time, double Value) override;
 	virtual void SetValue(double Time, int64 Value) override;
@@ -238,7 +242,6 @@ private:
 	const TCHAR* Name = nullptr;
 	const TCHAR* Group = nullptr;
 	const TCHAR* Description = nullptr;
-	uint64 ModCount = 0;
 	ECounterDisplayHint DisplayHint = CounterDisplayHint_None;
 	bool bIsFloatingPoint = false;
 	bool bIsResetEveryFrame = false;
