[Insights]
* Memory Insights: Filtered out the symbols from "NotLoaded" modules. Also added possiblity to filter by module name.
* Memory Insights: Filtered out the symbols from "/HAL/FMemory.inl" source path.

diff --git a/Engine/Source/Developer/TraceServices/Private/Model/ModuleProvider.cpp b/Engine/Source/Developer/TraceServices/Private/Model/ModuleProvider.cpp
--- a/Engine/Source/Developer/TraceServices/Private/Model/ModuleProvider.cpp
+++ b/Engine/Source/Developer/TraceServices/Private/Model/ModuleProvider.cpp
@@ -51,6 +51,7 @@ public:
 	virtual void Update(FResolvedSymbol& InSymbol) const override;
 
 private:
+	TArray<FStringView> IgnoreSymbolsByModuleName;
 	TArray<FStringView> IgnoreSymbolsByFunctionName;
 	TArray<FStringView> IgnoreSymbolsByFilePath;
 };
@@ -479,6 +480,11 @@ uint32 TModuleProvider<SymbolResolverType>::GetNumCachedSymbolsFromModule(uint64
 
 FResolvedSymbolFilter::FResolvedSymbolFilter()
 {
+	// Filter by module name (prefix)
+	//IgnoreSymbolsByModuleName.Add(TEXTVIEW("kernel32.dll"));
+	//IgnoreSymbolsByModuleName.Add(TEXTVIEW("ntdll.dll"));
+
+	// Filter by function name (prefix)
 	IgnoreSymbolsByFunctionName.Add(TEXTVIEW("FMemory::"));
 	IgnoreSymbolsByFunctionName.Add(TEXTVIEW("FMallocWrapper::"));
 	IgnoreSymbolsByFunctionName.Add(TEXTVIEW("FMallocPoisonProxy::"));
@@ -499,10 +505,12 @@ FResolvedSymbolFilter::FResolvedSymbolFilter()
 	IgnoreSymbolsByFunctionName.Add(TEXTVIEW("AnsiRealloc"));
 	IgnoreSymbolsByFunctionName.Add(TEXTVIEW("AnsiFree"));
 
+	// Filter by file path (substring)
 	IgnoreSymbolsByFilePath.Add(TEXTVIEW("/Containers/"));
 	IgnoreSymbolsByFilePath.Add(TEXTVIEW("/ConcurrentLinearAllocator"));
 	IgnoreSymbolsByFilePath.Add(TEXTVIEW("/D3D12PoolAllocator"));
 	IgnoreSymbolsByFilePath.Add(TEXTVIEW("/AutoRTFM/"));
+	IgnoreSymbolsByFilePath.Add(TEXTVIEW("/HAL/FMemory.inl"));
 }
 
 ////////////////////////////////////////////////////////////////////////////////////////////////////
@@ -517,6 +525,24 @@ void FResolvedSymbolFilter::Update(FResolvedSymbol& InSymbol) const
 {
 	bool bIsFiltered = false;
 
+	if (InSymbol.GetResult() == ESymbolQueryResult::NotLoaded)
+	{
+		bIsFiltered = true;
+	}
+
+	if (!bIsFiltered && InSymbol.Module)
+	{
+		// Ignore symbols by module name prefix.
+		for (const FStringView& Prefix : IgnoreSymbolsByModuleName)
+		{
+			if (FCString::Strnicmp(InSymbol.Module, Prefix.GetData(), Prefix.Len()) == 0)
+			{
+				bIsFiltered = true;
+				break;
+			}
+		}
+	}
+
 	if (!bIsFiltered && InSymbol.Name)
 	{
 		// Ignore symbols by function name prefix.
