[Insights] TraceServices: Added GetReadProvider() to IEditableThreadProvider interface.

diff --git a/Engine/Source/Developer/TraceServices/Private/Model/ThreadsPrivate.h b/Engine/Source/Developer/TraceServices/Private/Model/ThreadsPrivate.h
--- a/Engine/Source/Developer/TraceServices/Private/Model/ThreadsPrivate.h
+++ b/Engine/Source/Developer/TraceServices/Private/Model/ThreadsPrivate.h
@@ -19,10 +19,15 @@ public:
 	explicit FThreadProvider(IAnalysisSession& Session);
 	virtual ~FThreadProvider();
 
-	void AddGameThread(uint32 Id);
+	// IEditableThreadProvider
 	virtual void AddThread(uint32 Id, const TCHAR* Name, EThreadPriority Priority) override;
+	virtual const IThreadProvider* GetReadProvider() const  override { return this; }
+
+	void AddGameThread(uint32 Id);
 	void SetThreadPriority(uint32 Id, EThreadPriority Priority);
 	void SetThreadGroup(uint32 Id, const TCHAR* GroupName);
+
+	// IThreadProvider
 	virtual uint64 GetModCount() const override { return ModCount; }
 	virtual void EnumerateThreads(TFunctionRef<void(const FThreadInfo&)> Callback) const override;
 	virtual const TCHAR* GetThreadName(uint32 ThreadId) const override;
diff --git a/Engine/Source/Developer/TraceServices/Public/TraceServices/Model/Threads.h b/Engine/Source/Developer/TraceServices/Public/TraceServices/Model/Threads.h
--- a/Engine/Source/Developer/TraceServices/Public/TraceServices/Model/Threads.h
+++ b/Engine/Source/Developer/TraceServices/Public/TraceServices/Model/Threads.h
@@ -45,6 +45,12 @@ public:
 	* @param Priority	The system priority level of the thread, if available.
 	*/
 	virtual void AddThread(uint32 Id, const TCHAR* Name, EThreadPriority Priority) = 0;
+
+	/*
+	* Gets the read provider.
+	* @returns The read only provider or nullptr if not available.
+	*/
+	virtual const IThreadProvider* GetReadProvider() const { return nullptr; }
 };
 
 TRACESERVICES_API FName GetThreadProviderName();
