[Insights]
* MemoryInsights: Fixed alloc/free graph tracks now showing when trace is recorded with -trace=memory or -trace=memalloc.
* MemoryInsights: Changed the "Swap Memory Graph" and "Page In/Out Event Count" tracks to be created only when the trace has "swap op" events.
* TraceServices/AllocationsProvider: Simplified the code re timeline enumeration.
* TraceServices/AllocationsProvider: Added HasAllocationEvents() and HasSwapOpEvents() API.
* TraceServices/AllocationsProvider: Added more stats (ex: regarding number of traced events processed) displayed when analysis is completed.

diff --git a/Engine/Source/Developer/TraceInsights/Private/Insights/MemoryProfiler/Tracks/AllocationsGraphSeries.cpp b/Engine/Source/Developer/TraceInsights/Private/Insights/MemoryProfiler/Tracks/AllocationsGraphSeries.cpp
--- a/Engine/Source/Developer/TraceInsights/Private/Insights/MemoryProfiler/Tracks/AllocationsGraphSeries.cpp
+++ b/Engine/Source/Developer/TraceInsights/Private/Insights/MemoryProfiler/Tracks/AllocationsGraphSeries.cpp
@@ -259,46 +259,48 @@ void FAllocationsGraphSeries::PreUpdate(FGraphTrack& GraphTrack, const FTimingTr
 				}
 			};
 
+			using ETimelineU32 = TraceServices::IAllocationsProvider::ETimelineU32;
+			using ETimelineU64 = TraceServices::IAllocationsProvider::ETimelineU64;
 			switch (Timeline)
 			{
 			case ETimeline::MinTotalMem:
-				AllocationsProvider->EnumerateMinTotalAllocatedMemoryTimeline(StartIndex, EndIndex, Callback64);
+				AllocationsProvider->EnumerateTimeline(ETimelineU64::MinTotalAllocatedMemory, StartIndex, EndIndex, Callback64);
 				break;
 			case ETimeline::MaxTotalMem:
-				AllocationsProvider->EnumerateMaxTotalAllocatedMemoryTimeline(StartIndex, EndIndex, Callback64);
+				AllocationsProvider->EnumerateTimeline(ETimelineU64::MaxTotalAllocatedMemory, StartIndex, EndIndex, Callback64);
 				break;
 			case ETimeline::MinLiveAllocs:
-				AllocationsProvider->EnumerateMinLiveAllocationsTimeline(StartIndex, EndIndex, Callback32);
+				AllocationsProvider->EnumerateTimeline(ETimelineU32::MinLiveAllocations, StartIndex, EndIndex, Callback32);
 				break;
 			case ETimeline::MaxLiveAllocs:
-				AllocationsProvider->EnumerateMaxLiveAllocationsTimeline(StartIndex, EndIndex, Callback32);
+				AllocationsProvider->EnumerateTimeline(ETimelineU32::MaxLiveAllocations, StartIndex, EndIndex, Callback32);
 				break;
 			case ETimeline::MinSwapMem:
-				AllocationsProvider->EnumerateMinTotalSwapMemoryTimeline(StartIndex, EndIndex, Callback64);
+				AllocationsProvider->EnumerateTimeline(ETimelineU64::MinTotalSwapMemory, StartIndex, EndIndex, Callback64);
 				break;
 			case ETimeline::MaxSwapMem:
-				AllocationsProvider->EnumerateMaxTotalSwapMemoryTimeline(StartIndex, EndIndex, Callback64);
+				AllocationsProvider->EnumerateTimeline(ETimelineU64::MaxTotalSwapMemory, StartIndex, EndIndex, Callback64);
 				break;
 			case ETimeline::MinCompressedSwapMem:
-				AllocationsProvider->EnumerateMinTotalCompressedSwapMemoryTimeline(StartIndex, EndIndex, Callback64);
+				AllocationsProvider->EnumerateTimeline(ETimelineU64::MinTotalCompressedSwapMemory, StartIndex, EndIndex, Callback64);
 				break;
 			case ETimeline::MaxCompressedSwapMem:
-				AllocationsProvider->EnumerateMaxTotalCompressedSwapMemoryTimeline(StartIndex, EndIndex, Callback64);
+				AllocationsProvider->EnumerateTimeline(ETimelineU64::MaxTotalCompressedSwapMemory, StartIndex, EndIndex, Callback64);
 				break;
 			case ETimeline::AllocEvents:
-				AllocationsProvider->EnumerateAllocEventsTimeline(StartIndex, EndIndex, Callback32);
+				AllocationsProvider->EnumerateTimeline(ETimelineU32::AllocEvents, StartIndex, EndIndex, Callback32);
 				break;
 			case ETimeline::FreeEvents:
-				AllocationsProvider->EnumerateFreeEventsTimeline(StartIndex, EndIndex, Callback32);
+				AllocationsProvider->EnumerateTimeline(ETimelineU32::FreeEvents, StartIndex, EndIndex, Callback32);
 				break;
 			case ETimeline::PageInEvents:
-				AllocationsProvider->EnumeratePageInEventsTimeline(StartIndex, EndIndex, Callback32);
+				AllocationsProvider->EnumerateTimeline(ETimelineU32::PageInEvents, StartIndex, EndIndex, Callback32);
 				break;
 			case ETimeline::PageOutEvents:
-				AllocationsProvider->EnumeratePageOutEventsTimeline(StartIndex, EndIndex, Callback32);
+				AllocationsProvider->EnumerateTimeline(ETimelineU32::PageOutEvents, StartIndex, EndIndex, Callback32);
 				break;
 			case ETimeline::SwapFreeEvents:
-				AllocationsProvider->EnumerateSwapFreeEventsTimeline(StartIndex, EndIndex, Callback32);
+				AllocationsProvider->EnumerateTimeline(ETimelineU32::SwapFreeEvents, StartIndex, EndIndex, Callback32);
 				break;
 			}
 
@@ -367,46 +369,48 @@ void FAllocationsGraphSeries::Update(FGraphTrack& GraphTrack, const FTimingTrack
 				Builder.AddEvent(Time, Duration, -static_cast<double>(Value));
 			};
 
+			using ETimelineU32 = TraceServices::IAllocationsProvider::ETimelineU32;
+			using ETimelineU64 = TraceServices::IAllocationsProvider::ETimelineU64;
 			switch (Timeline)
 			{
 			case ETimeline::MinTotalMem:
-				AllocationsProvider->EnumerateMinTotalAllocatedMemoryTimeline(StartIndex, EndIndex, Callback64);
+				AllocationsProvider->EnumerateTimeline(ETimelineU64::MinTotalAllocatedMemory, StartIndex, EndIndex, Callback64);
 				break;
 			case ETimeline::MaxTotalMem:
-				AllocationsProvider->EnumerateMaxTotalAllocatedMemoryTimeline(StartIndex, EndIndex, Callback64);
+				AllocationsProvider->EnumerateTimeline(ETimelineU64::MaxTotalAllocatedMemory, StartIndex, EndIndex, Callback64);
 				break;
 			case ETimeline::MinLiveAllocs:
-				AllocationsProvider->EnumerateMinLiveAllocationsTimeline(StartIndex, EndIndex, Callback32);
+				AllocationsProvider->EnumerateTimeline(ETimelineU32::MinLiveAllocations, StartIndex, EndIndex, Callback32);
 				break;
 			case ETimeline::MaxLiveAllocs:
-				AllocationsProvider->EnumerateMaxLiveAllocationsTimeline(StartIndex, EndIndex, Callback32);
+				AllocationsProvider->EnumerateTimeline(ETimelineU32::MaxLiveAllocations, StartIndex, EndIndex, Callback32);
 				break;
 			case ETimeline::MinSwapMem:
-				AllocationsProvider->EnumerateMinTotalSwapMemoryTimeline(StartIndex, EndIndex, Callback64);
+				AllocationsProvider->EnumerateTimeline(ETimelineU64::MinTotalSwapMemory, StartIndex, EndIndex, Callback64);
 				break;
 			case ETimeline::MaxSwapMem:
-				AllocationsProvider->EnumerateMaxTotalSwapMemoryTimeline(StartIndex, EndIndex, Callback64);
+				AllocationsProvider->EnumerateTimeline(ETimelineU64::MaxTotalSwapMemory, StartIndex, EndIndex, Callback64);
 				break;
 			case ETimeline::MinCompressedSwapMem:
-				AllocationsProvider->EnumerateMinTotalCompressedSwapMemoryTimeline(StartIndex, EndIndex, Callback64);
+				AllocationsProvider->EnumerateTimeline(ETimelineU64::MinTotalCompressedSwapMemory, StartIndex, EndIndex, Callback64);
 				break;
 			case ETimeline::MaxCompressedSwapMem:
-				AllocationsProvider->EnumerateMaxTotalCompressedSwapMemoryTimeline(StartIndex, EndIndex, Callback64);
+				AllocationsProvider->EnumerateTimeline(ETimelineU64::MaxTotalCompressedSwapMemory, StartIndex, EndIndex, Callback64);
 				break;
 			case ETimeline::AllocEvents:
-				AllocationsProvider->EnumerateAllocEventsTimeline(StartIndex, EndIndex, Callback32);
+				AllocationsProvider->EnumerateTimeline(ETimelineU32::AllocEvents, StartIndex, EndIndex, Callback32);
 				break;
 			case ETimeline::FreeEvents:
-				AllocationsProvider->EnumerateFreeEventsTimeline(StartIndex, EndIndex, Callback32Negative);
+				AllocationsProvider->EnumerateTimeline(ETimelineU32::FreeEvents, StartIndex, EndIndex, Callback32Negative);
 				break;
 			case ETimeline::PageInEvents:
-				AllocationsProvider->EnumeratePageInEventsTimeline(StartIndex, EndIndex, Callback32);
+				AllocationsProvider->EnumerateTimeline(ETimelineU32::PageInEvents, StartIndex, EndIndex, Callback32);
 				break;
 			case ETimeline::PageOutEvents:
-				AllocationsProvider->EnumeratePageOutEventsTimeline(StartIndex, EndIndex, Callback32Negative);
+				AllocationsProvider->EnumerateTimeline(ETimelineU32::PageOutEvents, StartIndex, EndIndex, Callback32Negative);
 				break;
 			case ETimeline::SwapFreeEvents:
-				AllocationsProvider->EnumerateSwapFreeEventsTimeline(StartIndex, EndIndex, Callback32);
+				AllocationsProvider->EnumerateTimeline(ETimelineU32::SwapFreeEvents, StartIndex, EndIndex, Callback32);
 				break;
 			}
 		}
diff --git a/Engine/Source/Developer/TraceInsights/Private/Insights/MemoryProfiler/ViewModels/MemorySharedState.cpp b/Engine/Source/Developer/TraceInsights/Private/Insights/MemoryProfiler/ViewModels/MemorySharedState.cpp
--- a/Engine/Source/Developer/TraceInsights/Private/Insights/MemoryProfiler/ViewModels/MemorySharedState.cpp
+++ b/Engine/Source/Developer/TraceInsights/Private/Insights/MemoryProfiler/ViewModels/MemorySharedState.cpp
@@ -139,7 +139,7 @@ void FMemorySharedState::ResetTimingView()
 		TagPtr->RemoveAllTracks();
 	}
 
-	bShowHideAllMemoryTracks = false;
+	bShowHideAllMemoryTracks = true;
 
 	CreatedDefaultTracks.Reset();
 	LastTagCountForDefaultTracks = 0;
@@ -181,7 +181,9 @@ void FMemorySharedState::Tick(Timing::ITimingViewSession& InSession, const Trace
 		return;
 	}
 
-	bool bIsDetailedMemoryTrace = false;
+	bool bHasAllocationEvents = false;
+	bool bHasSwapOpEvents = false;
+
 	TSharedPtr<const TraceServices::IAnalysisSession> Session = FInsightsManager::Get()->GetSession();
 	if (Session.IsValid())
 	{
@@ -189,7 +191,8 @@ void FMemorySharedState::Tick(Timing::ITimingViewSession& InSession, const Trace
 		if (AllocationsProvider)
 		{
 			TraceServices::FProviderReadScopeLock ProviderReadScope(*AllocationsProvider);
-			bIsDetailedMemoryTrace = AllocationsProvider->IsInitialized() && AllocationsProvider->GetTimelineNumPoints() > 0;
+			bHasAllocationEvents = AllocationsProvider->IsInitialized() && AllocationsProvider->HasAllocationEvents();
+			bHasSwapOpEvents = AllocationsProvider->IsInitialized() && AllocationsProvider->HasSwapOpEvents();
 		}
 	}
 
@@ -216,7 +219,7 @@ void FMemorySharedState::Tick(Timing::ITimingViewSession& InSession, const Trace
 		TimingView->InvalidateScrollableTracksOrder();
 	}
 
-	if (!LiveAllocsGraphTrack.IsValid() && bIsDetailedMemoryTrace)
+	if (!LiveAllocsGraphTrack.IsValid() && bHasAllocationEvents)
 	{
 		LiveAllocsGraphTrack = CreateMemoryGraphTrack();
 		check(LiveAllocsGraphTrack);
@@ -238,7 +241,7 @@ void FMemorySharedState::Tick(Timing::ITimingViewSession& InSession, const Trace
 		TimingView->InvalidateScrollableTracksOrder();
 	}
 
-	if (!AllocFreeGraphTrack.IsValid() && bIsDetailedMemoryTrace)
+	if (!AllocFreeGraphTrack.IsValid() && bHasAllocationEvents)
 	{
 		AllocFreeGraphTrack = CreateMemoryGraphTrack();
 		check(AllocFreeGraphTrack);
@@ -260,7 +263,7 @@ void FMemorySharedState::Tick(Timing::ITimingViewSession& InSession, const Trace
 		TimingView->InvalidateScrollableTracksOrder();
 	}
 
-	if (!SwapMemoryGraphTrack.IsValid() && bIsDetailedMemoryTrace)
+	if (!SwapMemoryGraphTrack.IsValid() && bHasSwapOpEvents)
 	{
 		SwapMemoryGraphTrack = CreateMemoryGraphTrack();
 		check(SwapMemoryGraphTrack);
@@ -274,7 +277,7 @@ void FMemorySharedState::Tick(Timing::ITimingViewSession& InSession, const Trace
 		SwapMemoryGraphTrack->AddTimelineSeries(FAllocationsGraphSeries::ETimeline::MinCompressedSwapMem);
 		SwapMemoryGraphTrack->SetLabelUnit(EGraphTrackLabelUnit::MiB, 1);
 
-		SwapMemoryGraphTrack->SetVisibilityFlag(false); // don't show swaps track by default
+		SwapMemoryGraphTrack->SetVisibilityFlag(bShowHideAllMemoryTracks);
 
 		SwapMemoryGraphTrack->SetAvailableTrackHeight(EMemoryTrackHeightMode::Small, 50.0f);
 		SwapMemoryGraphTrack->SetAvailableTrackHeight(EMemoryTrackHeightMode::Medium, 100.0f);
@@ -284,7 +287,7 @@ void FMemorySharedState::Tick(Timing::ITimingViewSession& InSession, const Trace
 		TimingView->InvalidateScrollableTracksOrder();
 	}
 
-	if (!PageSwapGraphTrack.IsValid() && bIsDetailedMemoryTrace)
+	if (!PageSwapGraphTrack.IsValid() && bHasSwapOpEvents)
 	{
 		PageSwapGraphTrack = CreateMemoryGraphTrack();
 		check(PageSwapGraphTrack);
@@ -297,7 +300,7 @@ void FMemorySharedState::Tick(Timing::ITimingViewSession& InSession, const Trace
 		PageSwapGraphTrack->AddTimelineSeries(FAllocationsGraphSeries::ETimeline::PageOutEvents);
 		PageSwapGraphTrack->AddTimelineSeries(FAllocationsGraphSeries::ETimeline::SwapFreeEvents);
 
-		PageSwapGraphTrack->SetVisibilityFlag(false); // don't show swap tracks by default
+		PageSwapGraphTrack->SetVisibilityFlag(bShowHideAllMemoryTracks);
 
 		PageSwapGraphTrack->SetAvailableTrackHeight(EMemoryTrackHeightMode::Small, 50.0f);
 		PageSwapGraphTrack->SetAvailableTrackHeight(EMemoryTrackHeightMode::Medium, 100.0f);
diff --git a/Engine/Source/Developer/TraceServices/Private/Model/AllocationsProvider.cpp b/Engine/Source/Developer/TraceServices/Private/Model/AllocationsProvider.cpp
--- a/Engine/Source/Developer/TraceServices/Private/Model/AllocationsProvider.cpp
+++ b/Engine/Source/Developer/TraceServices/Private/Model/AllocationsProvider.cpp
@@ -1372,20 +1372,19 @@ FAllocationsProvider::FAllocationsProvider(IAnalysisSession& InSession, FMetadat
 	, MetadataProvider(InMetadataProvider)
 	, TagTracker(InSession)
 	, Timeline(Session.GetLinearAllocator(), 1024)
-	, MinTotalAllocatedMemoryTimeline(Session.GetLinearAllocator(), 1024)
-	, MaxTotalAllocatedMemoryTimeline(Session.GetLinearAllocator(), 1024)
-	, MinLiveAllocationsTimeline(Session.GetLinearAllocator(), 1024)
-	, MaxLiveAllocationsTimeline(Session.GetLinearAllocator(), 1024)
-	, MinTotalSwapMemoryTimeline(Session.GetLinearAllocator(), 1024)
-	, MaxTotalSwapMemoryTimeline(Session.GetLinearAllocator(), 1024)
-	, MinTotalCompressedSwapMemoryTimeline(Session.GetLinearAllocator(), 1024)
-	, MaxTotalCompressedSwapMemoryTimeline(Session.GetLinearAllocator(), 1024)
-	, AllocEventsTimeline(Session.GetLinearAllocator(), 1024)
-	, FreeEventsTimeline(Session.GetLinearAllocator(), 1024)
-	, PageInEventsTimeline(Session.GetLinearAllocator(), 1024)
-	, PageOutEventsTimeline(Session.GetLinearAllocator(), 1024)
-	, SwapFreeEventsTimeline(Session.GetLinearAllocator(), 1024)
 {
+	// Allocate the timelines
+	for (int32 TimelineIndex = 0; TimelineIndex < NumTimelinesU64; ++TimelineIndex)
+	{
+		ValueTimelinesU64[TimelineIndex] = new TPagedArray<uint64>(Session.GetLinearAllocator(), 1024);
+		SampleValuesU64[TimelineIndex] = 0;
+	}
+	for (int32 TimelineIndex = 0; TimelineIndex < NumTimelinesU32; ++TimelineIndex)
+	{
+		ValueTimelinesU32[TimelineIndex] = new TPagedArray<uint32>(Session.GetLinearAllocator(), 1024);
+		SampleValuesU32[TimelineIndex] = 0;
+	}
+
 	// Initial number of heap spec locations. Actual number of heap specs is unlimited (uint32).
 	HeapSpecs.AddDefaulted(FMath::Max(1024u, MaxRootHeaps));
 }
@@ -1427,6 +1426,24 @@ FAllocationsProvider::~FAllocationsProvider()
 		}
 	}
 	HeapSpecs.Reset();
+
+	// Delete the timelines.
+	for (int32 TimelineIndex = 0; TimelineIndex < NumTimelinesU64; ++TimelineIndex)
+	{
+		if (ValueTimelinesU64[TimelineIndex])
+		{
+			delete ValueTimelinesU64[TimelineIndex];
+			ValueTimelinesU64[TimelineIndex] = nullptr;
+		}
+	}
+	for (int32 TimelineIndex = 0; TimelineIndex < NumTimelinesU32; ++TimelineIndex)
+	{
+		if (ValueTimelinesU32[TimelineIndex])
+		{
+			delete ValueTimelinesU32[TimelineIndex];
+			ValueTimelinesU32[TimelineIndex] = nullptr;
+		}
+	}
 }
 
 ////////////////////////////////////////////////////////////////////////////////////////////////////
@@ -1466,6 +1483,7 @@ void FAllocationsProvider::EditAlloc(uint32 ThreadId, double Time, uint32 Callst
 		return;
 	}
 
+	++AllocEvents;
 	AllocInternal(ThreadId, Time, CallstackId, Address, Size, Alignment, RootHeapId);
 }
 
@@ -1591,9 +1609,9 @@ FAllocationItem* FAllocationsProvider::AllocInternal(uint32 ThreadId, double Tim
 		// Update stats for the current timeline sample.
 		TotalAllocatedMemory += Size;
 		++TotalLiveAllocations;
-		SampleMaxTotalAllocatedMemory = FMath::Max(SampleMaxTotalAllocatedMemory, TotalAllocatedMemory);
-		SampleMaxLiveAllocations = FMath::Max(SampleMaxLiveAllocations, TotalLiveAllocations);
-		++SampleAllocEvents;
+		SetMaxSampleU64(ETimelineU64::MaxTotalAllocatedMemory, TotalAllocatedMemory);
+		SetMaxSampleU32(ETimelineU32::MaxLiveAllocations, TotalLiveAllocations);
+		IncSampleU32(ETimelineU32::AllocEvents);
 	}
 
 	++AllocCount;
@@ -1621,6 +1639,7 @@ void FAllocationsProvider::EditFree(uint32 ThreadId, double Time, uint32 Callsta
 		return;
 	}
 
+	++FreeEvents;
 	FreeInternal(ThreadId, Time, CallstackId, Address, RootHeapId);
 }
 
@@ -1720,10 +1739,10 @@ FAllocationItem* FAllocationsProvider::FreeInternal(uint32 ThreadId, double Time
 		{
 			TotalAllocatedMemory -= OldSize;
 			--TotalLiveAllocations;
-			SampleMinTotalAllocatedMemory = FMath::Min(SampleMinTotalAllocatedMemory, TotalAllocatedMemory);
-			SampleMinLiveAllocations = FMath::Min(SampleMinLiveAllocations, TotalLiveAllocations);
+			SetMinSampleU64(ETimelineU64::MinTotalAllocatedMemory, TotalAllocatedMemory);
+			SetMinSampleU32(ETimelineU32::MinLiveAllocations, TotalLiveAllocations);
 		}
-		++SampleFreeEvents;
+		IncSampleU32(ETimelineU32::FreeEvents);
 	}
 	else
 	{
@@ -1783,6 +1802,8 @@ void FAllocationsProvider::EditSwapOp(uint32 ThreadId, double Time, uint64 Unmas
 		return;
 	}
 
+	++SwapOpEvents;
+
 	HeapId RootHeapId = EMemoryTraceRootHeap::SystemMemory;
 
 	if (!IsValidRootHeap(RootHeapId))
@@ -1814,17 +1835,17 @@ void FAllocationsProvider::EditSwapOp(uint32 ThreadId, double Time, uint64 Unmas
 	{
 	case EMemoryTraceSwapOperation::PageOut:
 	{
-		++SamplePageOutEvents;
+		IncSampleU32(ETimelineU32::PageOutEvents);
 		break;
 	}
 	case EMemoryTraceSwapOperation::PageIn:
 	{
-		++SamplePageInEvents;
+		IncSampleU32(ETimelineU32::PageInEvents);
 		break;
 	}
 	case EMemoryTraceSwapOperation::FreeInSwap:
 	{
-		++SampleSwapFreeEvents;
+		IncSampleU32(ETimelineU32::SwapFreeEvents);
 		break;
 	}
 	default:
@@ -1839,8 +1860,8 @@ void FAllocationsProvider::EditSwapOp(uint32 ThreadId, double Time, uint64 Unmas
 		{
 			TotalSwapMemory += PlatformPageSize;
 			TotalCompressedSwapMemory += CompressedPageSize;
-			SampleMaxSwapMemory = FMath::Max(SampleMaxSwapMemory, TotalSwapMemory);
-			SampleMaxCompressedSwapMemory = FMath::Max(SampleMaxCompressedSwapMemory, TotalCompressedSwapMemory);
+			SetMaxSampleU64(ETimelineU64::MaxTotalSwapMemory, TotalSwapMemory);
+			SetMaxSampleU64(ETimelineU64::MaxTotalCompressedSwapMemory, TotalCompressedSwapMemory);
 
 			FAllocationItem* AllocationPtr = RootHeap.LiveAllocs->AddNewSwap(MaskedPagedAddress);
 			FAllocationItem& Allocation = *AllocationPtr;
@@ -1869,7 +1890,7 @@ void FAllocationsProvider::EditSwapOp(uint32 ThreadId, double Time, uint64 Unmas
 			{
 				TotalSwapMemory -= PlatformPageSize;
 			}
-			SampleMinSwapMemory = FMath::Min(SampleMinSwapMemory, TotalSwapMemory);
+			SetMinSampleU64(ETimelineU64::MinTotalSwapMemory, TotalSwapMemory);
 
 			FAllocationItem* AllocationPtr = RootHeap.LiveAllocs->RemoveSwap(MaskedPagedAddress);
 
@@ -1889,7 +1910,7 @@ void FAllocationsProvider::EditSwapOp(uint32 ThreadId, double Time, uint64 Unmas
 				{
 					TotalCompressedSwapMemory -= CompressedPageSize;
 				}
-				SampleMinCompressedSwapMemory = FMath::Min(SampleMinCompressedSwapMemory, TotalCompressedSwapMemory);
+				SetMinSampleU64(ETimelineU64::MinTotalCompressedSwapMemory, TotalCompressedSwapMemory);
 
 				check(RootHeap.EventIndex > AllocationPtr->StartEventIndex);
 				AllocationPtr->EndEventIndex = RootHeap.EventIndex;
@@ -2090,6 +2111,8 @@ void FAllocationsProvider::EditMarkAllocationAsHeap(uint32 ThreadId, double Time
 		return;
 	}
 
+	++HeapEvents;
+
 	INSIGHTS_WATCH_API_LOGF(ThreadId, Time, Address, TEXT("MarkAllocAsHeap 0x%llX : Heap=%u Flags=0x%X CallstackId=%u"), Address, Heap, uint32(Flags), CallstackId);
 
 	if (!IsValidHeap(Heap))
@@ -2168,8 +2191,8 @@ void FAllocationsProvider::EditMarkAllocationAsHeap(uint32 ThreadId, double Time
 		// Update stats. Remove this allocation from the total.
 		TotalAllocatedMemory -= Alloc->GetSize();
 		--TotalLiveAllocations;
-		SampleMinTotalAllocatedMemory = FMath::Min(SampleMinTotalAllocatedMemory, TotalAllocatedMemory);
-		SampleMinLiveAllocations = FMath::Min(SampleMinLiveAllocations, TotalLiveAllocations);
+		SetMinSampleU64(ETimelineU64::MinTotalAllocatedMemory, TotalAllocatedMemory);
+		SetMinSampleU32(ETimelineU32::MinLiveAllocations, TotalLiveAllocations);
 	}
 	else
 	{
@@ -2194,6 +2217,8 @@ void FAllocationsProvider::EditUnmarkAllocationAsHeap(uint32 ThreadId, double Ti
 		return;
 	}
 
+	++HeapEvents;
+
 	INSIGHTS_WATCH_API_LOGF(ThreadId, Time, Address, TEXT("UnmarkAllocAsHeap 0x%llX : Heap=%u CallstackId=%u"), Address, Heap, CallstackId);
 
 	if (!IsValidHeap(Heap))
@@ -2425,54 +2450,44 @@ void FAllocationsProvider::AdvanceTimelines(double Time)
 	{
 		// Add the current sample to the timelines.
 		Timeline.EmplaceBack(SampleStartTimestamp);
-		MinTotalAllocatedMemoryTimeline.EmplaceBack(SampleMinTotalAllocatedMemory);
-		MaxTotalAllocatedMemoryTimeline.EmplaceBack(SampleMaxTotalAllocatedMemory);
-		MinLiveAllocationsTimeline.EmplaceBack(SampleMinLiveAllocations);
-		MaxLiveAllocationsTimeline.EmplaceBack(SampleMaxLiveAllocations);
-		MinTotalSwapMemoryTimeline.EmplaceBack(SampleMinSwapMemory);
-		MaxTotalSwapMemoryTimeline.EmplaceBack(SampleMaxSwapMemory);
-		MinTotalCompressedSwapMemoryTimeline.EmplaceBack(SampleMinCompressedSwapMemory);
-		MaxTotalCompressedSwapMemoryTimeline.EmplaceBack(SampleMaxCompressedSwapMemory);
-		AllocEventsTimeline.EmplaceBack(SampleAllocEvents);
-		FreeEventsTimeline.EmplaceBack(SampleFreeEvents);
-		PageInEventsTimeline.EmplaceBack(SamplePageInEvents);
-		PageOutEventsTimeline.EmplaceBack(SamplePageOutEvents);
-		SwapFreeEventsTimeline.EmplaceBack(SampleSwapFreeEvents);
+		for (int32 TimelineIndex = 0; TimelineIndex < NumTimelinesU64; ++TimelineIndex)
+		{
+			ValueTimelinesU64[TimelineIndex]->EmplaceBack(SampleValuesU64[TimelineIndex]);
+		}
+		for (int32 TimelineIndex = 0; TimelineIndex < NumTimelinesU32; ++TimelineIndex)
+		{
+			ValueTimelinesU32[TimelineIndex]->EmplaceBack(SampleValuesU32[TimelineIndex]);
+		}
 
 		// Start a new sample.
 		SampleStartTimestamp = Time;
-		SampleMinTotalAllocatedMemory = TotalAllocatedMemory;
-		SampleMaxTotalAllocatedMemory = TotalAllocatedMemory;
-		SampleMinLiveAllocations = TotalLiveAllocations;
-		SampleMaxLiveAllocations = TotalLiveAllocations;
-		SampleMinSwapMemory = TotalSwapMemory;
-		SampleMaxSwapMemory = TotalSwapMemory;
-		SampleMinCompressedSwapMemory = TotalCompressedSwapMemory;
-		SampleMaxCompressedSwapMemory = TotalCompressedSwapMemory;
-		SampleAllocEvents = 0;
-		SampleFreeEvents = 0;
-		SamplePageInEvents = 0;
-		SamplePageOutEvents = 0;
-		SampleSwapFreeEvents = 0;
+		SetSampleU64(ETimelineU64::MinTotalAllocatedMemory, TotalAllocatedMemory);
+		SetSampleU64(ETimelineU64::MaxTotalAllocatedMemory, TotalAllocatedMemory);
+		SetSampleU32(ETimelineU32::MinLiveAllocations, TotalLiveAllocations);
+		SetSampleU32(ETimelineU32::MaxLiveAllocations, TotalLiveAllocations);
+		SetSampleU64(ETimelineU64::MinTotalSwapMemory, TotalSwapMemory);
+		SetSampleU64(ETimelineU64::MaxTotalSwapMemory, TotalSwapMemory);
+		SetSampleU64(ETimelineU64::MinTotalCompressedSwapMemory, TotalCompressedSwapMemory);
+		SetSampleU64(ETimelineU64::MaxTotalCompressedSwapMemory, TotalCompressedSwapMemory);
+		SetSampleU32(ETimelineU32::AllocEvents, 0);
+		SetSampleU32(ETimelineU32::FreeEvents, 0);
+		SetSampleU32(ETimelineU32::PageInEvents, 0);
+		SetSampleU32(ETimelineU32::PageOutEvents, 0);
+		SetSampleU32(ETimelineU32::SwapFreeEvents, 0);
 
 		// If the previous sample is well distanced in time...
 		if (Time - SampleEndTimestamp > DefaultTimelineSampleGranularity)
 		{
 			// Add an intermediate "flat region" sample.
 			Timeline.EmplaceBack(SampleEndTimestamp);
-			MinTotalAllocatedMemoryTimeline.EmplaceBack(TotalAllocatedMemory);
-			MaxTotalAllocatedMemoryTimeline.EmplaceBack(TotalAllocatedMemory);
-			MinLiveAllocationsTimeline.EmplaceBack(TotalLiveAllocations);
-			MaxLiveAllocationsTimeline.EmplaceBack(TotalLiveAllocations);
-			MinTotalSwapMemoryTimeline.EmplaceBack(TotalSwapMemory);
-			MaxTotalSwapMemoryTimeline.EmplaceBack(TotalSwapMemory);
-			MinTotalCompressedSwapMemoryTimeline.EmplaceBack(TotalCompressedSwapMemory);
-			MaxTotalCompressedSwapMemoryTimeline.EmplaceBack(TotalCompressedSwapMemory);
-			AllocEventsTimeline.EmplaceBack(0);
-			FreeEventsTimeline.EmplaceBack(0);
-			PageInEventsTimeline.EmplaceBack(0);
-			PageOutEventsTimeline.EmplaceBack(0);
-			SwapFreeEventsTimeline.EmplaceBack(0);
+			for (int32 TimelineIndex = 0; TimelineIndex < NumTimelinesU64; ++TimelineIndex)
+			{
+				ValueTimelinesU64[TimelineIndex]->EmplaceBack(SampleValuesU64[TimelineIndex]);
+			}
+			for (int32 TimelineIndex = 0; TimelineIndex < NumTimelinesU32; ++TimelineIndex)
+			{
+				ValueTimelinesU32[TimelineIndex]->EmplaceBack(SampleValuesU32[TimelineIndex]);
+			}
 		}
 	}
 
@@ -2676,8 +2691,11 @@ void FAllocationsProvider::EditOnAnalysisCompleted(double Time)
 			++NumHeapSpecs;
 		}
 	}
-	UE_LOG(LogTraceServices, Log, TEXT("[MemAlloc] Analysis completed (%llu events, %llu allocs, %llu frees, %llu heaps, %u heap specs)."),
-		TotalEventCount, AllocCount, FreeCount, HeapCount, NumHeapSpecs);
+	UE_LOG(LogTraceServices, Log, TEXT("[MemAlloc] Indexed: %llu events"), TotalEventCount);
+	UE_LOG(LogTraceServices, Log, TEXT("[MemAlloc] Traced: %llu alloc events, %llu free events, %llu heap events, %llu swap op events"),
+		AllocEvents, FreeEvents, HeapEvents, SwapOpEvents);
+	UE_LOG(LogTraceServices, Log, TEXT("[MemAlloc] Analysis completed (%llu allocs, %llu frees, %llu heaps, %u heap specs)."),
+		AllocCount, FreeCount, HeapCount, NumHeapSpecs);
 }
 
 ////////////////////////////////////////////////////////////////////////////////////////////////////
@@ -2776,469 +2794,99 @@ void FAllocationsProvider::GetTimelineIndexRange(double StartTime, double EndTim
 
 ////////////////////////////////////////////////////////////////////////////////////////////////////
 
-void FAllocationsProvider::EnumerateMinTotalAllocatedMemoryTimeline(int32 StartIndex, int32 EndIndex, TFunctionRef<void(double Time, double Duration, uint64 Value)> Callback) const
-{
-	ReadAccessCheck();
-
-	const int32 NumPoints = static_cast<int32>(Timeline.Num());
-	StartIndex = FMath::Max(StartIndex, 0);
-	EndIndex = FMath::Min(EndIndex + 1, NumPoints); // make it exclusive
-	if (StartIndex < EndIndex)
-	{
-		auto TimeIt = Timeline.GetIteratorFromItem(StartIndex);
-		auto ValueIt = MinTotalAllocatedMemoryTimeline.GetIteratorFromItem(StartIndex);
-		double PrevTime = *TimeIt;
-		uint64 PrevValue = *ValueIt;
-		++TimeIt;
-		++ValueIt;
-		for (int32 Index = StartIndex + 1; Index < EndIndex; ++Index, ++TimeIt, ++ValueIt)
-		{
-			const double Time = *TimeIt;
-			Callback(PrevTime, Time - PrevTime, PrevValue);
-			PrevTime = *TimeIt;
-			PrevValue = *ValueIt;
-		}
-		if (EndIndex < NumPoints)
-		{
-			const double Time = *TimeIt;
-			Callback(PrevTime, Time - PrevTime, PrevValue);
-		}
-		else
-		{
-			Callback(PrevTime, std::numeric_limits<double>::infinity(), PrevValue);
-		}
-	}
-}
-
-////////////////////////////////////////////////////////////////////////////////////////////////////
-
-void FAllocationsProvider::EnumerateMaxTotalAllocatedMemoryTimeline(int32 StartIndex, int32 EndIndex, TFunctionRef<void(double Time, double Duration, uint64 Value)> Callback) const
-{
-	ReadAccessCheck();
-
-	const int32 NumPoints = static_cast<int32>(Timeline.Num());
-	StartIndex = FMath::Max(StartIndex, 0);
-	EndIndex = FMath::Min(EndIndex + 1, NumPoints); // make it exclusive
-	if (StartIndex < EndIndex)
-	{
-		auto TimeIt = Timeline.GetIteratorFromItem(StartIndex);
-		auto ValueIt = MaxTotalAllocatedMemoryTimeline.GetIteratorFromItem(StartIndex);
-		double PrevTime = *TimeIt;
-		uint64 PrevValue = *ValueIt;
-		++TimeIt;
-		++ValueIt;
-		for (int32 Index = StartIndex + 1; Index < EndIndex; ++Index, ++TimeIt, ++ValueIt)
-		{
-			const double Time = *TimeIt;
-			Callback(PrevTime, Time - PrevTime, PrevValue);
-			PrevTime = *TimeIt;
-			PrevValue = *ValueIt;
-		}
-		if (EndIndex < NumPoints)
-		{
-			const double Time = *TimeIt;
-			Callback(PrevTime, Time - PrevTime, PrevValue);
-		}
-		else
-		{
-			Callback(PrevTime, std::numeric_limits<double>::infinity(), PrevValue);
-		}
-	}
-}
-
-////////////////////////////////////////////////////////////////////////////////////////////////////
-
-void FAllocationsProvider::EnumerateMinLiveAllocationsTimeline(int32 StartIndex, int32 EndIndex, TFunctionRef<void(double Time, double Duration, uint32 Value)> Callback) const
+void FAllocationsProvider::EnumerateTimeline(ETimelineU64 TimelineIndex, int32 StartIndex, int32 EndIndex, TFunctionRef<void(double Time, double Duration, uint64 Value)> Callback) const
 {
 	ReadAccessCheck();
 
-	const int32 NumPoints = static_cast<int32>(Timeline.Num());
-	StartIndex = FMath::Max(StartIndex, 0);
-	EndIndex = FMath::Min(EndIndex + 1, NumPoints); // make it exclusive
-	if (StartIndex < EndIndex)
+	if (int32(TimelineIndex) < 0 ||
+		int32(TimelineIndex) >= NumTimelinesU64 ||
+		ValueTimelinesU64[int32(TimelineIndex)] == nullptr)
 	{
-		auto TimeIt = Timeline.GetIteratorFromItem(StartIndex);
-		auto ValueIt = MinLiveAllocationsTimeline.GetIteratorFromItem(StartIndex);
-		double PrevTime = *TimeIt;
-		uint32 PrevValue = *ValueIt;
-		++TimeIt;
-		++ValueIt;
-		for (int32 Index = StartIndex + 1; Index < EndIndex; ++Index, ++TimeIt, ++ValueIt)
-		{
-			const double Time = *TimeIt;
-			Callback(PrevTime, Time - PrevTime, PrevValue);
-			PrevTime = *TimeIt;
-			PrevValue = *ValueIt;
-		}
-		if (EndIndex < NumPoints)
-		{
-			const double Time = *TimeIt;
-			Callback(PrevTime, Time - PrevTime, PrevValue);
-		}
-		else
-		{
-			Callback(PrevTime, std::numeric_limits<double>::infinity(), PrevValue);
-		}
-	}
-}
-
-////////////////////////////////////////////////////////////////////////////////////////////////////
-
-void FAllocationsProvider::EnumerateMaxLiveAllocationsTimeline(int32 StartIndex, int32 EndIndex, TFunctionRef<void(double Time, double Duration, uint32 Value)> Callback) const
-{
-	ReadAccessCheck();
-
-	const int32 NumPoints = static_cast<int32>(Timeline.Num());
-	StartIndex = FMath::Max(StartIndex, 0);
-	EndIndex = FMath::Min(EndIndex + 1, NumPoints); // make it exclusive
-	if (StartIndex < EndIndex)
-	{
-		auto TimeIt = Timeline.GetIteratorFromItem(StartIndex);
-		auto ValueIt = MaxLiveAllocationsTimeline.GetIteratorFromItem(StartIndex);
-		double PrevTime = *TimeIt;
-		uint32 PrevValue = *ValueIt;
-		++TimeIt;
-		++ValueIt;
-		for (int32 Index = StartIndex + 1; Index < EndIndex; ++Index, ++TimeIt, ++ValueIt)
-		{
-			const double Time = *TimeIt;
-			Callback(PrevTime, Time - PrevTime, PrevValue);
-			PrevTime = *TimeIt;
-			PrevValue = *ValueIt;
-		}
-		if (EndIndex < NumPoints)
-		{
-			const double Time = *TimeIt;
-			Callback(PrevTime, Time - PrevTime, PrevValue);
-		}
-		else
-		{
-			Callback(PrevTime, std::numeric_limits<double>::infinity(), PrevValue);
-		}
+		return;
 	}
-}
-
-////////////////////////////////////////////////////////////////////////////////////////////////////
-
-void FAllocationsProvider::EnumerateMinTotalSwapMemoryTimeline(int32 StartIndex, int32 EndIndex, TFunctionRef<void(double Time, double Duration, uint64 Value)> Callback) const
-{
-	ReadAccessCheck();
+	const TPagedArray<uint64>& ValueTimeline = *ValueTimelinesU64[int32(TimelineIndex)];
 
 	const int32 NumPoints = static_cast<int32>(Timeline.Num());
 	StartIndex = FMath::Max(StartIndex, 0);
 	EndIndex = FMath::Min(EndIndex + 1, NumPoints); // make it exclusive
-	if (StartIndex < EndIndex)
+	if (StartIndex >= EndIndex)
 	{
-		auto TimeIt = Timeline.GetIteratorFromItem(StartIndex);
-		auto ValueIt = MinTotalSwapMemoryTimeline.GetIteratorFromItem(StartIndex);
-		double PrevTime = *TimeIt;
-		uint64 PrevValue = *ValueIt;
-		++TimeIt;
-		++ValueIt;
-		for (int32 Index = StartIndex + 1; Index < EndIndex; ++Index, ++TimeIt, ++ValueIt)
-		{
-			const double Time = *TimeIt;
-			Callback(PrevTime, Time - PrevTime, PrevValue);
-			PrevTime = *TimeIt;
-			PrevValue = *ValueIt;
-		}
-		if (EndIndex < NumPoints)
-		{
-			const double Time = *TimeIt;
-			Callback(PrevTime, Time - PrevTime, PrevValue);
-		}
-		else
-		{
-			Callback(PrevTime, std::numeric_limits<double>::infinity(), PrevValue);
-		}
+		return;
 	}
-}
 
-////////////////////////////////////////////////////////////////////////////////////////////////////
+	auto TimeIt = Timeline.GetIteratorFromItem(StartIndex);
+	auto ValueIt = ValueTimeline.GetIteratorFromItem(StartIndex);
 
-void FAllocationsProvider::EnumerateMaxTotalSwapMemoryTimeline(int32 StartIndex, int32 EndIndex, TFunctionRef<void(double Time, double Duration, uint64 Value)> Callback) const
-{
-	ReadAccessCheck();
+	double PrevTime = *TimeIt;
+	uint64 PrevValue = *ValueIt;
+	++TimeIt;
+	++ValueIt;
 
-	const int32 NumPoints = static_cast<int32>(Timeline.Num());
-	StartIndex = FMath::Max(StartIndex, 0);
-	EndIndex = FMath::Min(EndIndex + 1, NumPoints); // make it exclusive
-	if (StartIndex < EndIndex)
+	for (int32 Index = StartIndex + 1; Index < EndIndex; ++Index, ++TimeIt, ++ValueIt)
 	{
-		auto TimeIt = Timeline.GetIteratorFromItem(StartIndex);
-		auto ValueIt = MaxTotalSwapMemoryTimeline.GetIteratorFromItem(StartIndex);
-		double PrevTime = *TimeIt;
-		uint64 PrevValue = *ValueIt;
-		++TimeIt;
-		++ValueIt;
-		for (int32 Index = StartIndex + 1; Index < EndIndex; ++Index, ++TimeIt, ++ValueIt)
-		{
-			const double Time = *TimeIt;
-			Callback(PrevTime, Time - PrevTime, PrevValue);
-			PrevTime = *TimeIt;
-			PrevValue = *ValueIt;
-		}
-		if (EndIndex < NumPoints)
-		{
-			const double Time = *TimeIt;
-			Callback(PrevTime, Time - PrevTime, PrevValue);
-		}
-		else
-		{
-			Callback(PrevTime, std::numeric_limits<double>::infinity(), PrevValue);
-		}
+		const double Time = *TimeIt;
+		Callback(PrevTime, Time - PrevTime, PrevValue);
+		PrevTime = *TimeIt;
+		PrevValue = *ValueIt;
 	}
-}
 
-////////////////////////////////////////////////////////////////////////////////////////////////////
-
-void FAllocationsProvider::EnumerateMinTotalCompressedSwapMemoryTimeline(int32 StartIndex, int32 EndIndex, TFunctionRef<void(double Time, double Duration, uint64 Value)> Callback) const
-{
-	ReadAccessCheck();
-
-	const int32 NumPoints = static_cast<int32>(Timeline.Num());
-	StartIndex = FMath::Max(StartIndex, 0);
-	EndIndex = FMath::Min(EndIndex + 1, NumPoints); // make it exclusive
-	if (StartIndex < EndIndex)
+	if (EndIndex < NumPoints)
 	{
-		auto TimeIt = Timeline.GetIteratorFromItem(StartIndex);
-		auto ValueIt = MinTotalCompressedSwapMemoryTimeline.GetIteratorFromItem(StartIndex);
-		double PrevTime = *TimeIt;
-		uint64 PrevValue = *ValueIt;
-		++TimeIt;
-		++ValueIt;
-		for (int32 Index = StartIndex + 1; Index < EndIndex; ++Index, ++TimeIt, ++ValueIt)
-		{
-			const double Time = *TimeIt;
-			Callback(PrevTime, Time - PrevTime, PrevValue);
-			PrevTime = *TimeIt;
-			PrevValue = *ValueIt;
-		}
-		if (EndIndex < NumPoints)
-		{
-			const double Time = *TimeIt;
-			Callback(PrevTime, Time - PrevTime, PrevValue);
-		}
-		else
-		{
-			Callback(PrevTime, std::numeric_limits<double>::infinity(), PrevValue);
-		}
+		const double Time = *TimeIt;
+		Callback(PrevTime, Time - PrevTime, PrevValue);
 	}
-}
-
-////////////////////////////////////////////////////////////////////////////////////////////////////
-
-void FAllocationsProvider::EnumerateMaxTotalCompressedSwapMemoryTimeline(int32 StartIndex, int32 EndIndex, TFunctionRef<void(double Time, double Duration, uint64 Value)> Callback) const
-{
-	ReadAccessCheck();
-
-	const int32 NumPoints = static_cast<int32>(Timeline.Num());
-	StartIndex = FMath::Max(StartIndex, 0);
-	EndIndex = FMath::Min(EndIndex + 1, NumPoints); // make it exclusive
-	if (StartIndex < EndIndex)
+	else
 	{
-		auto TimeIt = Timeline.GetIteratorFromItem(StartIndex);
-		auto ValueIt = MaxTotalCompressedSwapMemoryTimeline.GetIteratorFromItem(StartIndex);
-		double PrevTime = *TimeIt;
-		uint64 PrevValue = *ValueIt;
-		++TimeIt;
-		++ValueIt;
-		for (int32 Index = StartIndex + 1; Index < EndIndex; ++Index, ++TimeIt, ++ValueIt)
-		{
-			const double Time = *TimeIt;
-			Callback(PrevTime, Time - PrevTime, PrevValue);
-			PrevTime = *TimeIt;
-			PrevValue = *ValueIt;
-		}
-		if (EndIndex < NumPoints)
-		{
-			const double Time = *TimeIt;
-			Callback(PrevTime, Time - PrevTime, PrevValue);
-		}
-		else
-		{
-			Callback(PrevTime, std::numeric_limits<double>::infinity(), PrevValue);
-		}
+		Callback(PrevTime, std::numeric_limits<double>::infinity(), PrevValue);
 	}
 }
 
 ////////////////////////////////////////////////////////////////////////////////////////////////////
 
-void FAllocationsProvider::EnumerateAllocEventsTimeline(int32 StartIndex, int32 EndIndex, TFunctionRef<void(double Time, double Duration, uint32 Value)> Callback) const
+void FAllocationsProvider::EnumerateTimeline(ETimelineU32 TimelineIndex, int32 StartIndex, int32 EndIndex, TFunctionRef<void(double Time, double Duration, uint32 Value)> Callback) const
 {
 	ReadAccessCheck();
 
-	const int32 NumPoints = static_cast<int32>(Timeline.Num());
-	StartIndex = FMath::Max(StartIndex, 0);
-	EndIndex = FMath::Min(EndIndex + 1, NumPoints); // make it exclusive
-	if (StartIndex < EndIndex)
+	if (int32(TimelineIndex) < 0 ||
+		int32(TimelineIndex) >= NumTimelinesU32 ||
+		ValueTimelinesU32[int32(TimelineIndex)] == nullptr)
 	{
-		auto TimeIt = Timeline.GetIteratorFromItem(StartIndex);
-		auto ValueIt = AllocEventsTimeline.GetIteratorFromItem(StartIndex);
-		double PrevTime = *TimeIt;
-		uint32 PrevValue = *ValueIt;
-		++TimeIt;
-		++ValueIt;
-		for (int32 Index = StartIndex + 1; Index < EndIndex; ++Index, ++TimeIt, ++ValueIt)
-		{
-			const double Time = *TimeIt;
-			Callback(PrevTime, Time - PrevTime, PrevValue);
-			PrevTime = *TimeIt;
-			PrevValue = *ValueIt;
-		}
-		if (EndIndex < NumPoints)
-		{
-			const double Time = *TimeIt;
-			Callback(PrevTime, Time - PrevTime, PrevValue);
-		}
-		else
-		{
-			Callback(PrevTime, std::numeric_limits<double>::infinity(), PrevValue);
-		}
+		return;
 	}
-}
-
-////////////////////////////////////////////////////////////////////////////////////////////////////
-
-void FAllocationsProvider::EnumerateFreeEventsTimeline(int32 StartIndex, int32 EndIndex, TFunctionRef<void(double Time, double Duration, uint32 Value)> Callback) const
-{
-	ReadAccessCheck();
+	const TPagedArray<uint32>& ValueTimeline = *ValueTimelinesU32[int32(TimelineIndex)];
 
 	const int32 NumPoints = static_cast<int32>(Timeline.Num());
 	StartIndex = FMath::Max(StartIndex, 0);
 	EndIndex = FMath::Min(EndIndex + 1, NumPoints); // make it exclusive
-	if (StartIndex < EndIndex)
+	if (StartIndex >= EndIndex)
 	{
-		auto TimeIt = Timeline.GetIteratorFromItem(StartIndex);
-		auto ValueIt = FreeEventsTimeline.GetIteratorFromItem(StartIndex);
-		double PrevTime = *TimeIt;
-		uint32 PrevValue = *ValueIt;
-		++TimeIt;
-		++ValueIt;
-		for (int32 Index = StartIndex + 1; Index < EndIndex; ++Index, ++TimeIt, ++ValueIt)
-		{
-			const double Time = *TimeIt;
-			Callback(PrevTime, Time - PrevTime, PrevValue);
-			PrevTime = *TimeIt;
-			PrevValue = *ValueIt;
-		}
-		if (EndIndex < NumPoints)
-		{
-			const double Time = *TimeIt;
-			Callback(PrevTime, Time - PrevTime, PrevValue);
-		}
-		else
-		{
-			Callback(PrevTime, std::numeric_limits<double>::infinity(), PrevValue);
-		}
+		return;
 	}
-}
 
-////////////////////////////////////////////////////////////////////////////////////////////////////
+	auto TimeIt = Timeline.GetIteratorFromItem(StartIndex);
+	auto ValueIt = ValueTimeline.GetIteratorFromItem(StartIndex);
 
-void FAllocationsProvider::EnumeratePageInEventsTimeline(int32 StartIndex, int32 EndIndex, TFunctionRef<void(double Time, double Duration, uint32 Value)> Callback) const
-{
-	ReadAccessCheck();
+	double PrevTime = *TimeIt;
+	uint32 PrevValue = *ValueIt;
+	++TimeIt;
+	++ValueIt;
 
-	const int32 NumPoints = static_cast<int32>(Timeline.Num());
-	StartIndex = FMath::Max(StartIndex, 0);
-	EndIndex = FMath::Min(EndIndex + 1, NumPoints); // make it exclusive
-	if (StartIndex < EndIndex)
+	for (int32 Index = StartIndex + 1; Index < EndIndex; ++Index, ++TimeIt, ++ValueIt)
 	{
-		auto TimeIt = Timeline.GetIteratorFromItem(StartIndex);
-		auto ValueIt = PageInEventsTimeline.GetIteratorFromItem(StartIndex);
-		double PrevTime = *TimeIt;
-		uint32 PrevValue = *ValueIt;
-		++TimeIt;
-		++ValueIt;
-		for (int32 Index = StartIndex + 1; Index < EndIndex; ++Index, ++TimeIt, ++ValueIt)
-		{
-			const double Time = *TimeIt;
-			Callback(PrevTime, Time - PrevTime, PrevValue);
-			PrevTime = *TimeIt;
-			PrevValue = *ValueIt;
-		}
-		if (EndIndex < NumPoints)
-		{
-			const double Time = *TimeIt;
-			Callback(PrevTime, Time - PrevTime, PrevValue);
-		}
-		else
-		{
-			Callback(PrevTime, std::numeric_limits<double>::infinity(), PrevValue);
-		}
+		const double Time = *TimeIt;
+		Callback(PrevTime, Time - PrevTime, PrevValue);
+		PrevTime = *TimeIt;
+		PrevValue = *ValueIt;
 	}
-}
-
-////////////////////////////////////////////////////////////////////////////////////////////////////
-
-void FAllocationsProvider::EnumeratePageOutEventsTimeline(int32 StartIndex, int32 EndIndex, TFunctionRef<void(double Time, double Duration, uint32 Value)> Callback) const
-{
-	ReadAccessCheck();
 
-	const int32 NumPoints = static_cast<int32>(Timeline.Num());
-	StartIndex = FMath::Max(StartIndex, 0);
-	EndIndex = FMath::Min(EndIndex + 1, NumPoints); // make it exclusive
-	if (StartIndex < EndIndex)
+	if (EndIndex < NumPoints)
 	{
-		auto TimeIt = Timeline.GetIteratorFromItem(StartIndex);
-		auto ValueIt = PageOutEventsTimeline.GetIteratorFromItem(StartIndex);
-		double PrevTime = *TimeIt;
-		uint32 PrevValue = *ValueIt;
-		++TimeIt;
-		++ValueIt;
-		for (int32 Index = StartIndex + 1; Index < EndIndex; ++Index, ++TimeIt, ++ValueIt)
-		{
-			const double Time = *TimeIt;
-			Callback(PrevTime, Time - PrevTime, PrevValue);
-			PrevTime = *TimeIt;
-			PrevValue = *ValueIt;
-		}
-		if (EndIndex < NumPoints)
-		{
-			const double Time = *TimeIt;
-			Callback(PrevTime, Time - PrevTime, PrevValue);
-		}
-		else
-		{
-			Callback(PrevTime, std::numeric_limits<double>::infinity(), PrevValue);
-		}
+		const double Time = *TimeIt;
+		Callback(PrevTime, Time - PrevTime, PrevValue);
 	}
-}
-
-////////////////////////////////////////////////////////////////////////////////////////////////////
-
-void FAllocationsProvider::EnumerateSwapFreeEventsTimeline(int32 StartIndex, int32 EndIndex, TFunctionRef<void(double Time, double Duration, uint32 Value)> Callback) const
-{
-	ReadAccessCheck();
-
-	const int32 NumPoints = static_cast<int32>(Timeline.Num());
-	StartIndex = FMath::Max(StartIndex, 0);
-	EndIndex = FMath::Min(EndIndex + 1, NumPoints); // make it exclusive
-	if (StartIndex < EndIndex)
+	else
 	{
-		auto TimeIt = Timeline.GetIteratorFromItem(StartIndex);
-		auto ValueIt = SwapFreeEventsTimeline.GetIteratorFromItem(StartIndex);
-		double PrevTime = *TimeIt;
-		uint32 PrevValue = *ValueIt;
-		++TimeIt;
-		++ValueIt;
-		for (int32 Index = StartIndex + 1; Index < EndIndex; ++Index, ++TimeIt, ++ValueIt)
-		{
-			const double Time = *TimeIt;
-			Callback(PrevTime, Time - PrevTime, PrevValue);
-			PrevTime = *TimeIt;
-			PrevValue = *ValueIt;
-		}
-		if (EndIndex < NumPoints)
-		{
-			const double Time = *TimeIt;
-			Callback(PrevTime, Time - PrevTime, PrevValue);
-		}
-		else
-		{
-			Callback(PrevTime, std::numeric_limits<double>::infinity(), PrevValue);
-		}
+		Callback(PrevTime, std::numeric_limits<double>::infinity(), PrevValue);
 	}
 }
 
diff --git a/Engine/Source/Developer/TraceServices/Private/Model/AllocationsProvider.h b/Engine/Source/Developer/TraceServices/Private/Model/AllocationsProvider.h
--- a/Engine/Source/Developer/TraceServices/Private/Model/AllocationsProvider.h
+++ b/Engine/Source/Developer/TraceServices/Private/Model/AllocationsProvider.h
@@ -325,6 +325,10 @@ public:
 
 	virtual bool IsInitialized() const override { ReadAccessCheck(); return bInitialized; }
 
+	virtual bool HasAllocationEvents() const { ReadAccessCheck(); return AllocEvents > 0 || FreeEvents > 0 || HeapEvents > 0; }
+
+	virtual bool HasSwapOpEvents() const { ReadAccessCheck(); return SwapOpEvents > 0; }
+
 	virtual void EnumerateTags(TFunctionRef<void(const TCHAR*, const TCHAR*, TagIdType, TagIdType)> Callback) const override;
 
 	virtual const TCHAR* GetTagName(TagIdType Tag) const override { ReadAccessCheck(); return TagTracker.GetTagString(Tag); }
@@ -336,19 +340,9 @@ public:
 
 	virtual int32 GetTimelineNumPoints() const override { ReadAccessCheck(); return static_cast<int32>(Timeline.Num()); }
 	virtual void GetTimelineIndexRange(double StartTime, double EndTime, int32& StartIndex, int32& EndIndex) const override;
-	virtual void EnumerateMinTotalAllocatedMemoryTimeline(int32 StartIndex, int32 EndIndex, TFunctionRef<void(double Time, double Duration, uint64 Value)> Callback) const override;
-	virtual void EnumerateMaxTotalAllocatedMemoryTimeline(int32 StartIndex, int32 EndIndex, TFunctionRef<void(double Time, double Duration, uint64 Value)> Callback) const override;
-	virtual void EnumerateMinLiveAllocationsTimeline(int32 StartIndex, int32 EndIndex, TFunctionRef<void(double Time, double Duration, uint32 Value)> Callback) const override;
-	virtual void EnumerateMaxLiveAllocationsTimeline(int32 StartIndex, int32 EndIndex, TFunctionRef<void(double Time, double Duration, uint32 Value)> Callback) const override;
-	virtual void EnumerateMinTotalSwapMemoryTimeline(int32 StartIndex, int32 EndIndex, TFunctionRef<void(double Time, double Duration, uint64 Value)> Callback) const override;
-	virtual void EnumerateMaxTotalSwapMemoryTimeline(int32 StartIndex, int32 EndIndex, TFunctionRef<void(double Time, double Duration, uint64 Value)> Callback) const override;
-	virtual void EnumerateMinTotalCompressedSwapMemoryTimeline(int32 StartIndex, int32 EndIndex, TFunctionRef<void(double Time, double Duration, uint64 Value)> Callback) const override;
-	virtual void EnumerateMaxTotalCompressedSwapMemoryTimeline(int32 StartIndex, int32 EndIndex, TFunctionRef<void(double Time, double Duration, uint64 Value)> Callback) const override;
-	virtual void EnumerateAllocEventsTimeline(int32 StartIndex, int32 EndIndex, TFunctionRef<void(double Time, double Duration, uint32 Value)> Callback) const override;
-	virtual void EnumerateFreeEventsTimeline(int32 StartIndex, int32 EndIndex, TFunctionRef<void(double Time, double Duration, uint32 Value)> Callback) const override;
-	virtual void EnumeratePageInEventsTimeline(int32 StartIndex, int32 EndIndex, TFunctionRef<void(double Time, double Duration, uint32 Value)> Callback) const override;
-	virtual void EnumeratePageOutEventsTimeline(int32 StartIndex, int32 EndIndex, TFunctionRef<void(double Time, double Duration, uint32 Value)> Callback) const override;
-	virtual void EnumerateSwapFreeEventsTimeline(int32 StartIndex, int32 EndIndex, TFunctionRef<void(double Time, double Duration, uint32 Value)> Callback) const override;
+
+	virtual void EnumerateTimeline(ETimelineU64 Timeline, int32 StartIndex, int32 EndIndex, TFunctionRef<void(double Time, double Duration, uint64 Value)> Callback) const override;
+	virtual void EnumerateTimeline(ETimelineU32 Timeline, int32 StartIndex, int32 EndIndex, TFunctionRef<void(double Time, double Duration, uint32 Value)> Callback) const override;
 
 	virtual FQueryHandle StartQuery(const FQueryParams& Params) const override;
 	virtual void CancelQuery(FQueryHandle Query) const override;
@@ -424,6 +418,42 @@ private:
 	FRootHeap& FindParentRootHeapUnchecked(HeapId HeapId) const;
 	void CreateRootHeap(HeapId Id);
 
+	void SetSampleU64(ETimelineU64 TimelineIndex, uint64 Value)
+	{
+		uint64& Sample = SampleValuesU64[int32(TimelineIndex)];
+		Sample = Value;
+	}
+	void SetMinSampleU64(ETimelineU64 TimelineIndex, uint64 Value)
+	{
+		uint64& Sample = SampleValuesU64[int32(TimelineIndex)];
+		Sample = FMath::Min(Sample, Value);
+	}
+	void SetMaxSampleU64(ETimelineU64 TimelineIndex, uint64 Value)
+	{
+		uint64& Sample = SampleValuesU64[int32(TimelineIndex)];
+		Sample = FMath::Max(Sample, Value);
+	}
+	void SetSampleU32(ETimelineU32 TimelineIndex, uint32 Value)
+	{
+		uint32& Sample = SampleValuesU32[int32(TimelineIndex)];
+		Sample = Value;
+	}
+	void SetMinSampleU32(ETimelineU32 TimelineIndex, uint32 Value)
+	{
+		uint32& Sample = SampleValuesU32[int32(TimelineIndex)];
+		Sample = FMath::Min(Sample, Value);
+	}
+	void SetMaxSampleU32(ETimelineU32 TimelineIndex, uint32 Value)
+	{
+		uint32& Sample = SampleValuesU32[int32(TimelineIndex)];
+		Sample = FMath::Max(Sample, Value);
+	}
+	void IncSampleU32(ETimelineU32 TimelineIndex)
+	{
+		uint32& Sample = SampleValuesU32[int32(TimelineIndex)];
+		++Sample;
+	}
+
 private:
 	mutable FProviderLock Lock;
 
@@ -446,6 +476,11 @@ private:
 	uint64 FreeCount = 0;
 	uint64 HeapCount = 0;
 
+	uint64 AllocEvents = 0;
+	uint64 FreeEvents = 0;
+	uint64 HeapEvents = 0;
+	uint64 SwapOpEvents = 0;
+
 	uint64 MiscWarnings = 0;
 	uint64 MiscErrors = 0;
 	uint64 HeapWarnings = 0;
@@ -462,36 +497,17 @@ private:
 	uint64 TotalSwapMemory = 0;
 	uint64 TotalCompressedSwapMemory = 0;
 
+	static constexpr int32 NumTimelinesU64 = int32(IAllocationsProvider::ETimelineU64::Count);
+	static constexpr int32 NumTimelinesU32 = int32(IAllocationsProvider::ETimelineU32::Count);
+
 	double SampleStartTimestamp = 0.0;
 	double SampleEndTimestamp = 0.0;
-	uint64 SampleMinTotalAllocatedMemory = 0;
-	uint64 SampleMaxTotalAllocatedMemory = 0;
-	uint32 SampleMinLiveAllocations = 0;
-	uint32 SampleMaxLiveAllocations = 0;
-	uint64 SampleMinSwapMemory = 0;
-	uint64 SampleMaxSwapMemory = 0;
-	uint64 SampleMinCompressedSwapMemory = 0;
-	uint64 SampleMaxCompressedSwapMemory = 0;
-	uint32 SampleAllocEvents = 0;
-	uint32 SampleFreeEvents = 0;
-	uint32 SamplePageInEvents = 0;
-	uint32 SamplePageOutEvents = 0;
-	uint32 SampleSwapFreeEvents = 0;
+	uint64 SampleValuesU64[NumTimelinesU64];
+	uint32 SampleValuesU32[NumTimelinesU32];
 
 	TPagedArray<double> Timeline;
-	TPagedArray<uint64> MinTotalAllocatedMemoryTimeline;
-	TPagedArray<uint64> MaxTotalAllocatedMemoryTimeline;
-	TPagedArray<uint32> MinLiveAllocationsTimeline;
-	TPagedArray<uint32> MaxLiveAllocationsTimeline;
-	TPagedArray<uint64> MinTotalSwapMemoryTimeline;
-	TPagedArray<uint64> MaxTotalSwapMemoryTimeline;
-	TPagedArray<uint64> MinTotalCompressedSwapMemoryTimeline;
-	TPagedArray<uint64> MaxTotalCompressedSwapMemoryTimeline;
-	TPagedArray<uint32> AllocEventsTimeline;
-	TPagedArray<uint32> FreeEventsTimeline;
-	TPagedArray<uint32> PageInEventsTimeline;
-	TPagedArray<uint32> PageOutEventsTimeline;
-	TPagedArray<uint32> SwapFreeEventsTimeline;
+	TPagedArray<uint64>* ValueTimelinesU64[NumTimelinesU64];
+	TPagedArray<uint32>* ValueTimelinesU32[NumTimelinesU32];
 };
 
 ////////////////////////////////////////////////////////////////////////////////////////////////////
diff --git a/Engine/Source/Developer/TraceServices/Public/TraceServices/Model/AllocationsProvider.h b/Engine/Source/Developer/TraceServices/Public/TraceServices/Model/AllocationsProvider.h
--- a/Engine/Source/Developer/TraceServices/Public/TraceServices/Model/AllocationsProvider.h
+++ b/Engine/Source/Developer/TraceServices/Public/TraceServices/Model/AllocationsProvider.h
@@ -51,6 +51,55 @@ public:
 		//A_xor_B, // live either at A or at B; ({aAf} U {aBf}) \ {aABf}
 	};
 
+	enum class ETimelineU64
+	{
+		// The Minimum Total Allocated Memory
+		MinTotalAllocatedMemory,
+
+		// The Maximum Total Allocated Memory
+		MaxTotalAllocatedMemory,
+
+		// The Minimum Total Swap Memory
+		MinTotalSwapMemory,
+
+		// The Maximum Total Swap Memory
+		MaxTotalSwapMemory,
+
+		// The Minimum Total Compressed Swap Memory
+		MinTotalCompressedSwapMemory,
+
+		// The Maximum Total Compressed Swap Memory
+		MaxTotalCompressedSwapMemory,
+
+		Count
+	};
+
+	enum class ETimelineU32
+	{
+		// The Minimum Number of Live Allocations
+		MinLiveAllocations,
+
+		// The Maximum Number of Live Allocations
+		MaxLiveAllocations,
+
+		// The Number of Alloc Events
+		AllocEvents,
+
+		// The Number of Free Events
+		FreeEvents,
+
+		// The Number of Page In Events
+		PageInEvents,
+
+		// The Number of Page Out Events
+		PageOutEvents,
+
+		// The Number of Swap Free Events
+		SwapFreeEvents,
+
+		Count
+	};
+
 	struct FQueryParams
 	{
 		EQueryRule Rule;
@@ -126,6 +175,12 @@ public:
 
 	virtual bool IsInitialized() const = 0;
 
+	// Returns true if provider has processed at least one Alloc, Free or Heap event.
+	virtual bool HasAllocationEvents() const = 0;
+
+	// Returns true if provider has processed at least one Swap Op event.
+	virtual bool HasSwapOpEvents() const = 0;
+
 	// Enumerates the discovered tags.
 	virtual void EnumerateTags(TFunctionRef<void(const TCHAR*, const TCHAR*, TagIdType, TagIdType)> Callback) const = 0;
 
@@ -144,44 +199,91 @@ public:
 	// Index values are in range { -1, 0, .. , N-1, N }, where N = GetTimelineNumPoints().
 	virtual void GetTimelineIndexRange(double StartTime, double EndTime, int32& StartIndex, int32& EndIndex) const = 0;
 
-	// Enumerates the Min Total Allocated Memory timeline points in the inclusive index interval [StartIndex, EndIndex].
-	virtual void EnumerateMinTotalAllocatedMemoryTimeline(int32 StartIndex, int32 EndIndex, TFunctionRef<void(double Time, double Duration, uint64 Value)> Callback) const = 0;
+#if 1 // deprecated API
+	UE_DEPRECATED(5.7, "Use EnumerateTimeline(ETimelineUInt64, ...) instead")
+	virtual void EnumerateMinTotalAllocatedMemoryTimeline(int32 StartIndex, int32 EndIndex, TFunctionRef<void(double Time, double Duration, uint64 Value)> Callback) const
+	{
+		EnumerateTimeline(ETimelineU64::MinTotalAllocatedMemory, StartIndex, EndIndex, Callback);
+	}
 
-	// Enumerates the Max Total Allocated Memory timeline points in the inclusive index interval [StartIndex, EndIndex].
-	virtual void EnumerateMaxTotalAllocatedMemoryTimeline(int32 StartIndex, int32 EndIndex, TFunctionRef<void(double Time, double Duration, uint64 Value)> Callback) const = 0;
+	UE_DEPRECATED(5.7, "Use EnumerateTimeline(ETimelineUInt64, ...) instead")
+	virtual void EnumerateMaxTotalAllocatedMemoryTimeline(int32 StartIndex, int32 EndIndex, TFunctionRef<void(double Time, double Duration, uint64 Value)> Callback) const
+	{
+		EnumerateTimeline(ETimelineU64::MaxTotalAllocatedMemory, StartIndex, EndIndex, Callback);
+	}
+
+	UE_DEPRECATED(5.7, "Use EnumerateTimeline(ETimelineUInt32, ...) instead")
+	virtual void EnumerateMinLiveAllocationsTimeline(int32 StartIndex, int32 EndIndex, TFunctionRef<void(double Time, double Duration, uint32 Value)> Callback) const
+	{
+		EnumerateTimeline(ETimelineU32::MinLiveAllocations, StartIndex, EndIndex, Callback);
+	}
+
+	UE_DEPRECATED(5.7, "Use EnumerateTimeline(ETimelineUInt32, ...) instead")
+	virtual void EnumerateMaxLiveAllocationsTimeline(int32 StartIndex, int32 EndIndex, TFunctionRef<void(double Time, double Duration, uint32 Value)> Callback) const
+	{
+		EnumerateTimeline(ETimelineU32::MaxLiveAllocations, StartIndex, EndIndex, Callback);
+	}
 
-	// Enumerates the Min Live Allocations timeline points in the inclusive index interval [StartIndex, EndIndex].
-	virtual void EnumerateMinLiveAllocationsTimeline(int32 StartIndex, int32 EndIndex, TFunctionRef<void(double Time, double Duration, uint32 Value)> Callback) const = 0;
+	UE_DEPRECATED(5.7, "Use EnumerateTimeline(ETimelineUInt64, ...) instead")
+	virtual void EnumerateMinTotalSwapMemoryTimeline(int32 StartIndex, int32 EndIndex, TFunctionRef<void(double Time, double Duration, uint64 Value)> Callback) const
+	{
+		EnumerateTimeline(ETimelineU64::MinTotalSwapMemory, StartIndex, EndIndex, Callback);
+	}
 
-	// Enumerates the Max Live Allocations timeline points in the inclusive index interval [StartIndex, EndIndex].
-	virtual void EnumerateMaxLiveAllocationsTimeline(int32 StartIndex, int32 EndIndex, TFunctionRef<void(double Time, double Duration, uint32 Value)> Callback) const = 0;
+	UE_DEPRECATED(5.7, "Use EnumerateTimeline(ETimelineUInt64, ...) instead")
+	virtual void EnumerateMaxTotalSwapMemoryTimeline(int32 StartIndex, int32 EndIndex, TFunctionRef<void(double Time, double Duration, uint64 Value)> Callback) const
+	{
+		EnumerateTimeline(ETimelineU64::MaxTotalSwapMemory, StartIndex, EndIndex, Callback);
+	}
 
-	// Enumerates the Min Total Swap Memory timeline points in the inclusive index interval [StartIndex, EndIndex].
-	virtual void EnumerateMinTotalSwapMemoryTimeline(int32 StartIndex, int32 EndIndex, TFunctionRef<void(double Time, double Duration, uint64 Value)> Callback) const = 0;
+	UE_DEPRECATED(5.7, "Use EnumerateTimeline(ETimelineUInt64, ...) instead")
+	virtual void EnumerateMinTotalCompressedSwapMemoryTimeline(int32 StartIndex, int32 EndIndex, TFunctionRef<void(double Time, double Duration, uint64 Value)> Callback) const
+	{
+		EnumerateTimeline(ETimelineU64::MinTotalCompressedSwapMemory, StartIndex, EndIndex, Callback);
+	}
 
-	// Enumerates the Max Total Swap Memory timeline points in the inclusive index interval [StartIndex, EndIndex].
-	virtual void EnumerateMaxTotalSwapMemoryTimeline(int32 StartIndex, int32 EndIndex, TFunctionRef<void(double Time, double Duration, uint64 Value)> Callback) const = 0;
+	UE_DEPRECATED(5.7, "Use EnumerateTimeline(ETimelineUInt64, ...) instead")
+	virtual void EnumerateMaxTotalCompressedSwapMemoryTimeline(int32 StartIndex, int32 EndIndex, TFunctionRef<void(double Time, double Duration, uint64 Value)> Callback) const
+	{
+		EnumerateTimeline(ETimelineU64::MaxTotalCompressedSwapMemory, StartIndex, EndIndex, Callback);
+	}
 
-	// Enumerates the Min Total Compressed Swap Memory timeline points in the inclusive index interval [StartIndex, EndIndex].
-	virtual void EnumerateMinTotalCompressedSwapMemoryTimeline(int32 StartIndex, int32 EndIndex, TFunctionRef<void(double Time, double Duration, uint64 Value)> Callback) const = 0;
+	UE_DEPRECATED(5.7, "Use EnumerateTimeline(ETimelineUInt32, ...) instead")
+	virtual void EnumerateAllocEventsTimeline(int32 StartIndex, int32 EndIndex, TFunctionRef<void(double Time, double Duration, uint32 Value)> Callback) const
+	{
+		EnumerateTimeline(ETimelineU32::AllocEvents, StartIndex, EndIndex, Callback);
+	}
 
-	// Enumerates the Max Total Compressed Swap Memory timeline points in the inclusive index interval [StartIndex, EndIndex].
-	virtual void EnumerateMaxTotalCompressedSwapMemoryTimeline(int32 StartIndex, int32 EndIndex, TFunctionRef<void(double Time, double Duration, uint64 Value)> Callback) const = 0;
+	UE_DEPRECATED(5.7, "Use EnumerateTimeline(ETimelineUInt32, ...) instead")
+	virtual void EnumerateFreeEventsTimeline(int32 StartIndex, int32 EndIndex, TFunctionRef<void(double Time, double Duration, uint32 Value)> Callback) const
+	{
+		EnumerateTimeline(ETimelineU32::FreeEvents, StartIndex, EndIndex, Callback);
+	}
 
-	// Enumerates the Alloc Events timeline points in the inclusive index interval [StartIndex, EndIndex].
-	virtual void EnumerateAllocEventsTimeline(int32 StartIndex, int32 EndIndex, TFunctionRef<void(double Time, double Duration, uint32 Value)> Callback) const = 0;
+	UE_DEPRECATED(5.7, "Use EnumerateTimeline(ETimelineUInt32, ...) instead")
+	virtual void EnumeratePageInEventsTimeline(int32 StartIndex, int32 EndIndex, TFunctionRef<void(double Time, double Duration, uint32 Value)> Callback) const
+	{
+		EnumerateTimeline(ETimelineU32::PageInEvents, StartIndex, EndIndex, Callback);
+	}
 
-	// Enumerates the Free Events timeline points in the inclusive index interval [StartIndex, EndIndex].
-	virtual void EnumerateFreeEventsTimeline(int32 StartIndex, int32 EndIndex, TFunctionRef<void(double Time, double Duration, uint32 Value)> Callback) const = 0;
+	UE_DEPRECATED(5.7, "Use EnumerateTimeline(ETimelineUInt32, ...) instead")
+	virtual void EnumeratePageOutEventsTimeline(int32 StartIndex, int32 EndIndex, TFunctionRef<void(double Time, double Duration, uint32 Value)> Callback) const
+	{
+		EnumerateTimeline(ETimelineU32::PageOutEvents, StartIndex, EndIndex, Callback);
+	}
 
-	// Enumerates the Page In Events timeline points in the inclusive index interval [StartIndex, EndIndex].
-	virtual void EnumeratePageInEventsTimeline(int32 StartIndex, int32 EndIndex, TFunctionRef<void(double Time, double Duration, uint32 Value)> Callback) const = 0;
+	UE_DEPRECATED(5.7, "Use EnumerateTimeline(ETimelineUInt32, ...) instead")
+	virtual void EnumerateSwapFreeEventsTimeline(int32 StartIndex, int32 EndIndex, TFunctionRef<void(double Time, double Duration, uint32 Value)> Callback) const
+	{
+		EnumerateTimeline(ETimelineU32::SwapFreeEvents, StartIndex, EndIndex, Callback);
+	}
+#endif
 
-	// Enumerates the Page Out Events timeline points in the inclusive index interval [StartIndex, EndIndex].
-	virtual void EnumeratePageOutEventsTimeline(int32 StartIndex, int32 EndIndex, TFunctionRef<void(double Time, double Duration, uint32 Value)> Callback) const = 0;
+	// Enumerates points (time-value pairs) of a UInt64 timeline, in the inclusive index interval [StartIndex, EndIndex].
+	virtual void EnumerateTimeline(ETimelineU64 Timeline, int32 StartIndex, int32 EndIndex, TFunctionRef<void(double Time, double Duration, uint64 Value)> Callback) const = 0;
 
-	// Enumerates the Swap Free Events timeline points in the inclusive index interval [StartIndex, EndIndex].
-	virtual void EnumerateSwapFreeEventsTimeline(int32 StartIndex, int32 EndIndex, TFunctionRef<void(double Time, double Duration, uint32 Value)> Callback) const = 0;
+	// Enumerates points (time-value pairs) of a UInt32 timeline, in the inclusive index interval [StartIndex, EndIndex].
+	virtual void EnumerateTimeline(ETimelineU32 Timeline, int32 StartIndex, int32 EndIndex, TFunctionRef<void(double Time, double Duration, uint32 Value)> Callback) const = 0;
 
 	virtual FQueryHandle StartQuery(const FQueryParams& Params) const = 0;
 	virtual void CancelQuery(FQueryHandle Query) const = 0;
