[Insights] TraceServices: Moved SetRatioOfThreadsToUse() from ITimingProfilerProvider to IEditableTimingProfilerProvider. Also prepared the ITimingProfilerProvider / IEditableTimingProfilerProvider / FTimingProfilerProvider for further changes.

#rnx

diff --git a/Engine/Source/Developer/TraceInsights/Private/Insights/TimingProfiler/Widgets/STimersView.cpp b/Engine/Source/Developer/TraceInsights/Private/Insights/TimingProfiler/Widgets/STimersView.cpp
--- a/Engine/Source/Developer/TraceInsights/Private/Insights/TimingProfiler/Widgets/STimersView.cpp
+++ b/Engine/Source/Developer/TraceInsights/Private/Insights/TimingProfiler/Widgets/STimersView.cpp
@@ -510,14 +510,18 @@ void STimersView::Construct(const FArguments& InArgs)
 
 	// Update the Session (i.e. when analysis session was already started).
 	InsightsManager_OnSessionChanged();
-	
-	const TraceServices::ITimingProfilerProvider* TimingProfilerProvider = TraceServices::ReadTimingProfilerProvider(*Session.Get());
-	if (TimingProfilerProvider)
+
+	if (Session.IsValid())
 	{
-		TraceServices::FAnalysisSessionEditScope SessionEditScope(*const_cast<TraceServices::IAnalysisSession*>(Session.Get()));
-		
-		FInsightsSettings& Settings = FInsightsManager::Get()->GetSettings();
-		TimingProfilerProvider->SetRatioOfThreadsToUse(Settings.GetThreadRatio());
+		TraceServices::IAnalysisSession& EditableSession = *const_cast<TraceServices::IAnalysisSession*>(Session.Get());
+		TraceServices::IEditableTimingProfilerProvider* TimingProfilerProvider = TraceServices::EditTimingProfilerProvider(EditableSession);
+		if (TimingProfilerProvider)
+		{
+			TraceServices::FAnalysisSessionEditScope SessionEditScope(EditableSession);
+
+			FInsightsSettings& Settings = FInsightsManager::Get()->GetSettings();
+			TimingProfilerProvider->SetRatioOfThreadsToUse(Settings.GetThreadRatio());
+		}
 	}
 }
 END_SLATE_FUNCTION_BUILD_OPTIMIZATION
@@ -1785,7 +1789,7 @@ FText STimersView::GetThreadRatio() const
 		if (TimingProfilerProvider)
 		{
 			TraceServices::FAnalysisSessionReadScope SessionReadScope(*Session.Get());
-			
+
 			return FText::AsNumber((int)(TimingProfilerProvider->GetRatioOfThreadsToUse() * GThreadPool->GetNumThreads()));
 		}
 	}
@@ -1800,18 +1804,18 @@ void STimersView::OnThreadRatioCommited(double InNewValue)
 		InNewValue / GThreadPool->GetNumThreads(),
 		GetMinimumThreadRatio(),
 		1.0);
-	
+
 	if (Session.IsValid())
 	{
-		const TraceServices::ITimingProfilerProvider* TimingProfilerProvider = TraceServices::ReadTimingProfilerProvider(*Session.Get());
+		TraceServices::IAnalysisSession& EditableSession = *const_cast<TraceServices::IAnalysisSession*>(Session.Get());
+		TraceServices::IEditableTimingProfilerProvider* TimingProfilerProvider = TraceServices::EditTimingProfilerProvider(EditableSession);
 		if (TimingProfilerProvider)
 		{
-			TraceServices::FAnalysisSessionEditScope SessionEditScope(*const_cast<TraceServices::IAnalysisSession*>(Session.Get()));
-
+			TraceServices::FAnalysisSessionEditScope SessionEditScope(EditableSession);
 			TimingProfilerProvider->SetRatioOfThreadsToUse(ClampedNewValue);
 		}
 	}
-	
+
 	FInsightsSettings& Settings = FInsightsManager::Get()->GetSettings();
 	Settings.SetAndSaveThreadRatio(ClampedNewValue);
 }
@@ -3584,7 +3588,7 @@ void STimersView::SetColorToNode(FTimerNode& TimerNode, FLinearColor CustomColor
 }
 
 ////////////////////////////////////////////////////////////////////////////////////////////////////
-	
+
 void STimersView::SetRandomColorToNode(FTimerNode& TimerNode) const
 {
 	uint64 Time = FPlatformTime::Cycles64();
diff --git a/Engine/Source/Developer/TraceServices/Private/Model/TimingProfilerPrivate.h b/Engine/Source/Developer/TraceServices/Private/Model/TimingProfilerPrivate.h
--- a/Engine/Source/Developer/TraceServices/Private/Model/TimingProfilerPrivate.h
+++ b/Engine/Source/Developer/TraceServices/Private/Model/TimingProfilerPrivate.h
@@ -30,65 +30,107 @@ public:
 	explicit FTimingProfilerProvider(IAnalysisSession& InSession);
 	virtual ~FTimingProfilerProvider();
 
-	TimelineInternal& EditGpuTimeline();
-	TimelineInternal& EditGpu2Timeline();
-
+	//////////////////////////////////////////////////
 	// ITimingProfilerProvider
-	virtual bool GetCpuThreadTimelineIndex(uint32 ThreadId, uint32& OutTimelineIndex) const override;
-	virtual bool GetGpuTimelineIndex(uint32& OutTimelineIndex) const override;
-	virtual bool GetGpu2TimelineIndex(uint32& OutTimelineIndex) const override;
-	virtual void EnumerateGpuQueues(TFunctionRef<void(const FGpuQueueInfo&)> Callback) const override;
-	virtual bool GetGpuQueueTimelineIndex(uint32 QueueId, uint32& OutTimelineIndex) const override;
-	virtual bool GetVerseTimelineIndex(uint32& OutTimelineIndex) const override;
 
-	virtual void EnumerateGpuSignalFences(uint32 QueueId, double StartTime, double EndTime, EnumerateGpuSignalFencesCallback Callback) const override;
-	virtual void EnumerateGpuWaitFences(uint32 QueueId, double StartTime, double EndTime, EnumerateGpuWaitFencesCallback Callback) const override;
-	virtual void EnumerateGpuFences(uint32 QueueId, double StartTime, double EndTime, EnumerateGpuFencesCallback Callback) const override;
-	virtual void EnumerateResolvedGpuFences(uint32 QueueId, double StartTime, double EndTime, EnumerateResolvedGpuFencesCallback Callback) const override;
+	// Timelines
 
 	virtual bool ReadTimeline(uint32 Index, TFunctionRef<void(const Timeline&)> Callback) const override;
 	virtual uint32 GetTimelineCount() const override { return Timelines.Num(); }
 	virtual void EnumerateTimelines(TFunctionRef<void(const Timeline&)> Callback) const override;
+
+	// Timers
+
 	virtual void ReadTimers(TFunctionRef<void(const ITimingProfilerTimerReader&)> Callback) const override;
+
+	// Metadata
+
 	virtual uint32 GetOriginalTimerIdFromMetadata(uint32 MetadataTimerId) const override;
 	virtual TArrayView<const uint8> GetMetadata(uint32 MetadataTimerId) const override;
 	virtual const FMetadataSpec* GetMetadataSpec(uint32 MetadataSpecId) const;
+
+	// GPU
+
+	virtual bool GetGpuTimelineIndex(uint32& OutTimelineIndex) const override; // backward compatibility
+	virtual bool GetGpu2TimelineIndex(uint32& OutTimelineIndex) const override; // backward compatibility
+	virtual void EnumerateGpuQueues(TFunctionRef<void(const FGpuQueueInfo&)> Callback) const override;
+	virtual bool GetGpuQueueTimelineIndex(uint32 QueueId, uint32& OutTimelineIndex) const override;
+	virtual void EnumerateGpuSignalFences(uint32 QueueId, double StartTime, double EndTime, EnumerateGpuSignalFencesCallback Callback) const override;
+	virtual void EnumerateGpuWaitFences(uint32 QueueId, double StartTime, double EndTime, EnumerateGpuWaitFencesCallback Callback) const override;
+	virtual void EnumerateGpuFences(uint32 QueueId, double StartTime, double EndTime, EnumerateGpuFencesCallback Callback) const override;
+	virtual void EnumerateResolvedGpuFences(uint32 QueueId, double StartTime, double EndTime, EnumerateResolvedGpuFencesCallback Callback) const override;
+
+	// Verse
+
+	virtual bool GetVerseTimelineIndex(uint32& OutTimelineIndex) const override;
+
+	// CPU
+
+	virtual bool GetCpuThreadTimelineIndex(uint32 ThreadId, uint32& OutTimelineIndex) const override;
+
+	// Aggregation
+
 	virtual ITable<FTimingProfilerAggregatedStats>* CreateAggregation(const FCreateAggregationParams& Params) const override;
 	virtual ITimingProfilerButterfly* CreateButterfly(const FCreateButterflyParams& Params) const override;
+
+	// Misc
+
 	virtual double GetRatioOfThreadsToUse() const override { Session.ReadAccessCheck(); return RatioOfThreadsToUse; }
-	virtual const void SetRatioOfThreadsToUse(double InRatioOfThreadsToUse) const override { Session.WriteAccessCheck(); RatioOfThreadsToUse = InRatioOfThreadsToUse; }
 
+	//////////////////////////////////////////////////
 	// ITimingProfilerTimerReader
+
 	virtual const FTimingProfilerTimer* GetTimer(uint32 TimerId) const override;
 	virtual uint32 GetTimerCount() const override;
 	//virtual uint32 GetOriginalTimerIdFromMetadata(uint32 MetadataTimerId) const override;
 	//virtual TArrayView<const uint8> GetMetadata(uint32 MetadataTimerId) const override;
 
+	//////////////////////////////////////////////////
 	// IEditableTimingProfilerProvider
-	virtual uint32 AddCpuTimer(FStringView Name, const TCHAR* File = nullptr, uint32 Line = 0) override;
-	virtual uint32 AddGpuTimer(FStringView Name, const TCHAR* File = nullptr, uint32 Line = 0) override;
-	virtual uint32 AddVerseTimer(FStringView Name, const TCHAR* File = nullptr, uint32 Line = 0) override;
+
+	// Timers
+
 	virtual void SetTimerName(uint32 TimerId, FStringView Name) override;
 	virtual void SetTimerNameAndLocation(uint32 TimerId, FStringView Name, const TCHAR* File, uint32 Line) override;
+
+	// Metadata
+
 	virtual uint32 AddMetadata(uint32 OriginalTimerId, TArray<uint8>&& Metadata) override;
 	virtual void SetMetadata(uint32 MetadataTimerId, TArray<uint8>&& Metadata) override;
 	virtual void SetMetadata(uint32 MetadataTimerId, TArray<uint8>&& Metadata, uint32 NewTimerId) override;
 	virtual TArrayView<uint8> GetEditableMetadata(uint32 MetadataTimerId) override;
-
 	virtual uint32 AddMetadataSpec(FMetadataSpec&& Metadata) override;
 	virtual void SetMetadataSpec(uint32 TimerId, uint32 MetadataSpecId) override;
 
+	// GPU
+
+	virtual uint32 AddGpuTimer(FStringView Name, const TCHAR* File = nullptr, uint32 Line = 0) override;
 	virtual void AddGpuQueue(uint32 QueueId, uint8 GPU, uint8 Index, uint8 Type, const TCHAR* Name) override;
 	virtual void AddGpuSignalFence(uint32 QueueId, const FGpuSignalFence& SignalFence) override;
 	virtual void AddGpuWaitFence(uint32 QueueId, const FGpuWaitFence& WaitFence) override;
-
-	virtual IEditableTimeline<FTimingProfilerEvent>& GetCpuThreadEditableTimeline(uint32 ThreadId) override;
 	virtual IEditableTimeline<FTimingProfilerEvent>* GetGpuQueueEditableTimeline(uint32 QueueId) override;
 	virtual IEditableTimeline<FTimingProfilerEvent>* GetGpuQueueWorkEditableTimeline(uint32 QueueId) override;
+
+	// Verse
+
+	virtual uint32 AddVerseTimer(FStringView Name, const TCHAR* File = nullptr, uint32 Line = 0) override;
 	virtual IEditableTimeline<FTimingProfilerEvent>* GetVerseEditableTimeline() override;
 
+	// CPU
+
+	virtual uint32 AddCpuTimer(FStringView Name, const TCHAR* File = nullptr, uint32 Line = 0) override;
+	virtual IEditableTimeline<FTimingProfilerEvent>& GetCpuThreadEditableTimeline(uint32 ThreadId) override;
+
+	// Misc
+
+	virtual void SetRatioOfThreadsToUse(double InRatioOfThreadsToUse) override { Session.WriteAccessCheck(); RatioOfThreadsToUse = InRatioOfThreadsToUse; }
 	virtual const ITimingProfilerProvider* GetReadProvider() const override { return this; }
 
+	//////////////////////////////////////////////////
+
+	TimelineInternal& EditGpuTimeline();
+	TimelineInternal& EditGpu2Timeline();
+
 private:
 	enum class ETimingProfilerTimerType : uint32
 	{
@@ -110,16 +152,24 @@ private:
 	TArray<FMetadataSpec> MetadataSpecs;
 	TArray<FTimingProfilerTimer> Timers;
 	TArray<TSharedRef<TimelineInternal>> Timelines;
-	TMap<uint32, uint32> CpuThreadTimelineIndexMap; // CPU Thread Id --> timeline index, in Timelines
+	TTableLayout<FTimingProfilerAggregatedStats> AggregatedStatsTableLayout;
+	double RatioOfThreadsToUse = 0.75;
+
+	// GPU
+
 	static const uint32 GpuTimelineIndex = 0;
 	static const uint32 Gpu2TimelineIndex = 1;
-	static const uint32 VerseTimelineIndex = 2;
 	TArray<FGpuQueueInfo> GpuQueues;
 	TArray<FGpuQueueData> GpuQueueData;
 	TMap<uint32, uint32> GpuQueueIdToQueueIndexMap; // GPU Queue Id --> queue index, in GpuQueues
-	TTableLayout<FTimingProfilerAggregatedStats> AggregatedStatsTableLayout;
-	mutable double RatioOfThreadsToUse = 0.75;
 
+	// Verse
+
+	static const uint32 VerseTimelineIndex = 2;
+
+	// CPU
+
+	TMap<uint32, uint32> CpuThreadTimelineIndexMap; // CPU Thread Id --> timeline index, in Timelines
 };
 
 } // namespace TraceServices
diff --git a/Engine/Source/Developer/TraceServices/Public/TraceServices/Model/TimingProfiler.h b/Engine/Source/Developer/TraceServices/Public/TraceServices/Model/TimingProfiler.h
--- a/Engine/Source/Developer/TraceServices/Public/TraceServices/Model/TimingProfiler.h
+++ b/Engine/Source/Developer/TraceServices/Public/TraceServices/Model/TimingProfiler.h
@@ -239,23 +239,21 @@
 
 	virtual ~ITimingProfilerProvider() = default;
 
-	virtual bool GetCpuThreadTimelineIndex(uint32 ThreadId, uint32& OutTimelineIndex) const = 0;
-	virtual bool GetGpuTimelineIndex(uint32& OutTimelineIndex) const = 0;
-	virtual bool GetGpu2TimelineIndex(uint32& OutTimelineIndex) const = 0;
-	virtual void EnumerateGpuQueues(TFunctionRef<void(const FGpuQueueInfo&)> Callback) const {}
-	virtual bool GetGpuQueueTimelineIndex(uint32 QueueId, uint32& OutTimelineIndex) const { return false; }
-	virtual bool GetVerseTimelineIndex(uint32& OutTimelineIndex) const { return false; }
+	//////////////////////////////////////////////////
+	// Timelines
 
-	virtual void EnumerateGpuSignalFences(uint32 QueueId, double StartTime, double EndTime, EnumerateGpuSignalFencesCallback Callback) const {}
-	virtual void EnumerateGpuWaitFences(uint32 QueueId, double StartTime, double EndTime, EnumerateGpuWaitFencesCallback Callback) const {}
-	virtual void EnumerateGpuFences(uint32 QueueId, double StartTime, double EndTime, EnumerateGpuFencesCallback Callback) const {}
-	virtual void EnumerateResolvedGpuFences(uint32 QueueId, double StartTime, double EndTime, EnumerateResolvedGpuFencesCallback Callback) const {}
-
 	virtual bool ReadTimeline(uint32 Index, TFunctionRef<void(const Timeline&)> Callback) const = 0;
 	virtual uint32 GetTimelineCount() const = 0;
 	virtual void EnumerateTimelines(TFunctionRef<void(const Timeline&)> Callback) const = 0;
 
+	//////////////////////////////////////////////////
+	// Timers
+
 	virtual void ReadTimers(TFunctionRef<void(const ITimingProfilerTimerReader&)> Callback) const = 0;
+
+	//////////////////////////////////////////////////
+	// Metadata
+
 	virtual uint32 GetOriginalTimerIdFromMetadata(uint32 MetadataTimerId) const { return MetadataTimerId; }
 	virtual TArrayView<const uint8> GetMetadata(uint32 MetadataTimerId) const { return TArrayView<const uint8>(); }
 
@@ -267,6 +265,35 @@
 	*/
 	virtual const FMetadataSpec* GetMetadataSpec(uint32 MetadataSpecId) const { return nullptr; }
 
+	//////////////////////////////////////////////////
+	// GPU
+
+	// Only used for backward compatibility with old GPU Insights.
+	virtual bool GetGpuTimelineIndex(uint32& OutTimelineIndex) const { return false; }
+	// Only used for backward compatibility with old GPU Insights.
+	virtual bool GetGpu2TimelineIndex(uint32& OutTimelineIndex) const { return false; }
+
+	virtual void EnumerateGpuQueues(TFunctionRef<void(const FGpuQueueInfo&)> Callback) const {}
+	virtual bool GetGpuQueueTimelineIndex(uint32 QueueId, uint32& OutTimelineIndex) const { return false; }
+
+	virtual void EnumerateGpuSignalFences(uint32 QueueId, double StartTime, double EndTime, EnumerateGpuSignalFencesCallback Callback) const {}
+	virtual void EnumerateGpuWaitFences(uint32 QueueId, double StartTime, double EndTime, EnumerateGpuWaitFencesCallback Callback) const {}
+	virtual void EnumerateGpuFences(uint32 QueueId, double StartTime, double EndTime, EnumerateGpuFencesCallback Callback) const {}
+	virtual void EnumerateResolvedGpuFences(uint32 QueueId, double StartTime, double EndTime, EnumerateResolvedGpuFencesCallback Callback) const {}
+
+	//////////////////////////////////////////////////
+	// Verse
+
+	virtual bool GetVerseTimelineIndex(uint32& OutTimelineIndex) const { return false; }
+
+	//////////////////////////////////////////////////
+	// CPU
+
+	virtual bool GetCpuThreadTimelineIndex(uint32 ThreadId, uint32& OutTimelineIndex) const = 0;
+
+	//////////////////////////////////////////////////
+	// Aggregation
+
 	/**
 	* Creates a table of aggregated stats.
 	*
@@ -293,17 +320,14 @@
 		return CreateButterfly(Params);
 	}
 
+	//////////////////////////////////////////////////
+	// Misc
+
 	/*
 	* Gets the Ratio of Threads to use.
 	* @return The ratio scaled between 0-1.
 	*/
 	virtual double GetRatioOfThreadsToUse() const { return 0.75; }
-
-	/*
-	* Sets the Ratio of Threads to be used.
-	* @param InRatioOfThreadsToUse The new ratio of threads to use.
-	*/
-	virtual void const SetRatioOfThreadsToUse(double InRatioOfThreadsToUse) const {}
 };
 
 /*
@@ -315,38 +339,8 @@
 public:
 	virtual ~IEditableTimingProfilerProvider() = default;
 
-	/*
-	* Adds/registers a new CPU timer.
-	*
-	* @param Name	The name attached to the timer.
-	* @param File	The source file in which the timer is defined.
-	* @param Line	The line number of the source file in which the timer is defined.
-	*
-	* @return The identity of the CPU timer.
-	*/
-	virtual uint32 AddCpuTimer(FStringView Name, const TCHAR* File, uint32 Line) = 0;
+	//////////////////////////////////////////////////
+	// Timers
 
 	/*
-	* Adds/registers a new GPU timer.
-	*
-	* @param Name	The name attached to the timer.
-	* @param File	The source file in which the timer is defined.
-	* @param Line	The line number of the source file in which the timer is defined.
-	*
-	* @return The identity of the GPU timer.
-	*/
-	virtual uint32 AddGpuTimer(FStringView Name, const TCHAR* File, uint32 Line) { return 0; }
-
-	/*
-	* Adds/registers a new Verse timer.
-	*
-	* @param Name	The name attached to the timer.
-	* @param File	The source file in which the timer is defined.
-	* @param Line	The line number of the source file in which the timer is defined.
-	*
-	* @return The identity of the GPU timer.
-	*/
-	virtual uint32 AddVerseTimer(FStringView Name, const TCHAR* File, uint32 Line) { return 0; }
-
-	/*
 	* Updates an existing timer with information. Some information is unavailable when it's created.
@@ -366,6 +360,9 @@
 	*/
 	virtual void SetTimerNameAndLocation(uint32 TimerId, FStringView Name, const TCHAR* File, uint32 Line) = 0;
 
+	//////////////////////////////////////////////////
+	// Metadata
+
 	/*
 	* Sets the metadata spec for an existing timer.
 	*
@@ -394,7 +391,7 @@
 	virtual void SetMetadata(uint32 MetadataTimerId, TArray<uint8>&& Metadata) {}
 	
 	/*
-	* Sets metadata for the specified MetadataTimerId and replaces the OriginialTimerId with a new value. The MetadataTimerId must be a value returned by AddMetadata. The function is meant to be used to replace the metadata at a MetadataTimerId.
+	* Sets metadata for the specified MetadataTimerId and replaces the OriginalTimerId with a new value. The MetadataTimerId must be a value returned by AddMetadata. The function is meant to be used to replace the metadata at a MetadataTimerId.
 	*
 	* @param MetadataTimerId	The identity of the metadata to replace.
 	* @param Metadata			The metadata.
@@ -421,6 +418,20 @@
 	*/
 	virtual uint32 AddMetadataSpec(FMetadataSpec&& Metadata) { return 0; }
 
+	//////////////////////////////////////////////////
+	// GPU
+
 	/*
+	* Adds/registers a new GPU timer.
+	*
+	* @param Name	The name attached to the timer.
+	* @param File	The source file in which the timer is defined.
+	* @param Line	The line number of the source file in which the timer is defined.
+	*
+	* @return The identity of the GPU timer.
+	*/
+	virtual uint32 AddGpuTimer(FStringView Name, const TCHAR* File, uint32 Line) { return 0; }
+
+	/*
 	* Adds a new GPU queue.
 	*
@@ -445,31 +456,76 @@
 	virtual void AddGpuWaitFence(uint32 QueueId, const FGpuWaitFence& WaitFence) {};
 
 	/*
-	* Gets an object to receive ordered timeline events for a CPU thread.
-	*
-	* @param ThreadId	The thread for which the events are for.
-	*
-	* @return The object to receive the serial events for the specified thread.
-	*/
-	virtual IEditableTimeline<FTimingProfilerEvent>& GetCpuThreadEditableTimeline(uint32 ThreadId) = 0;
-
-	/*
-	* Gets an object to receive ordered timeline events for a GPU queue.
+	* Gets an object to receive ordered GPU timing events for a GPU queue.
 	*
 	* @param QueueId	The queue for which the events are for.
 	*
-	* @return The object to receive the serial events for the specified queue.
+	* @return The object to receive the serial GPU timing events for the specified queue.
 	*/
 	virtual IEditableTimeline<FTimingProfilerEvent>* GetGpuQueueEditableTimeline(uint32 QueueId) { return nullptr; }
+
+	/*
+	* Gets an object to receive ordered GPU Work timing events for a GPU queue.
+	*
+	* @param QueueId	The queue for which the events are for.
+	*
+	* @return The object to receive the serial GPU Work timing events for the specified queue.
+	*/
 	virtual IEditableTimeline<FTimingProfilerEvent>* GetGpuQueueWorkEditableTimeline(uint32 QueueId) { return nullptr; }
 
+	//////////////////////////////////////////////////
+	// Verse
+
 	/*
-	* Gets an object to receive ordered timeline events for the Verse sampling.
+	* Adds/registers a new Verse timer.
 	*
-	* @return The object to receive the serial events for Verse sampling.
+	* @param Name	The name attached to the timer.
+	* @param File	The source file in which the timer is defined.
+	* @param Line	The line number of the source file in which the timer is defined.
+	*
+	* @return The identity of the Verse timer.
 	*/
+	virtual uint32 AddVerseTimer(FStringView Name, const TCHAR* File, uint32 Line) { return 0; }
+
+	/*
+	* Gets an object to receive ordered Verse timing events for the Verse sampling.
+	*
+	* @return The object to receive the serial Verse timing events for Verse sampling.
+	*/
 	virtual IEditableTimeline<FTimingProfilerEvent>* GetVerseEditableTimeline() { return nullptr; }
 
+	//////////////////////////////////////////////////
+	// CPU
+
 	/*
+	* Adds/registers a new CPU timer.
+	*
+	* @param Name	The name attached to the timer.
+	* @param File	The source file in which the timer is defined.
+	* @param Line	The line number of the source file in which the timer is defined.
+	*
+	* @return The identity of the CPU timer.
+	*/
+	virtual uint32 AddCpuTimer(FStringView Name, const TCHAR* File, uint32 Line) = 0;
+
+	/*
+	* Gets an object to receive ordered CPU timing events for a CPU thread.
+	*
+	* @param ThreadId	The thread for which the events are for.
+	*
+	* @return The object to receive the serial CPU timing events for the specified thread.
+	*/
+	virtual IEditableTimeline<FTimingProfilerEvent>& GetCpuThreadEditableTimeline(uint32 ThreadId) = 0;
+
+	//////////////////////////////////////////////////
+	// Misc
+
+	/*
+	* Sets the Ratio of Threads to be used.
+	* @param InRatioOfThreadsToUse The new ratio of threads to use.
+	*/
+	virtual void SetRatioOfThreadsToUse(double InRatioOfThreadsToUse) {}
+
+	/*
 	* Gets the read provider.
 	* @returns The read only provider or nullptr if not available.
