[Insights] Changed ILoadTimeProfilerProvider::GetTimelineCount() to return uint32.

#rnx

diff --git a/Engine/Source/Developer/TraceServices/Private/Model/LoadTimeProfilerPrivate.h b/Engine/Source/Developer/TraceServices/Private/Model/LoadTimeProfilerPrivate.h
--- a/Engine/Source/Developer/TraceServices/Private/Model/LoadTimeProfilerPrivate.h
+++ b/Engine/Source/Developer/TraceServices/Private/Model/LoadTimeProfilerPrivate.h
@@ -24,7 +24,7 @@ public:
 	explicit FLoadTimeProfilerProvider(IAnalysisSession& Session, IEditableCounterProvider& EditableCounterProvider);
 	virtual ~FLoadTimeProfilerProvider() {}
 
-	virtual uint64 GetTimelineCount() const override { return CpuTimelines.Num(); }
+	virtual uint32 GetTimelineCount() const override { return CpuTimelines.Num(); }
 	virtual bool GetCpuThreadTimelineIndex(uint32 ThreadId, uint32& OutTimelineIndex) const override;
 	virtual bool ReadTimeline(uint32 Index, TFunctionRef<void(const CpuTimeline&)> Callback) const override;
 	virtual ITable<FLoadTimeProfilerAggregatedStats>* CreateEventAggregation(double IntervalStart, double IntervalEnd) const override;
diff --git a/Engine/Source/Developer/TraceServices/Public/TraceServices/Model/LoadTimeProfiler.h b/Engine/Source/Developer/TraceServices/Public/TraceServices/Model/LoadTimeProfiler.h
--- a/Engine/Source/Developer/TraceServices/Public/TraceServices/Model/LoadTimeProfiler.h
+++ b/Engine/Source/Developer/TraceServices/Public/TraceServices/Model/LoadTimeProfiler.h
@@ -170,7 +170,7 @@ public:
 	typedef ITimeline<FLoadTimeProfilerCpuEvent> CpuTimeline;
 
 	virtual ~ILoadTimeProfilerProvider() = default;
-	virtual uint64 GetTimelineCount() const = 0;
+	virtual uint32 GetTimelineCount() const = 0;
 	virtual bool GetCpuThreadTimelineIndex(uint32 ThreadId, uint32& OutTimelineIndex) const = 0;
 	virtual bool ReadTimeline(uint32 Index, TFunctionRef<void(const CpuTimeline&)> Callback) const = 0;
 	virtual ITable<FLoadTimeProfilerAggregatedStats>* CreateEventAggregation(double IntervalStart, double IntervalEnd) const = 0;
