[Insights] Improved TimingProfiler analysis API (IEditableTimingProfilerProvider and ITimingProfilerProvider) to allow better control over multiple types of timers.

diff --git a/Engine/Source/Developer/TraceServices/Private/Model/TimingProfiler.cpp b/Engine/Source/Developer/TraceServices/Private/Model/TimingProfiler.cpp
--- a/Engine/Source/Developer/TraceServices/Private/Model/TimingProfiler.cpp
+++ b/Engine/Source/Developer/TraceServices/Private/Model/TimingProfiler.cpp
@@ -15,6 +15,11 @@ namespace TraceServices
 FTimingProfilerProvider::FTimingProfilerProvider(IAnalysisSession& InSession)
 	: Session(InSession)
 {
+	for (uint32 Index = 0; Index < uint32(ETimingProfilerTimerType::Count); ++Index)
+	{
+		bHasData[Index] = false;
+	}
+
 	// Adds the old GPU timelines.
 	Timelines.Add(MakeShared<TimelineInternal>(Session.GetLinearAllocator())); // GPU Index 0
 	Timelines.Add(MakeShared<TimelineInternal>(Session.GetLinearAllocator())); // GPU Index 1
@@ -40,6 +45,12 @@ FTimingProfilerProvider::FTimingProfilerProvider(IAnalysisSession& InSession)
 		AddColumn(&FTimingProfilerAggregatedStats::MaxExclusiveTime, TEXT("E.Max")).
 		AddColumn(&FTimingProfilerAggregatedStats::AverageExclusiveTime, TEXT("E.Avg")).
 		AddColumn(&FTimingProfilerAggregatedStats::MedianExclusiveTime, TEXT("E.Med"));
+
+	// Add a default timer.
+	FTimingProfilerTimer& Timer = Timers.AddDefaulted_GetRef();
+	Timer.Id = 0;
+	Timer.Name = TEXT("<default>");
+	Timer.Type = ETimingProfilerTimerType::CpuScope;
 }
 
 FTimingProfilerProvider::~FTimingProfilerProvider()
@@ -50,14 +61,14 @@ bool FTimingProfilerProvider::HasCpuTiming() const
 {
 	Session.ReadAccessCheck();
 
-	return bHasCpuTiming;
+	return bHasData[uint32(ETimingProfilerTimerType::CpuScope)];
 }
 
 bool FTimingProfilerProvider::HasGpuTiming() const
 {
 	Session.ReadAccessCheck();
 
-	return bHasGpuTiming ||
+	return bHasData[uint32(ETimingProfilerTimerType::GpuScope)] ||
 		Timelines[GpuTimelineIndex]->GetEventCount() > 0 ||
 		Timelines[Gpu2TimelineIndex]->GetEventCount() > 0;
 }
@@ -66,72 +77,119 @@ bool FTimingProfilerProvider::HasVerseTiming() const
 {
 	Session.ReadAccessCheck();
 
-	return bHasVerseTiming ||
+	return bHasData[uint32(ETimingProfilerTimerType::VerseSampling)] ||
 		Timelines[VerseTimelineIndex]->GetEventCount() > 0;
 }
 
-uint32 FTimingProfilerProvider::AddCpuTimer(FStringView Name, const TCHAR* File, uint32 Line)
+void FTimingProfilerProvider::ReadTimers(TFunctionRef<void(const ITimingProfilerTimerReader&)> Callback) const
 {
-	Session.WriteAccessCheck();
+	Session.ReadAccessCheck();
+	Callback(*this);
+}
 
-	bHasCpuTiming = true;
-	FTimingProfilerTimer& Timer = AddTimerInternal(Name, File, Line, ETimingProfilerTimerType::CPU);
-	return Timer.Id;
+uint32 FTimingProfilerProvider::GetTimerCount() const
+{
+	return Timers.Num();
 }
 
-uint32 FTimingProfilerProvider::AddGpuTimer(FStringView Name, const TCHAR* File, uint32 Line)
+const FTimingProfilerTimer* FTimingProfilerProvider::GetTimer(uint32 TimerId) const
+{
+	if (int32(TimerId) < 0)
+	{
+		TimerId = ~TimerId;
+		TimerId = Metadatas[TimerId].TimerId;
+	}
+	return (TimerId < uint32(Timers.Num())) ? Timers.GetData() + TimerId : nullptr;
+}
+
+uint32 FTimingProfilerProvider::AddTimer(ETimingProfilerTimerType TimerType)
 {
 	Session.WriteAccessCheck();
 
-	bHasGpuTiming = true;
-	FTimingProfilerTimer& Timer = AddTimerInternal(Name, File, Line, ETimingProfilerTimerType::GPU);
+	check(uint32(TimerType) < uint32(ETimingProfilerTimerType::Count));
+	bHasData[uint32(TimerType)] = true;
+
+	FTimingProfilerTimer& Timer = Timers.AddDefaulted_GetRef();
+	Timer.Id = Timers.Num() - 1;
+	Timer.Type = TimerType;
 	return Timer.Id;
 }
 
-uint32 FTimingProfilerProvider::AddVerseTimer(FStringView Name, const TCHAR* File, uint32 Line)
+uint32 FTimingProfilerProvider::AddTimer(ETimingProfilerTimerType TimerType, FStringView TimerName)
 {
 	Session.WriteAccessCheck();
 
-	bHasVerseTiming = true;
-	FTimingProfilerTimer& Timer = AddTimerInternal(Name, File, Line, ETimingProfilerTimerType::Verse);
+	check(uint32(TimerType) < uint32(ETimingProfilerTimerType::Count));
+	bHasData[uint32(TimerType)] = true;
+
+	FTimingProfilerTimer& Timer = Timers.AddDefaulted_GetRef();
+	Timer.Id = Timers.Num() - 1;
+	Timer.Name = Session.StoreString(TimerName);
+	Timer.Type = TimerType;
 	return Timer.Id;
 }
 
-FTimingProfilerTimer& FTimingProfilerProvider::AddTimerInternal(FStringView Name, const TCHAR* File, uint32 Line, ETimingProfilerTimerType TimerType)
+uint32 FTimingProfilerProvider::AddTimer(ETimingProfilerTimerType TimerType, FStringView TimerName, const TCHAR* File, uint32 Line)
 {
+	Session.WriteAccessCheck();
+
+	check(uint32(TimerType) < uint32(ETimingProfilerTimerType::Count));
+	bHasData[uint32(TimerType)] = true;
+
 	FTimingProfilerTimer& Timer = Timers.AddDefaulted_GetRef();
 	Timer.Id = Timers.Num() - 1;
-	Timer.Name = Session.StoreString(Name);
+	Timer.Name = Session.StoreString(TimerName);
 	Timer.File = File ? Session.StoreString(File) : nullptr;
 	Timer.Line = Line;
 	Timer.Type = TimerType;
-	return Timer;
+	return Timer.Id;
 }
 
 void FTimingProfilerProvider::SetTimerName(uint32 TimerId, FStringView Name)
 {
 	Session.WriteAccessCheck();
 
-	FTimingProfilerTimer& Timer = Timers[TimerId];
-	Timer.Name = Session.StoreString(Name);
+	if (ensure(TimerId < uint32(Timers.Num())))
+	{
+		FTimingProfilerTimer& Timer = Timers[TimerId];
+		Timer.Name = Session.StoreString(Name);
+	}
 }
 
 void FTimingProfilerProvider::SetTimerNameAndLocation(uint32 TimerId, FStringView Name, const TCHAR* File, uint32 Line)
 {
 	Session.WriteAccessCheck();
 
-	FTimingProfilerTimer& Timer = Timers[TimerId];
-	Timer.Name = Session.StoreString(Name);
-	Timer.File = File ? Session.StoreString(File) : nullptr;
-	Timer.Line = Line;
+	if (ensure(TimerId < uint32(Timers.Num())))
+	{
+		FTimingProfilerTimer& Timer = Timers[TimerId];
+		Timer.Name = Session.StoreString(Name);
+		Timer.File = File ? Session.StoreString(File) : nullptr;
+		Timer.Line = Line;
+	}
+}
+
+void FTimingProfilerProvider::SetTimerLocation(uint32 TimerId, const TCHAR* File, uint32 Line)
+{
+	Session.WriteAccessCheck();
+
+	if (ensure(TimerId < uint32(Timers.Num())))
+	{
+		FTimingProfilerTimer& Timer = Timers[TimerId];
+		Timer.File = File ? Session.StoreString(File) : nullptr;
+		Timer.Line = Line;
+	}
 }
 
 void FTimingProfilerProvider::SetMetadataSpec(uint32 TimerId, uint32 MetadataSpecId)
 {
 	Session.WriteAccessCheck();
 
-	FTimingProfilerTimer& Timer = Timers[TimerId];
-	Timer.MetadataSpecId = MetadataSpecId;
+	if (ensure(TimerId < uint32(Timers.Num())))
+	{
+		FTimingProfilerTimer& Timer = Timers[TimerId];
+		Timer.MetadataSpecId = MetadataSpecId;
+	}
 }
 
 uint32 FTimingProfilerProvider::AddMetadata(uint32 OriginalTimerId, TArray<uint8>&& Metadata)
@@ -273,7 +331,7 @@ IEditableTimeline<FTimingProfilerEvent>& FTimingProfilerProvider::GetCpuThreadEd
 		uint32 TimelineIndex = Timelines.Num();
 		CpuThreadTimelineIndexMap.Add(ThreadId, TimelineIndex);
 		Timelines.Add(Timeline);
-		bHasCpuTiming = true;
+		bHasData[uint32(ETimingProfilerTimerType::CpuScope)] = true;
 		return Timeline.Get();
 	}
 	else
@@ -286,6 +344,7 @@ FTimingProfilerProvider::TimelineInternal& FTimingProfilerProvider::EditGpuTimel
 {
 	Session.WriteAccessCheck();
 
+	bHasData[uint32(ETimingProfilerTimerType::GpuScope)] = true;
 	return Timelines[GpuTimelineIndex].Get();
 }
 
@@ -293,6 +352,7 @@ FTimingProfilerProvider::TimelineInternal& FTimingProfilerProvider::EditGpu2Time
 {
 	Session.WriteAccessCheck();
 
+	bHasData[uint32(ETimingProfilerTimerType::GpuScope)] = true;
 	return Timelines[Gpu2TimelineIndex].Get();
 }
 
@@ -321,7 +381,8 @@ void FTimingProfilerProvider::AddGpuQueue(uint32 QueueId, uint8 GPU, uint8 Index
 		GpuQueue.WorkTimelineIndex = Timelines.Num();
 		TSharedRef<TimelineInternal> WorkTimeline = MakeShared<TimelineInternal>(Session.GetLinearAllocator());
 		Timelines.Add(WorkTimeline);
-		bHasGpuTiming = true;
+
+		bHasData[uint32(ETimingProfilerTimerType::GpuScope)] = true;
 	}
 }
 
@@ -342,7 +403,6 @@ void FTimingProfilerProvider::AddGpuSignalFence(uint32 QueueId, const FGpuSignal
 		}
 
 		Data.SignalFenceArray.Add(Fence);
-		bHasGpuTiming = true;
 	}
 }
 
@@ -363,7 +423,6 @@ void FTimingProfilerProvider::AddGpuWaitFence(uint32 QueueId, const FGpuWaitFenc
 		}
 
 		Data.WaitFenceArray.Add(Fence);
-		bHasGpuTiming = true;
 	}
 }
 
@@ -542,7 +601,7 @@ void FTimingProfilerProvider::EnumerateResolvedGpuFences(uint32 QueueId, double
 		}
 		++SignalFenceIndex;
 	}
-	
+
 	const TArray<FGpuWaitFence>& WaitFenceArray = QueueData.WaitFenceArray;
 
 	int32 WaitFenceIndex = Algo::LowerBoundBy(WaitFenceArray, StartTime,
@@ -698,27 +757,6 @@ void FTimingProfilerProvider::EnumerateTimelines(TFunctionRef<void(const Timelin
 	}
 }
 
-const FTimingProfilerTimer* FTimingProfilerProvider::GetTimer(uint32 TimerId) const
-{
-	if (int32(TimerId) < 0)
-	{
-		TimerId = ~TimerId;
-		TimerId = Metadatas[TimerId].TimerId;
-	}
-	return (TimerId < uint32(Timers.Num())) ? Timers.GetData() + TimerId : nullptr;
-}
-
-uint32 FTimingProfilerProvider::GetTimerCount() const
-{
-	return Timers.Num();
-}
-
-void FTimingProfilerProvider::ReadTimers(TFunctionRef<void(const ITimingProfilerTimerReader&)> Callback) const
-{
-	Session.ReadAccessCheck();
-	Callback(*this);
-}
-
 ITable<FTimingProfilerAggregatedStats>* FTimingProfilerProvider::CreateAggregation(const FCreateAggregationParams& Params) const
 {
 	Session.ReadAccessCheck();
diff --git a/Engine/Source/Developer/TraceServices/Private/Model/TimingProfilerPrivate.h b/Engine/Source/Developer/TraceServices/Private/Model/TimingProfilerPrivate.h
--- a/Engine/Source/Developer/TraceServices/Private/Model/TimingProfilerPrivate.h
+++ b/Engine/Source/Developer/TraceServices/Private/Model/TimingProfilerPrivate.h
@@ -45,9 +45,9 @@ public:
 
 	// Metadata
 
-	virtual uint32 GetOriginalTimerIdFromMetadata(uint32 MetadataTimerId) const override;
-	virtual TArrayView<const uint8> GetMetadata(uint32 MetadataTimerId) const override;
-	virtual const FMetadataSpec* GetMetadataSpec(uint32 MetadataSpecId) const;
+	virtual uint32 GetOriginalTimerIdFromMetadata(uint32 MetadataTimerId) const override; // also in ITimingProfilerTimerReader
+	virtual TArrayView<const uint8> GetMetadata(uint32 MetadataTimerId) const override; // also in ITimingProfilerTimerReader
+	virtual const FMetadataSpec* GetMetadataSpec(uint32 MetadataSpecId) const override;
 
 	// GPU
 
@@ -83,18 +83,26 @@ public:
 	//////////////////////////////////////////////////
 	// ITimingProfilerTimerReader
 
-	virtual const FTimingProfilerTimer* GetTimer(uint32 TimerId) const override;
 	virtual uint32 GetTimerCount() const override;
-	//virtual uint32 GetOriginalTimerIdFromMetadata(uint32 MetadataTimerId) const override;
-	//virtual TArrayView<const uint8> GetMetadata(uint32 MetadataTimerId) const override;
+
+	// The returned FTimingProfilerTimer* is only valid under the current read lock.
+	virtual const FTimingProfilerTimer* GetTimer(uint32 TimerId) const override;
+
+	//virtual uint32 GetOriginalTimerIdFromMetadata(uint32 MetadataTimerId) const override; // also in ITimingProfilerProvider
+	//virtual TArrayView<const uint8> GetMetadata(uint32 MetadataTimerId) const override; // also in ITimingProfilerProvider
 
 	//////////////////////////////////////////////////
 	// IEditableTimingProfilerProvider
 
 	// Timers
 
+	virtual uint32 AddTimer(ETimingProfilerTimerType Type) override;
+	virtual uint32 AddTimer(ETimingProfilerTimerType Type, FStringView Name) override;
+	virtual uint32 AddTimer(ETimingProfilerTimerType Type, FStringView Name, const TCHAR* File, uint32 Line) override;
+
 	virtual void SetTimerName(uint32 TimerId, FStringView Name) override;
 	virtual void SetTimerNameAndLocation(uint32 TimerId, FStringView Name, const TCHAR* File, uint32 Line) override;
+	virtual void SetTimerLocation(uint32 TimerId, const TCHAR* File, uint32 Line) override;
 
 	// Metadata
 
@@ -107,7 +115,6 @@ public:
 
 	// GPU
 
-	virtual uint32 AddGpuTimer(FStringView Name, const TCHAR* File = nullptr, uint32 Line = 0) override;
 	virtual void AddGpuQueue(uint32 QueueId, uint8 GPU, uint8 Index, uint8 Type, const TCHAR* Name) override;
 	virtual void AddGpuSignalFence(uint32 QueueId, const FGpuSignalFence& SignalFence) override;
 	virtual void AddGpuWaitFence(uint32 QueueId, const FGpuWaitFence& WaitFence) override;
@@ -116,12 +123,10 @@ public:
 
 	// Verse
 
-	virtual uint32 AddVerseTimer(FStringView Name, const TCHAR* File = nullptr, uint32 Line = 0) override;
 	virtual IEditableTimeline<FTimingProfilerEvent>* GetVerseEditableTimeline() override;
 
 	// CPU
 
-	virtual uint32 AddCpuTimer(FStringView Name, const TCHAR* File = nullptr, uint32 Line = 0) override;
 	virtual IEditableTimeline<FTimingProfilerEvent>& GetCpuThreadEditableTimeline(uint32 ThreadId) override;
 
 	// Misc
@@ -135,7 +140,7 @@ public:
 	TimelineInternal& EditGpu2Timeline();
 
 private:
-	FTimingProfilerTimer& AddTimerInternal(FStringView Name, const TCHAR* File, uint32 Line, ETimingProfilerTimerType TimerType);
+	FTimingProfilerTimer& AddTimerInternal(ETimingProfilerTimerType TimerType, FStringView Name, const TCHAR* File, uint32 Line);
 
 	struct FMetadata
 	{
@@ -151,9 +156,7 @@ private:
 	TTableLayout<FTimingProfilerAggregatedStats> AggregatedStatsTableLayout;
 	double RatioOfThreadsToUse = 0.75;
 
-	bool bHasGpuTiming = false;
-	bool bHasVerseTiming = false;
-	bool bHasCpuTiming = false;
+	bool bHasData[uint32(ETimingProfilerTimerType::Count)];
 
 	// GPU
 
diff --git a/Engine/Source/Developer/TraceServices/Public/TraceServices/Model/TimingProfiler.h b/Engine/Source/Developer/TraceServices/Public/TraceServices/Model/TimingProfiler.h
--- a/Engine/Source/Developer/TraceServices/Public/TraceServices/Model/TimingProfiler.h
+++ b/Engine/Source/Developer/TraceServices/Public/TraceServices/Model/TimingProfiler.h
@@ -29,9 +29,16 @@ struct FMetadataSpec
 
 enum class ETimingProfilerTimerType : uint8
 {
-	CPU,
-	GPU,
-	Verse,
+	CpuScope,
+	CpuSampling,
+	GpuScope,
+	VerseSampling,
+
+	Count,
+
+	CPU = CpuScope, // temporary backward compatibility
+	GPU = GpuScope, // temporary backward compatibility
+	Verse = VerseSampling, // temporary backward compatibility
 };
 
 struct FTimingProfilerTimer
@@ -176,9 +183,22 @@ public:
 class ITimingProfilerTimerReader
 {
 public:
-	virtual const FTimingProfilerTimer* GetTimer(uint32 TimerId) const = 0;
+	/**
+	 * Gets the number of timers.
+	 * A valid TimerId is an index in the [0 .. TimerCount-1] range.
+	 */
 	virtual uint32 GetTimerCount() const = 0;
+
+	/**
+	 * Gets the timer info.
+	 * The returned pointer is only valid in the current ITimingProfilerTimerReader lock scope.
+	 * @param TimerId A valid timer id (index in the [0 .. TimerCount-1] range) or a bit inverted metadata id (index in the metadata table).
+	 * @return The timer info.
+	 */
+	virtual const FTimingProfilerTimer* GetTimer(uint32 TimerId) const = 0;
+
 	virtual uint32 GetOriginalTimerIdFromMetadata(uint32 MetadataTimerId) const { return MetadataTimerId; }
+
 	virtual TArrayView<const uint8> GetMetadata(uint32 MetadataTimerId) const { return TArrayView<const uint8>(); }
 };
 
@@ -256,15 +276,15 @@ public:
 	//////////////////////////////////////////////////
 	// Metadata
 
-	virtual uint32 GetOriginalTimerIdFromMetadata(uint32 MetadataTimerId) const { return MetadataTimerId; }
-	virtual TArrayView<const uint8> GetMetadata(uint32 MetadataTimerId) const { return TArrayView<const uint8>(); }
+	virtual uint32 GetOriginalTimerIdFromMetadata(uint32 MetadataTimerId) const { return MetadataTimerId; } // also in ITimingProfilerTimerReader
+	virtual TArrayView<const uint8> GetMetadata(uint32 MetadataTimerId) const { return TArrayView<const uint8>(); } // also in ITimingProfilerTimerReader
 
 	/**
-	* Get the metadata spec associated with the MetadataSpecId.
-	*
-	* @param MetadataSpecId The Metadata spec id.
-	* @return A pointer to the FMetadataSpec or nullptr if no spec is associated with the provided Id. The pointer is only valid in the same Session Read Scope as the function call.
-	*/
+	 * Get the metadata spec associated with the MetadataSpecId.
+	 *
+	 * @param MetadataSpecId The Metadata spec id.
+	 * @return A pointer to the FMetadataSpec or nullptr if no spec is associated with the provided Id. The pointer is only valid in the same Session Read Scope as the function call.
+	 */
 	virtual const FMetadataSpec* GetMetadataSpec(uint32 MetadataSpecId) const { return nullptr; }
 
 	//////////////////////////////////////////////////
@@ -303,17 +323,17 @@ public:
 	// Aggregation
 
 	/**
-	* Creates a table of aggregated stats.
-	*
-	* @param Params	The params for the aggregation.
-	*/
+	 * Creates a table of aggregated stats.
+	 *
+	 * @param Params	The params for the aggregation.
+	 */
 	virtual ITable<FTimingProfilerAggregatedStats>* CreateAggregation(const FCreateAggregationParams& Params) const { return nullptr; }
 
 	/**
-	* Creates a butterfly aggregation.
-	*
-	* @param Params	The params for the aggregation.
-	*/
+	 * Creates a butterfly aggregation.
+	 *
+	 * @param Params	The params for the aggregation.
+	 */
 	virtual ITimingProfilerButterfly* CreateButterfly(const FCreateButterflyParams& Params) const { return nullptr; }
 
 	UE_DEPRECATED(5.6, "Use FCreateButterflyParams instead.")
@@ -331,16 +351,13 @@ public:
 	//////////////////////////////////////////////////
 	// Misc
 
-	/*
-	* Gets the Ratio of Threads to use.
-	* @return The ratio scaled between 0-1.
-	*/
+	/**
+	 * Gets the Ratio of Threads to use.
+	 * @return The ratio scaled between 0-1.
+	 */
 	virtual double GetRatioOfThreadsToUse() const { return 0.75; }
 };
 
-/*
-* An interface that can consume timeline CpuProfiler events from a session.
-*/
 class IEditableTimingProfilerProvider
 	: public IEditableProvider
 {
@@ -350,194 +367,262 @@ public:
 	//////////////////////////////////////////////////
 	// Timers
 
-	/*
-	* Updates an existing timer with information. Some information is unavailable when it's created.
-	*
-	* @param TimerId	The identity of the timer to update.
-	* @param Name		The name attached to the timer.
-	*/
+	/**
+	 * Adds/registers a new timer.
+	 *
+	 * @param Type	The type of the new timer.
+	 * @return		The identity of the new timer.
+	 */
+	virtual uint32 AddTimer(ETimingProfilerTimerType Type) { return 0; }
+
+	/**
+	 * Adds/registers a new timer.
+	 *
+	 * @param Type	The type of the new timer.
+	 * @param Name	The name attached to the timer.
+	 * @return		The identity of the new timer.
+	 */
+	virtual uint32 AddTimer(ETimingProfilerTimerType Type, FStringView Name)
+	{
+		const uint32 TimerId = AddTimer(Type);
+		if (TimerId != 0)
+		{
+			SetTimerName(TimerId, Name);
+		}
+		return TimerId;
+	}
+
+	/**
+	 * Adds/registers a new timer.
+	 *
+	 * @param Type	The type of the new timer.
+	 * @param Name	The name attached to the timer.
+	 * @param File	The source file in which the timer is defined.
+	 * @param Line	The line number of the source file in which the timer is defined.
+	 * @return		The identity of the new timer.
+	 */
+	virtual uint32 AddTimer(ETimingProfilerTimerType Type, FStringView Name, const TCHAR* File, uint32 Line)
+	{
+		const uint32 TimerId = AddTimer(Type);
+		if (TimerId != 0)
+		{
+			SetTimerNameAndLocation(TimerId, Name, File, Line);
+		}
+		return TimerId;
+	}
+
+	/**
+	 * Updates an existing timer with information. Some information is unavailable when it's created.
+	 *
+	 * @param TimerId	The identity of the timer to update.
+	 * @param Name		The name attached to the timer.
+	 */
 	virtual void SetTimerName(uint32 TimerId, FStringView Name) = 0;
 
-	/*
-	* Updates an existing timer with information. Some information is unavailable when it's created.
-	*
-	* @param TimerId	The identity of the timer to update.
-	* @param Name		The name attached to the timer.
-	* @param File		The source file in which the timer is defined.
-	* @param Line		The line number of the source file in which the timer is defined.
-	*/
-	virtual void SetTimerNameAndLocation(uint32 TimerId, FStringView Name, const TCHAR* File, uint32 Line) = 0;
+	/**
+	 * Updates an existing timer with information. Some information is unavailable when it's created.
+	 *
+	 * @param TimerId	The identity of the timer to update.
+	 * @param Name		The name attached to the timer.
+	 * @param File		The source file in which the timer is defined.
+	 * @param Line		The line number of the source file in which the timer is defined.
+	 */
+	virtual void SetTimerNameAndLocation(uint32 TimerId, FStringView Name, const TCHAR* File, uint32 Line)
+	{
+		SetTimerName(TimerId, Name);
+		SetTimerLocation(TimerId, File, Line);
+	}
+
+	/**
+	 * Updates an existing timer with information. Some information is unavailable when it's created.
+	 *
+	 * @param TimerId	The identity of the timer to update.
+	 * @param File		The source file in which the timer is defined.
+	 * @param Line		The line number of the source file in which the timer is defined.
+	 */
+	virtual void SetTimerLocation(uint32 TimerId, const TCHAR* File, uint32 Line) {}
 
 	//////////////////////////////////////////////////
 	// Metadata
 
-	/*
-	* Sets the metadata spec for an existing timer.
-	*
-	* @param TimerId		The identity of the timer to update.
-	* @param MetadataSpecId	The metadata spec to associate with the timer.
-	*/
+	/**
+	 * Sets the metadata spec for an existing timer.
+	 *
+	 * @param TimerId			The identity of the timer to update.
+	 * @param MetadataSpecId	The metadata spec to associate with the timer.
+	 */
 	virtual void SetMetadataSpec(uint32 TimerId, uint32 MetadataSpecId) {}
 
-	/*
-	* Adds metadata to a CPU or GPU timer.
-	*
-	* @param OriginalTimerId	The identity of the timer to add metadata to.
-	* @param Metadata			The metadata.
-	*
-	* @return The identity of the metadata.
-	*/
+	/**
+	 * Adds metadata to a CPU or GPU timer.
+	 *
+	 * @param OriginalTimerId	The identity of the timer to add metadata to.
+	 * @param Metadata			The metadata.
+	 *
+	 * @return The identity of the metadata.
+	 */
 	virtual uint32 AddMetadata(uint32 OriginalTimerId, TArray<uint8>&& Metadata) = 0;
 
-	/*
-	* Sets metadata for the specified MetadataTimerId. The MetadataTimerId must be a value returned by AddMetadata. The function is meant to be used to replace the metadata at a MetadataTimerId.
-	*
-	* @param MetadataTimerId	The identity of the metadata to replace.
-	* @param Metadata			The metadata.
-	*
-	*/
+	/**
+	 * Sets metadata for the specified MetadataTimerId.
+	 * The MetadataTimerId must be a value returned by AddMetadata.
+	 * The function is meant to be used to replace the metadata at a MetadataTimerId.
+	 *
+	 * @param MetadataTimerId	The identity of the metadata to replace.
+	 * @param Metadata			The metadata.
+	 *
+	 */
 	virtual void SetMetadata(uint32 MetadataTimerId, TArray<uint8>&& Metadata) {}
-	
-	/*
-	* Sets metadata for the specified MetadataTimerId and replaces the OriginalTimerId with a new value. The MetadataTimerId must be a value returned by AddMetadata. The function is meant to be used to replace the metadata at a MetadataTimerId.
-	*
-	* @param MetadataTimerId	The identity of the metadata to replace.
-	* @param Metadata			The metadata.
-	* @param NewTimerId			The new TimerId. Represents the identity of the timer the metadata is attached to.
-	*
-	*/
+
+	/**
+	 * Sets metadata for the specified MetadataTimerId and replaces the OriginalTimerId with a new value.
+	 * The MetadataTimerId must be a value returned by AddMetadata.
+	 * The function is meant to be used to replace the metadata at a MetadataTimerId.
+	 *
+	 * @param MetadataTimerId	The identity of the metadata to replace.
+	 * @param Metadata			The metadata.
+	 * @param NewTimerId		The new TimerId. Represents the identity of the timer the metadata is attached to.
+	 *
+	 */
 	virtual void SetMetadata(uint32 MetadataTimerId, TArray<uint8>&& Metadata, uint32 NewTimerId) {}
 
-	/*
-	* Gets metadata by id.
-	*
-	* @param MetadataTimerId	The identity of the metadata.
-	*
-	* @return The metadata.
-	*/
+	/**
+	 * Gets metadata by id.
+	 *
+	 * @param MetadataTimerId	The identity of the metadata.
+	 *
+	 * @return The metadata.
+	 */
 	virtual TArrayView<uint8> GetEditableMetadata(uint32 MetadataTimerId) = 0;
 
-	/*
-	* Adds a metadata spec to storage.
-	*
-	* @param Metadata	The metadata spec to add.
-	*
-	* @return			The identity of the metadata spec.
-	*/
+	/**
+	 * Adds a metadata spec to storage.
+	 *
+	 * @param Metadata	The metadata spec to add.
+	 *
+	 * @return			The identity of the metadata spec.
+	 */
 	virtual uint32 AddMetadataSpec(FMetadataSpec&& Metadata) { return 0; }
 
 	//////////////////////////////////////////////////
 	// GPU
 
-	/*
-	* Adds/registers a new GPU timer.
-	*
-	* @param Name	The name attached to the timer.
-	* @param File	The source file in which the timer is defined.
-	* @param Line	The line number of the source file in which the timer is defined.
-	*
-	* @return The identity of the GPU timer.
-	*/
-	virtual uint32 AddGpuTimer(FStringView Name, const TCHAR* File, uint32 Line) { return 0; }
-
-	/*
-	* Adds a new GPU queue.
-	*
-	* @param QueueId	The GPU queue for which the events are for.
-	*/
+	/**
+	 * Adds/registers a new GPU timer.
+	 *
+	 * @param Name	The name attached to the timer.
+	 * @param File	The source file in which the timer is defined.
+	 * @param Line	The line number of the source file in which the timer is defined.
+	 *
+	 * @return The identity of the GPU timer.
+	 */
+	virtual uint32 AddGpuTimer(FStringView Name)
+	{
+		return AddTimer(ETimingProfilerTimerType::GpuScope, Name);
+	}
+
+	/**
+	 * Adds a new GPU queue.
+	 *
+	 * @param QueueId	The GPU queue for which the events are for.
+	 */
 	virtual void AddGpuQueue(uint32 QueueId, uint8 GPU, uint8 Index, uint8 Type, const TCHAR* Name) {}
 
-	/*
-	* Adds a new GPU signal fence to a queue.
-	*
-	* @param QueueId		The GPU queue that issues the signal.
-	* @param SignalFence	The signal fence.
-	*/
+	/**
+	 * Adds a new GPU signal fence to a queue.
+	 *
+	 * @param QueueId		The GPU queue that issues the signal.
+	 * @param SignalFence	The signal fence.
+	 */
 	virtual void AddGpuSignalFence(uint32 QueueId, const FGpuSignalFence& SignalFence) {};
 
-	/*
-	* Adds a new GPU wait fence to a queue.
-	*
-	* @param QueueId		The GPU queue that waits.
-	* @param WaitFence		The wait fence.
-	*/
+	/**
+	 * Adds a new GPU wait fence to a queue.
+	 *
+	 * @param QueueId		The GPU queue that waits.
+	 * @param WaitFence		The wait fence.
+	 */
 	virtual void AddGpuWaitFence(uint32 QueueId, const FGpuWaitFence& WaitFence) {};
 
-	/*
-	* Gets an object to receive ordered GPU timing events for a GPU queue.
-	*
-	* @param QueueId	The queue for which the events are for.
-	*
-	* @return The object to receive the serial GPU timing events for the specified queue.
-	*/
+	/**
+	 * Gets an object to receive ordered GPU timing events for a GPU queue.
+	 *
+	 * @param QueueId	The queue for which the events are for.
+	 *
+	 * @return The object to receive the serial GPU timing events for the specified queue.
+	 */
 	virtual IEditableTimeline<FTimingProfilerEvent>* GetGpuQueueEditableTimeline(uint32 QueueId) { return nullptr; }
 
-	/*
-	* Gets an object to receive ordered GPU Work timing events for a GPU queue.
-	*
-	* @param QueueId	The queue for which the events are for.
-	*
-	* @return The object to receive the serial GPU Work timing events for the specified queue.
-	*/
+	/**
+	 * Gets an object to receive ordered GPU Work timing events for a GPU queue.
+	 *
+	 * @param QueueId	The queue for which the events are for.
+	 *
+	 * @return The object to receive the serial GPU Work timing events for the specified queue.
+	 */
 	virtual IEditableTimeline<FTimingProfilerEvent>* GetGpuQueueWorkEditableTimeline(uint32 QueueId) { return nullptr; }
 
 	//////////////////////////////////////////////////
-	// Verse
+	// Verse Sampling
 
-	/*
-	* Adds/registers a new Verse timer.
-	*
-	* @param Name	The name attached to the timer.
-	* @param File	The source file in which the timer is defined.
-	* @param Line	The line number of the source file in which the timer is defined.
-	*
-	* @return The identity of the Verse timer.
-	*/
-	virtual uint32 AddVerseTimer(FStringView Name, const TCHAR* File, uint32 Line) { return 0; }
-
-	/*
-	* Gets an object to receive ordered Verse timing events for the Verse sampling.
-	*
-	* @return The object to receive the serial Verse timing events for Verse sampling.
-	*/
+	/**
+	 * Adds/registers a new Verse Sampling timer.
+	 *
+	 * @param Name	The name attached to the timer.
+	 *
+	 * @return The identity of the CPU timer.
+	 */
+	virtual uint32 AddVerseTimer(FStringView Name)
+	{
+		return AddTimer(ETimingProfilerTimerType::VerseSampling, Name);
+	}
+
+	/**
+	 * Gets an object to receive ordered Verse timing events for the Verse sampling.
+	 *
+	 * @return The object to receive the serial Verse timing events for Verse sampling.
+	 */
 	virtual IEditableTimeline<FTimingProfilerEvent>* GetVerseEditableTimeline() { return nullptr; }
 
 	//////////////////////////////////////////////////
 	// CPU
 
-	/*
-	* Adds/registers a new CPU timer.
-	*
-	* @param Name	The name attached to the timer.
-	* @param File	The source file in which the timer is defined.
-	* @param Line	The line number of the source file in which the timer is defined.
-	*
-	* @return The identity of the CPU timer.
-	*/
-	virtual uint32 AddCpuTimer(FStringView Name, const TCHAR* File, uint32 Line) = 0;
-
-	/*
-	* Gets an object to receive ordered CPU timing events for a CPU thread.
-	*
-	* @param ThreadId	The thread for which the events are for.
-	*
-	* @return The object to receive the serial CPU timing events for the specified thread.
-	*/
+	/**
+	 * Adds/registers a new CPU Scope timer.
+	 *
+	 * @param Name	The name attached to the timer.
+	 * @param File	The source file in which the timer is defined.
+	 * @param Line	The line number of the source file in which the timer is defined.
+	 *
+	 * @return The identity of the CPU timer.
+	 */
+	virtual uint32 AddCpuTimer(FStringView Name, const TCHAR* File, uint32 Line)
+	{
+		return AddTimer(ETimingProfilerTimerType::CpuScope, Name, File, Line);
+	}
+
+	/**
+	 * Gets an object to receive ordered CPU timing events for a CPU thread.
+	 *
+	 * @param ThreadId	The thread for which the events are for.
+	 *
+	 * @return The object to receive the serial CPU timing events for the specified thread.
+	 */
 	virtual IEditableTimeline<FTimingProfilerEvent>& GetCpuThreadEditableTimeline(uint32 ThreadId) = 0;
 
 	//////////////////////////////////////////////////
 	// Misc
 
-	/*
-	* Sets the Ratio of Threads to be used.
-	* @param InRatioOfThreadsToUse The new ratio of threads to use.
-	*/
+	/**
+	 * Sets the Ratio of Threads to be used.
+	 * @param InRatioOfThreadsToUse The new ratio of threads to use.
+	 */
 	virtual void SetRatioOfThreadsToUse(double InRatioOfThreadsToUse) {}
 
-	/*
-	* Gets the read provider.
-	* @returns The read only provider or nullptr if not available.
-	*/
+	/**
+	 * Gets the read provider.
+	 * @returns The read only provider or nullptr if not available.
+	 */
 	virtual const ITimingProfilerProvider* GetReadProvider() const { return nullptr; }
 };
 
diff --git a/Engine/Source/Editor/UnrealEd/Private/Commandlets/SummarizeTraceUtils.cpp b/Engine/Source/Editor/UnrealEd/Private/Commandlets/SummarizeTraceUtils.cpp
--- a/Engine/Source/Editor/UnrealEd/Private/Commandlets/SummarizeTraceUtils.cpp
+++ b/Engine/Source/Editor/UnrealEd/Private/Commandlets/SummarizeTraceUtils.cpp
@@ -303,7 +303,21 @@ void FSummarizeCpuProfilerProvider::AddThread(uint32 Id, const TCHAR* Name, EThr
 	}
 }
 
+uint32 FSummarizeCpuProfilerProvider::AddTimer(TraceServices::ETimingProfilerTimerType Type)
+{
+	if (Type == TraceServices::ETimingProfilerTimerType::CPU)
+	{
+		return AddCpuTimerInternal(FStringView());
+	}
+	return 0;
+}
+
 uint32 FSummarizeCpuProfilerProvider::AddCpuTimer(FStringView Name, const TCHAR* File, uint32 Line)
+{
+	return AddCpuTimerInternal(Name);
+}
+
+uint32 FSummarizeCpuProfilerProvider::AddCpuTimerInternal(FStringView Name)
 {
 	TOptional<FString> ScopeName;
 
@@ -342,11 +356,6 @@ void FSummarizeCpuProfilerProvider::SetTimerName(uint32 TimerId, FStringView Nam
 	}
 }
 
-void FSummarizeCpuProfilerProvider::SetTimerNameAndLocation(uint32 TimerId, FStringView Name, const TCHAR* File, uint32 Line)
-{
-	SetTimerName(TimerId, Name);
-}
-
 uint32 FSummarizeCpuProfilerProvider::AddMetadata(uint32 MasterTimerId, TArray<uint8>&& Metadata)
 {
 	uint32 MetadataId = Metadatas.Num();
diff --git a/Engine/Source/Editor/UnrealEd/Private/Commandlets/SummarizeTraceUtils.h b/Engine/Source/Editor/UnrealEd/Private/Commandlets/SummarizeTraceUtils.h
--- a/Engine/Source/Editor/UnrealEd/Private/Commandlets/SummarizeTraceUtils.h
+++ b/Engine/Source/Editor/UnrealEd/Private/Commandlets/SummarizeTraceUtils.h
@@ -205,9 +205,9 @@ private:
 	virtual void AddThread(uint32 Id, const TCHAR* Name, EThreadPriority Priority) override;
 
 	// TraceServices::IEditableTimingProfilerProvider
+	virtual uint32 AddTimer(TraceServices::ETimingProfilerTimerType Type) override;
 	virtual uint32 AddCpuTimer(FStringView Name, const TCHAR* File, uint32 Line) override;
 	virtual void SetTimerName(uint32 TimerId, FStringView Name) override;
-	virtual void SetTimerNameAndLocation(uint32 TimerId, FStringView Name, const TCHAR* File, uint32 Line) override;
 	virtual uint32 AddMetadata(uint32 MasterTimerId, TArray<uint8>&& Metadata) override;
 	virtual TArrayView<uint8> GetEditableMetadata(uint32 TimerId) override;
 	virtual TraceServices::IEditableTimeline<TraceServices::FTimingProfilerEvent>& GetCpuThreadEditableTimeline(uint32 ThreadId) override;
@@ -222,6 +222,8 @@ private:
 	// Resolve a ScopeId to string
 	const FString* LookupScopeName(uint32 ScopeId);
 
+	uint32 AddCpuTimerInternal(FStringView Name);
+
 private:
 	// The state at any moment of the threads
 	TMap<uint32, TUniquePtr<FThread>> Threads;
