[Insights] Added initial analysis/provider for the new Stack Sampling profiler.

diff --git a/Engine/Source/Developer/TraceServices/Private/Analyzers/PlatformEventTraceAnalysis.cpp b/Engine/Source/Developer/TraceServices/Private/Analyzers/PlatformEventTraceAnalysis.cpp
--- a/Engine/Source/Developer/TraceServices/Private/Analyzers/PlatformEventTraceAnalysis.cpp
+++ b/Engine/Source/Developer/TraceServices/Private/Analyzers/PlatformEventTraceAnalysis.cpp
@@ -28,6 +28,20 @@ void FPlatformEventTraceAnalyzer::OnAnalysisBegin(const FOnAnalysisContext& Cont
 	Builder.RouteEvent(RouteId_ThreadName, "PlatformEvent", "ThreadName");
 }
 
+void FPlatformEventTraceAnalyzer::OnAnalysisEnd()
+{
+	int32 TryCount = 10;
+	while (TryCount > 0)
+	{
+		if (!StackSamplesProvider.UpdatePendingTimers())
+		{
+			break;
+		}
+		FPlatformProcess::Sleep(1.0f);
+		--TryCount;
+	}
+}
+
 bool FPlatformEventTraceAnalyzer::OnEvent(uint16 RouteId, EStyle Style, const FOnEventContext& Context)
 {
 	LLM_SCOPE_BYNAME(TEXT("Insights/FPlatformEventTraceAnalyzer"));
@@ -55,15 +69,19 @@ bool FPlatformEventTraceAnalyzer::OnEvent(uint16 RouteId, EStyle Style, const FO
 
 	case RouteId_StackSample:
 	{
-		FAnalysisSessionEditScope _(Session);
 		const auto& EventData = Context.EventData;
 		double Time = Context.EventTime.AsSeconds(EventData.GetValue<uint64>("Time"));
 		uint32 ThreadId = EventData.GetValue<uint32>("ThreadId");
 		const TArrayReader<uint64>& Addresses = EventData.GetArray<uint64>("Addresses");
-		StackSamplesProvider.Add(ThreadId, Time, Addresses.Num(), Addresses.GetData());
-
-		Session.UpdateDurationSeconds(Time);
-
+		{
+			TraceServices::FProviderEditScopeLock ScopedLock(StackSamplesProvider);
+			StackSamplesProvider.AddStackSample(ThreadId, Time, Addresses.Num(), Addresses.GetData());
+		}
+		StackSamplesProvider.UpdatePendingTimers();
+		{
+			FAnalysisSessionEditScope _(Session);
+			Session.UpdateDurationSeconds(Time);
+		}
 		break;
 	}
 
@@ -75,8 +93,14 @@ bool FPlatformEventTraceAnalyzer::OnEvent(uint16 RouteId, EStyle Style, const FO
 		FStringView Name;
 		if (EventData.GetString("Name", Name))
 		{
-			TraceServices::FProviderEditScopeLock ScopedLock(ContextSwitchesProvider);
-			ContextSwitchesProvider.AddThreadName(ThreadId, ProcessId, Name);
+			{
+				TraceServices::FProviderEditScopeLock ScopedLock(ContextSwitchesProvider);
+				ContextSwitchesProvider.AddThreadName(ThreadId, ProcessId, Name);
+			}
+			{
+				TraceServices::FProviderEditScopeLock ScopedLock(StackSamplesProvider);
+				StackSamplesProvider.AddThreadName(ThreadId, ProcessId, Name);
+			}
 		}
 		break;
 	}
@@ -90,8 +114,14 @@ void FPlatformEventTraceAnalyzer::OnThreadInfo(const FThreadInfo& ThreadInfo)
 {
 	LLM_SCOPE_BYNAME(TEXT("Insights/FPlatformEventTraceAnalyzer"));
 
-	TraceServices::FProviderEditScopeLock ScopedLock(ContextSwitchesProvider);
-	ContextSwitchesProvider.AddThreadInfo(ThreadInfo.GetId(), ThreadInfo.GetSystemId());
+	{
+		TraceServices::FProviderEditScopeLock ScopedLock(ContextSwitchesProvider);
+		ContextSwitchesProvider.AddThreadInfo(ThreadInfo.GetId(), ThreadInfo.GetSystemId());
+	}
+	{
+		//TraceServices::FProviderEditScopeLock ScopedLock(StackSamplesProvider);
+		//StackSamplesProvider.AddThreadInfo(ThreadInfo.GetId(), ThreadInfo.GetSystemId());
+	}
 }
 
 } // namespace TraceServices
diff --git a/Engine/Source/Developer/TraceServices/Private/Analyzers/PlatformEventTraceAnalysis.h b/Engine/Source/Developer/TraceServices/Private/Analyzers/PlatformEventTraceAnalysis.h
--- a/Engine/Source/Developer/TraceServices/Private/Analyzers/PlatformEventTraceAnalysis.h
+++ b/Engine/Source/Developer/TraceServices/Private/Analyzers/PlatformEventTraceAnalysis.h
@@ -20,6 +20,7 @@ public:
 								FContextSwitchesProvider& ContextSwitchesProvider,
 								FStackSamplesProvider& StackSamplesProvider);
 	virtual void OnAnalysisBegin(const FOnAnalysisContext& Context) override;
+	virtual void OnAnalysisEnd() override;
 	virtual bool OnEvent(uint16 RouteId, EStyle Style, const FOnEventContext& Context) override;
 	virtual void OnThreadInfo(const FThreadInfo& ThreadInfo) override;
 
diff --git a/Engine/Source/Developer/TraceServices/Private/Model/StackSamples.cpp b/Engine/Source/Developer/TraceServices/Private/Model/StackSamples.cpp
--- a/Engine/Source/Developer/TraceServices/Private/Model/StackSamples.cpp
+++ b/Engine/Source/Developer/TraceServices/Private/Model/StackSamples.cpp
@@ -2,72 +2,324 @@
 
 #include "TraceServices/Model/StackSamples.h"
 #include "Model/StackSamplesPrivate.h"
-#include "AnalysisServicePrivate.h"
+
 #include "Algo/Sort.h"
+#include "HAL/PlatformMath.h"
+
+// TraceServices
+#include "AnalysisServicePrivate.h"
+#include "TraceServices/Model/Modules.h"
 
 namespace TraceServices
 {
 
+thread_local FProviderLock::FThreadLocalState GStackSamplesProviderLockState;
+
 FStackSamplesProvider::FStackSamplesProvider(IAnalysisSession& InSession)
 	: Session(InSession)
-	, AddressValues(Session.GetLinearAllocator(), 4096)
+	, Threads(Session.GetLinearAllocator(), 256)
+	, StackFrames(Session.GetLinearAllocator(), 4 << 10)
 {
 }
 
 FStackSamplesProvider::~FStackSamplesProvider()
 {
-	for (const auto& ThreadSamples : Threads)
+}
+
+uint32 FStackSamplesProvider::GetStackFrameCount() const
+{
+	ReadAccessCheck();
+
+	return static_cast<uint32>(StackFrames.Num());
+}
+
+void FStackSamplesProvider::EnumerateStackFrames(TFunctionRef<void(const FStackSampleFrame& Frame)> Callback) const
+{
+	ReadAccessCheck();
+
+	auto It = StackFrames.GetIterator();
+	const FStackSampleFrame* StackFrame = It.GetCurrentItem();
+	while (StackFrame)
+	{
+		Callback(*StackFrame);
+		StackFrame = It.NextItem();
+	}
+}
+
+uint32 FStackSamplesProvider::GetThreadCount() const
+{
+	ReadAccessCheck();
+
+	return static_cast<uint32>(Threads.Num());
+}
+
+void FStackSamplesProvider::EnumerateThreads(TFunctionRef<void(const FStackSampleThread& Thread)> Callback) const
+{
+	ReadAccessCheck();
+
+	auto It = Threads.GetIterator();
+	const FStackSamplesProviderThread* Thread = It.GetCurrentItem();
+	while (Thread)
+	{
+		Callback(FStackSampleThread{ Thread->GetSystemThreadId(), Thread->GetName(), Thread->GetTimeline() });
+		Thread = It.NextItem();
+	}
+}
+
+uint32 FStackSamplesProvider::GetTimelineCount() const
+{
+	ReadAccessCheck();
+
+	return static_cast<uint32>(Timelines.Num());
+}
+
+void FStackSamplesProvider::EnumerateTimelines(TFunctionRef<void(const FStackSampleTimeline& Timeline, const FStackSampleThread& Thread)> Callback) const
+{
+	ReadAccessCheck();
+
+	for (const FStackSamplesProviderThread* Thread : Timelines)
+	{
+		check(Thread->GetTimeline() != nullptr);
+		Callback(*Thread->GetTimeline(), FStackSampleThread{ Thread->GetSystemThreadId(), Thread->GetName(), Thread->GetTimeline() });
+	}
+}
+
+const FStackSampleTimeline* FStackSamplesProvider::GetTimeline(uint32 SystemThreadId) const
+{
+	ReadAccessCheck();
+
+	FStackSamplesProviderThread*const* FoundThread = ThreadsBySystemId.Find(SystemThreadId);
+	if (!FoundThread)
+	{
+		return nullptr;
+	}
+
+	return (*FoundThread)->GetTimeline();
+}
+
+FStackSamplesProviderThread& FStackSamplesProvider::GetOrAddThread(uint32 SystemThreadId)
+{
+	FStackSamplesProviderThread** FoundThread = ThreadsBySystemId.Find(SystemThreadId);
+	if (FoundThread)
 	{
-		delete ThreadSamples.Value;
+		return **FoundThread;
 	}
+
+	FStackSamplesProviderThread& NewThread = Threads.EmplaceBack(FStackSamplesProviderThread{ SystemThreadId });
+	ThreadsBySystemId.Add(SystemThreadId, &NewThread);
+	return NewThread;
 }
 
-const TPagedArray<FStackSample>* FStackSamplesProvider::GetStackSamples(uint32 ThreadId) const
+void FStackSamplesProvider::AddThreadName(uint32 SystemThreadId, uint32 SystemProcessId, FStringView Name)
 {
-	Session.ReadAccessCheck();
+	EditAccessCheck();
 
-	auto StackSamples = Threads.Find(ThreadId);
-	return StackSamples ? *StackSamples : nullptr;
+	FStackSamplesProviderThread& Thread = GetOrAddThread(SystemThreadId);
+	Thread.Name = Session.StoreString(Name);
 }
 
-void FStackSamplesProvider::Add(uint32 ThreadId, double Time, uint32 Count, const uint64* Addresses)
+FStackSampleFrame& FStackSamplesProvider::GetOrAddStackFrame(uint64 Address)
 {
-	Session.WriteAccessCheck();
+	FStackSampleFrame** FoundStackFrame = StackFramesByAddress.Find(Address);
+	if (FoundStackFrame)
+	{
+		return **FoundStackFrame;
+	}
+
+	FStackSampleFrame& NewStackFrame = StackFrames.EmplaceBack(FStackSampleFrame{ Address });
+	StackFramesByAddress.Add(Address, &NewStackFrame);
 
-	TPagedArray<FStackSample>** StackSamples = Threads.Find(ThreadId);
-	if (!StackSamples)
+	if (ModuleProvider)
 	{
-		StackSamples = &Threads.Add(ThreadId, new TPagedArray<FStackSample>(Session.GetLinearAllocator(), 4096));
+		// This will return immediately. The result will be empty if the symbol
+		// has not been encountered before, and resolution has been queued up.
+		NewStackFrame.Symbol = ModuleProvider->GetSymbol(Address);
+	}
+	else
+	{
+		static const FResolvedSymbol NeverResolveSymbol(ESymbolQueryResult::NotLoaded, nullptr, nullptr, nullptr, 0, EResolvedSymbolFilterStatus::NotFiltered);
+		NewStackFrame.Symbol = &NeverResolveSymbol;
 	}
 
-	FStackSample& StackSample = (*StackSamples)->PushBack();
-	StackSample.Time = Time;
-	StackSample.Count = Count;
+	CreateTimer(NewStackFrame);
+
+	return NewStackFrame;
+}
 
-	if (AddressValues.NumPages() == 0)
+void FStackSamplesProvider::AddStackSample(uint32 SystemThreadId, double Time, uint32 Count, const uint64* Addresses)
+{
+	EditAccessCheck();
+
+	if (!EditableTimingProfilerProvider)
 	{
-		for (uint32 Index = 0; Index < Count; Index++)
+		EditableTimingProfilerProvider = EditTimingProfilerProvider(Session);
+	}
+	if (!ModuleProvider)
+	{
+		ModuleProvider = const_cast<IModuleProvider*>(ReadModuleProvider(Session));
+	}
+
+	FStackSamplesProviderThread& Thread = GetOrAddThread(SystemThreadId);
+	if (!Thread.Timeline.IsValid())
+	{
+		Thread.Timeline = MakeShared<FStackSamplesTimeline>(Session.GetLinearAllocator());
+		Timelines.Add(&Thread);
+	}
+
+	TArray<uint64> Callstack;
+	Callstack.Reserve(Count);
+	for (int32 StackIndex = int32(Count) - 1; StackIndex >= 0; --StackIndex)
+	{
+		Callstack.Add(Addresses[StackIndex]);
+	}
+
+	const int32 LastSampleStackSize = Thread.LastSampleStack.Num();
+	const int32 NewSampleStackSize = Callstack.Num();
+
+	if (NewSampleStackSize > Thread.MaxStackSize)
+	{
+		Thread.MaxStackSize = NewSampleStackSize;
+	}
+
+	int32 StartNewStackIndex = 0;
+	if (Time - Thread.LastSampleTime < MaxSampleDuration)
+	{
+		for (int32 StackIndex = 0; StackIndex < LastSampleStackSize; ++StackIndex)
 		{
-			AddressValues.EmplaceBack(Addresses[Index]);
+			if ((StackIndex >= NewSampleStackSize) ||
+				(Callstack[StackIndex] != Thread.LastSampleStack[StackIndex]))
+			{
+				break;
+			}
+			StartNewStackIndex = StackIndex + 1;
 		}
-		StackSample.Addresses = &AddressValues.First();
-		return;
 	}
 
-	uint64 Available = AddressValues.GetPageSize() - AddressValues.GetLastPage()->Count;
-	if (Available < Count)
+	const double EndTime = FMath::Min(Thread.LastSampleTime + MaxSampleDuration / 2.0, Time);
+	for (int32 StackIndex = StartNewStackIndex; StackIndex < LastSampleStackSize; ++StackIndex)
+	{
+		Thread.Timeline->AppendEndEvent(EndTime);
+	}
+	for (int32 StackIndex = StartNewStackIndex; StackIndex < NewSampleStackSize; ++StackIndex)
+	{
+		FStackSampleFrame& StackFrame = GetOrAddStackFrame(Callstack[StackIndex]);
+		FTimingProfilerEvent Event;
+		Event.TimerIndex = StackFrame.TimerId;
+		Thread.Timeline->AppendBeginEvent(Time, Event);
+	}
+
+	// LastSampleStack = Callstack
+	if (StartNewStackIndex < LastSampleStackSize)
+	{
+		Thread.LastSampleStack.RemoveAt(StartNewStackIndex, LastSampleStackSize - StartNewStackIndex, EAllowShrinking::No);
+	}
+	for (int32 StackIndex = StartNewStackIndex; StackIndex < NewSampleStackSize; ++StackIndex)
+	{
+		Thread.LastSampleStack.Push(Callstack[StackIndex]);
+	}
+
+	Thread.LastSampleTime = Time;
+	Thread.NumSamples++;
+	TotalNumSamples++;
+}
+
+void FStackSamplesProvider::CreateTimer(FStackSampleFrame& StackFrame)
+{
+	EndEdit();
+
 	{
-		for (uint64 Temp = 0; Temp < Available; Temp++)
+		FAnalysisSessionEditScope _(Session);
+		StackFrame.TimerId = EditableTimingProfilerProvider->AddTimer(ETimingProfilerTimerType::CpuSampling);
+	}
+
+	if (!UpdateTimerName(StackFrame, true))
+	{
+		PendingTimers.Add(&StackFrame);
+	}
+
+	BeginEdit();
+}
+
+bool FStackSamplesProvider::UpdateTimerName(FStackSampleFrame& StackFrame, bool bFirstTime)
+{
+	bool bIsPending = false;
+
+	FString TimerName;
+	if (!StackFrame.Symbol)
+	{
+		TimerName = FString::Printf(TEXT("0x%X"), StackFrame.Address);
+	}
+	else
+	{
+		switch (StackFrame.Symbol->GetResult())
 		{
-			AddressValues.PushBack();
+			case TraceServices::ESymbolQueryResult::Pending:
+				if (bFirstTime)
+				{
+					TimerName = FString::Printf(TEXT("0x%X pending..."), StackFrame.Address);
+				}
+				bIsPending = true;
+				break;
+
+			case TraceServices::ESymbolQueryResult::OK:
+				TimerName = FString::Printf(TEXT("%s (%s / 0x%X)"),
+					StackFrame.Symbol->Name ? StackFrame.Symbol->Name : TEXT("N/A"),
+					StackFrame.Symbol->Module ? StackFrame.Symbol->Module : TEXT("N/A"),
+					StackFrame.Address);
+				{
+					FAnalysisSessionEditScope _(Session);
+					EditableTimingProfilerProvider->SetTimerLocation(StackFrame.TimerId, StackFrame.Symbol->File, StackFrame.Symbol->Line);
+				}
+				break;
+
+			case TraceServices::ESymbolQueryResult::NotLoaded:
+				TimerName = FString::Printf(TEXT("%s / 0x%X NotLoaded"),
+					StackFrame.Symbol->Module ? StackFrame.Symbol->Module : TEXT("N/A"),
+					StackFrame.Address);
+				break;
+
+			case TraceServices::ESymbolQueryResult::Mismatch:
+				TimerName = FString::Printf(TEXT("%s / 0x%X Mismatch"),
+					StackFrame.Symbol->Module ? StackFrame.Symbol->Module : TEXT("N/A"),
+					StackFrame.Address);
+				break;
+
+			case TraceServices::ESymbolQueryResult::NotFound:
+				TimerName = FString::Printf(TEXT("%s / 0x%X NotFound"),
+					StackFrame.Symbol->Module ? StackFrame.Symbol->Module : TEXT("N/A"),
+					StackFrame.Address);
+				break;
+
+			default:
+				TimerName = FString::Printf(TEXT("0x%X Invalid !!!"), StackFrame.Address);
+				break;
 		}
 	}
 
-	for (uint32 Index = 0; Index < Count; Index++)
+	if (!bIsPending || bFirstTime)
 	{
-		AddressValues.EmplaceBack(Addresses[Index]);
+		FAnalysisSessionEditScope _(Session);
+		EditableTimingProfilerProvider->SetTimerName(StackFrame.TimerId, TimerName);
 	}
-	StackSample.Addresses = &AddressValues.Last() - Count;
+
+	return !bIsPending;
+}
+
+bool FStackSamplesProvider::UpdatePendingTimers()
+{
+	for (int32 Index = 0; Index < PendingTimers.Num(); ++Index)
+	{
+		FStackSampleFrame& Frame = *PendingTimers[Index];
+
+		if (UpdateTimerName(Frame, false))
+		{
+			PendingTimers[Index] = PendingTimers[PendingTimers.Num() - 1];
+			PendingTimers.RemoveAt(PendingTimers.Num() - 1, EAllowShrinking::No);
+			--Index;
+		}
+	}
+
+	return PendingTimers.Num() > 0;
 }
 
 FName GetStackSamplesProviderName()
@@ -76,9 +328,9 @@ FName GetStackSamplesProviderName()
 	return Name;
 }
 
-const IStackSamplesProvider& ReadStackSamplesProvider(const IAnalysisSession& Session)
+const IStackSamplesProvider* ReadStackSamplesProvider(const IAnalysisSession& Session)
 {
-	return *Session.ReadProvider<IStackSamplesProvider>(GetStackSamplesProviderName());
+	return Session.ReadProvider<IStackSamplesProvider>(GetStackSamplesProviderName());
 }
 
 } // namespace TraceServices
diff --git a/Engine/Source/Developer/TraceServices/Private/Model/StackSamplesPrivate.h b/Engine/Source/Developer/TraceServices/Private/Model/StackSamplesPrivate.h
--- a/Engine/Source/Developer/TraceServices/Private/Model/StackSamplesPrivate.h
+++ b/Engine/Source/Developer/TraceServices/Private/Model/StackSamplesPrivate.h
@@ -2,29 +2,117 @@
 
 #pragma once
 
-#include "Common/PagedArray.h"
+#include "CoreTypes.h"
 #include "Containers/Map.h"
+
+// TraceServices
+#include "Common/PagedArray.h"
+#include "Common/ProviderLock.h"
+#include "Model/MonotonicTimeline.h"
 #include "TraceServices/Model/AnalysisSession.h"
 #include "TraceServices/Model/StackSamples.h"
+#include "TraceServices/Model/TimingProfiler.h"
 
 namespace TraceServices
 {
 
+class IModuleProvider;
+
+extern thread_local FProviderLock::FThreadLocalState GStackSamplesProviderLockState;
+
+typedef TMonotonicTimeline<FTimingProfilerEvent> FStackSamplesTimeline;
+
+class FStackSamplesProviderThread
+{
+	friend class FStackSamplesProvider;
+
+public:
+	FStackSamplesProviderThread(uint32 InSystemThreadId)
+		: SystemThreadId(InSystemThreadId)
+	{
+	}
+	~FStackSamplesProviderThread()
+	{
+	}
+
+	uint32 GetSystemThreadId() const { return SystemThreadId; }
+	const TCHAR* GetName() const { return Name; }
+	FStackSamplesTimeline* GetTimeline() const { return Timeline.Get(); }
+
+private:
+	uint32 SystemThreadId = 0;
+	const TCHAR* Name = nullptr;
+	TSharedPtr<FStackSamplesTimeline> Timeline;
+	TArray<uint64> LastSampleStack;
+	double LastSampleTime = 0.0;
+	uint32 NumSamples = 0;
+	int32 MaxStackSize = 0;
+};
+
 class FStackSamplesProvider
 	: public IStackSamplesProvider
+	, public IEditableProvider
 {
 public:
 	explicit FStackSamplesProvider(IAnalysisSession& Session);
 	virtual ~FStackSamplesProvider();
 
-	const TPagedArray<FStackSample>* GetStackSamples(uint32 ThreadId) const override;
+	//////////////////////////////////////////////////
+	// Read operations
 
-	void Add(uint32 ThreadId, double Time, uint32 Count, const uint64* Addresses);
+	virtual void BeginRead() const override { Lock.BeginRead(GStackSamplesProviderLockState); }
+	virtual void EndRead() const override { Lock.EndRead(GStackSamplesProviderLockState); }
+	virtual void ReadAccessCheck() const override { Lock.ReadAccessCheck(GStackSamplesProviderLockState); }
+
+	virtual uint32 GetStackFrameCount() const override;
+	virtual void EnumerateStackFrames(TFunctionRef<void(const FStackSampleFrame& Frame)> Callback) const override;
+
+	virtual uint32 GetThreadCount() const override;
+	virtual void EnumerateThreads(TFunctionRef<void(const FStackSampleThread& Thread)> Callback) const override;
+
+	virtual uint32 GetTimelineCount() const override;
+	virtual void EnumerateTimelines(TFunctionRef<void(const FStackSampleTimeline& Timeline, const FStackSampleThread& Thread)> Callback) const override;
+	virtual const FStackSampleTimeline* GetTimeline(uint32 SystemThreadId) const override;
+
+	//////////////////////////////////////////////////
+	// Edit operations
+
+	virtual void BeginEdit() const override { Lock.BeginWrite(GStackSamplesProviderLockState); }
+	virtual void EndEdit() const override { Lock.EndWrite(GStackSamplesProviderLockState); }
+	virtual void EditAccessCheck() const override { Lock.WriteAccessCheck(GStackSamplesProviderLockState); }
+
+	void AddThreadName(uint32 SystemThreadId, uint32 SystemProcessId, FStringView Name);
+	void AddStackSample(uint32 SystemThreadId, double Time, uint32 Count, const uint64* Addresses);
+
+	bool UpdatePendingTimers();
+
+	//////////////////////////////////////////////////
+
+private:
+	FStackSamplesProviderThread& GetOrAddThread(uint32 SystemThreadId);
+	FStackSampleFrame& GetOrAddStackFrame(uint64 Address);
+	void CreateTimer(FStackSampleFrame& StackFrame);
+	bool UpdateTimerName(FStackSampleFrame& StackFrame, bool bFirstTime);
 
 private:
+	mutable FProviderLock Lock;
+
 	IAnalysisSession& Session;
-	TMap<uint32, TPagedArray<FStackSample>*> Threads;
-	TPagedArray<uint64> AddressValues;
+	IEditableTimingProfilerProvider* EditableTimingProfilerProvider = nullptr;
+	IModuleProvider* ModuleProvider = nullptr;
+
+	TPagedArray<FStackSamplesProviderThread> Threads;
+	TArray<FStackSamplesProviderThread*> Timelines; // the threads with a valid timeline
+	TMap<uint32, FStackSamplesProviderThread*> ThreadsBySystemId; // SystemThreadId --> FThreadSamplingData*
+	TPagedArray<FStackSampleFrame> StackFrames;
+	TMap<uint64, FStackSampleFrame*> StackFramesByAddress; // Address --> FStackSampleFrame*
+	TMap<uint32, FStackSampleFrame*> StackFramesByTimerId; // TimerId --> FStackSampleFrame*
+	TArray<FStackSampleFrame*> PendingTimers;
+
+	uint32 UnknownTimerId = 0;
+	uint32 TotalNumSamples = 0;
+
+	static constexpr double MaxSampleDuration = 0.004;
 };
 
 } // namespace TraceServices
diff --git a/Engine/Source/Developer/TraceServices/Public/TraceServices/Model/StackSamples.h b/Engine/Source/Developer/TraceServices/Public/TraceServices/Model/StackSamples.h
--- a/Engine/Source/Developer/TraceServices/Public/TraceServices/Model/StackSamples.h
+++ b/Engine/Source/Developer/TraceServices/Public/TraceServices/Model/StackSamples.h
@@ -2,19 +2,34 @@
 
 #pragma once
 
-#include "Common/PagedArray.h" // TraceServices
+#include "CoreTypes.h"
 #include "HAL/Platform.h"
-#include "TraceServices/Model/AnalysisSession.h"
 #include "UObject/NameTypes.h"
 
+// TraceServices
+#include "TraceServices/Containers/Timelines.h"
+#include "TraceServices/Model/AnalysisSession.h"
+#include "TraceServices/Model/TimingProfiler.h"
+
 namespace TraceServices
 {
 
-struct FStackSample
+struct FResolvedSymbol;
+
+struct FStackSampleFrame
 {
-	double Time;
-	uint32 Count;
-	uint64* Addresses;
+	uint64 Address = 0;
+	uint32 TimerId = 0;
+	const FResolvedSymbol* Symbol = nullptr;
+};
+
+typedef ITimeline<FTimingProfilerEvent> FStackSampleTimeline;
+
+struct FStackSampleThread
+{
+	uint32 SystemThreadId = 0;
+	const TCHAR* Name = nullptr;
+	const FStackSampleTimeline* Timeline = nullptr;
 };
 
 class IStackSamplesProvider
@@ -22,10 +37,19 @@ class IStackSamplesProvider
 {
 public:
 	virtual ~IStackSamplesProvider() = default;
-	virtual const TPagedArray<FStackSample>* GetStackSamples(uint32 ThreadId) const = 0;
+
+	virtual uint32 GetStackFrameCount() const = 0;
+	virtual void EnumerateStackFrames(TFunctionRef<void(const FStackSampleFrame& Frame)> Callback) const = 0;
+
+	virtual uint32 GetThreadCount() const = 0;
+	virtual void EnumerateThreads(TFunctionRef<void(const FStackSampleThread& Thread)> Callback) const = 0;
+
+	virtual uint32 GetTimelineCount() const = 0;
+	virtual void EnumerateTimelines(TFunctionRef<void(const FStackSampleTimeline& Timeline, const FStackSampleThread& Thread)> Callback) const = 0;
+	virtual const FStackSampleTimeline* GetTimeline(uint32 SystemThreadId) const = 0;
 };
 
 TRACESERVICES_API FName GetStackSamplesProviderName();
-TRACESERVICES_API const IStackSamplesProvider& ReadStackSamplesProvider(const IAnalysisSession& Session);
+TRACESERVICES_API const IStackSamplesProvider* ReadStackSamplesProvider(const IAnalysisSession& Session);
 
 } // namespace TraceServices
