[Insights] CPU Stack Sampling: Improved the merge samples algorithm. Now it centers samples on their time. Also now it uses an additional 50% max tolerance (so 0.5ms for the current 1ms intervals), but only for the distant samples (ex: when dt > 1ms). Before it was using 4ms enlarged samples.

diff --git a/Engine/Source/Developer/TraceServices/Private/Model/StackSamples.cpp b/Engine/Source/Developer/TraceServices/Private/Model/StackSamples.cpp
--- a/Engine/Source/Developer/TraceServices/Private/Model/StackSamples.cpp
+++ b/Engine/Source/Developer/TraceServices/Private/Model/StackSamples.cpp
@@ -20,6 +20,7 @@ FStackSamplesProvider::FStackSamplesProvider(IAnalysisSession& InSession)
 	, Threads(Session.GetLinearAllocator(), 256)
 	, StackFrames(Session.GetLinearAllocator(), 4 << 10)
 {
+	SetSamplingInterval(1000); // 1ms
 }
 
 FStackSamplesProvider::~FStackSamplesProvider()
@@ -146,6 +147,12 @@ FStackSampleFrame& FStackSamplesProvider::GetOrAddStackFrame(uint64 Address)
 	return NewStackFrame;
 }
 
+void FStackSamplesProvider::SetSamplingInterval(uint32 Microseconds)
+{
+	SamplingIntervalUsec = Microseconds;
+	MaxSampleDuration = (double)Microseconds / 1000000.0; // [seconds]
+}
+
 void FStackSamplesProvider::AddStackSample(uint32 SystemThreadId, double Time, uint32 Count, const uint64* Addresses)
 {
 	EditAccessCheck();
@@ -181,8 +188,25 @@ void FStackSamplesProvider::AddStackSample(uint32 SystemThreadId, double Time, u
 		Thread.MaxStackSize = NewSampleStackSize;
 	}
 
+	// Ensure positive time values.
+	if (!ensure(Time >= 0.0))
+	{
+		Time = 0.0;
+	}
+
+	// Ensure monotonic increase of time values.
+	if (!ensure(Time >= Thread.LastSampleTime))
+	{
+		Time = Thread.LastSampleTime;
+	}
+
+	const double LastSampleEndTime = Thread.LastSampleTime + MaxSampleDuration/2;
+	const double CurrentSampleStartTime = FMath::Max(0.0, Time - MaxSampleDuration/2);
+
+	constexpr double Tolerance = 0.5; // 50%; i.e. 0.5ms for 1ms intervals
+
 	int32 StartNewStackIndex = 0;
-	if (Time - Thread.LastSampleTime < MaxSampleDuration)
+	if (CurrentSampleStartTime - MaxSampleDuration * Tolerance <= LastSampleEndTime)
 	{
 		for (int32 StackIndex = 0; StackIndex < LastSampleStackSize; ++StackIndex)
 		{
@@ -195,7 +219,10 @@ void FStackSamplesProvider::AddStackSample(uint32 SystemThreadId, double Time, u
 		}
 	}
 
-	const double EndTime = FMath::Min(Thread.LastSampleTime + MaxSampleDuration / 2.0, Time);
+	const double MidTime = (Thread.LastSampleTime + Time) / 2.0;
+	const double EndTime = FMath::Min(LastSampleEndTime, MidTime);
+	const double StartTime = FMath::Max(CurrentSampleStartTime, MidTime);
+
 	for (int32 StackIndex = StartNewStackIndex; StackIndex < LastSampleStackSize; ++StackIndex)
 	{
 		Thread.Timeline->AppendEndEvent(EndTime);
@@ -205,7 +232,7 @@ void FStackSamplesProvider::AddStackSample(uint32 SystemThreadId, double Time, u
 		FStackSampleFrame& StackFrame = GetOrAddStackFrame(Callstack[StackIndex]);
 		FTimingProfilerEvent Event;
 		Event.TimerIndex = StackFrame.TimerId;
-		Thread.Timeline->AppendBeginEvent(Time, Event);
+		Thread.Timeline->AppendBeginEvent(StartTime, Event);
 	}
 
 	// LastSampleStack = Callstack
diff --git a/Engine/Source/Developer/TraceServices/Private/Model/StackSamplesPrivate.h b/Engine/Source/Developer/TraceServices/Private/Model/StackSamplesPrivate.h
--- a/Engine/Source/Developer/TraceServices/Private/Model/StackSamplesPrivate.h
+++ b/Engine/Source/Developer/TraceServices/Private/Model/StackSamplesPrivate.h
@@ -83,6 +83,7 @@ public:
 
 	void AddThreadName(uint32 SystemThreadId, uint32 SystemProcessId, FStringView Name);
 	void AddStackSample(uint32 SystemThreadId, double Time, uint32 Count, const uint64* Addresses);
+	void SetSamplingInterval(uint32 Microseconds);
 
 	bool UpdatePendingTimers();
 
@@ -112,7 +113,8 @@ private:
 	uint32 UnknownTimerId = 0;
 	uint32 TotalNumSamples = 0;
 
-	static constexpr double MaxSampleDuration = 0.004;
+	uint32 SamplingIntervalUsec = 0;// [microseconds]
+	double MaxSampleDuration = 0.0;// [seconds]
 };
 
 } // namespace TraceServices
