[Insights]
* CPU Sampling Profiler: Added analysis for the configurable sampling interval.
* CPU Sampling Profiler: Fixed end time for last sample in each thread.
* CPU Sampling Profiler: Added end analysis stats (log messages).

diff --git a/Engine/Source/Developer/TraceServices/Private/Analyzers/PlatformEventTraceAnalysis.cpp b/Engine/Source/Developer/TraceServices/Private/Analyzers/PlatformEventTraceAnalysis.cpp
--- a/Engine/Source/Developer/TraceServices/Private/Analyzers/PlatformEventTraceAnalysis.cpp
+++ b/Engine/Source/Developer/TraceServices/Private/Analyzers/PlatformEventTraceAnalysis.cpp
@@ -23,6 +23,7 @@ void FPlatformEventTraceAnalyzer::OnAnalysisBegin(const FOnAnalysisContext& Cont
 {
 	auto& Builder = Context.InterfaceBuilder;
 
+	Builder.RouteEvent(RouteId_Settings, "PlatformEvent", "Settings");
 	Builder.RouteEvent(RouteId_ContextSwitch, "PlatformEvent", "ContextSwitch");
 	Builder.RouteEvent(RouteId_StackSample, "PlatformEvent", "StackSample");
 	Builder.RouteEvent(RouteId_ThreadName, "PlatformEvent", "ThreadName");
@@ -30,6 +31,17 @@ void FPlatformEventTraceAnalyzer::OnAnalysisBegin(const FOnAnalysisContext& Cont
 
 void FPlatformEventTraceAnalyzer::OnAnalysisEnd()
 {
+	double MaxEndTime = 0.0;
+	{
+		TraceServices::FProviderEditScopeLock ScopedLock(StackSamplesProvider);
+		MaxEndTime = StackSamplesProvider.OnAnalysisEnd();
+	}
+	if (MaxEndTime > 0.0)
+	{
+		FAnalysisSessionEditScope _(Session);
+		Session.UpdateDurationSeconds(MaxEndTime);
+	}
+
 	int32 TryCount = 10;
 	while (TryCount > 0)
 	{
@@ -49,6 +61,17 @@ bool FPlatformEventTraceAnalyzer::OnEvent(uint16 RouteId, EStyle Style, const FO
 	switch (RouteId)
 	{
 
+	case RouteId_Settings:
+	{
+		const auto& EventData = Context.EventData;
+		uint32 SamplingIntervalUsec = EventData.GetValue<uint32>("SamplingInterval");
+		{
+			TraceServices::FProviderEditScopeLock ScopedLock(StackSamplesProvider);
+			StackSamplesProvider.SetSamplingInterval(SamplingIntervalUsec);
+		}
+		break;
+	}
+
 	case RouteId_ContextSwitch:
 	{
 		const auto& EventData = Context.EventData;
diff --git a/Engine/Source/Developer/TraceServices/Private/Analyzers/PlatformEventTraceAnalysis.h b/Engine/Source/Developer/TraceServices/Private/Analyzers/PlatformEventTraceAnalysis.h
--- a/Engine/Source/Developer/TraceServices/Private/Analyzers/PlatformEventTraceAnalysis.h
+++ b/Engine/Source/Developer/TraceServices/Private/Analyzers/PlatformEventTraceAnalysis.h
@@ -27,6 +27,7 @@ public:
 private:
 	enum : uint16
 	{
+		RouteId_Settings,
 		RouteId_ContextSwitch,
 		RouteId_StackSample,
 		RouteId_ThreadName,
diff --git a/Engine/Source/Developer/TraceServices/Private/Model/StackSamples.cpp b/Engine/Source/Developer/TraceServices/Private/Model/StackSamples.cpp
--- a/Engine/Source/Developer/TraceServices/Private/Model/StackSamples.cpp
+++ b/Engine/Source/Developer/TraceServices/Private/Model/StackSamples.cpp
@@ -8,6 +8,7 @@
 
 // TraceServices
 #include "AnalysisServicePrivate.h"
+#include "Common/Utils.h"
 #include "TraceServices/Model/Modules.h"
 
 namespace TraceServices
@@ -27,6 +28,12 @@ FStackSamplesProvider::~FStackSamplesProvider()
 {
 }
 
+void FStackSamplesProvider::SetSamplingInterval(uint32 Microseconds)
+{
+	SamplingIntervalUsec = Microseconds;
+	MaxSampleDuration = (double)Microseconds / 1000000.0; // [seconds]
+}
+
 uint32 FStackSamplesProvider::GetStackFrameCount() const
 {
 	ReadAccessCheck();
@@ -147,12 +154,6 @@ FStackSampleFrame& FStackSamplesProvider::GetOrAddStackFrame(uint64 Address)
 	return NewStackFrame;
 }
 
-void FStackSamplesProvider::SetSamplingInterval(uint32 Microseconds)
-{
-	SamplingIntervalUsec = Microseconds;
-	MaxSampleDuration = (double)Microseconds / 1000000.0; // [seconds]
-}
-
 void FStackSamplesProvider::AddStackSample(uint32 SystemThreadId, double Time, uint32 Count, const uint64* Addresses)
 {
 	EditAccessCheck();
@@ -183,20 +184,10 @@ void FStackSamplesProvider::AddStackSample(uint32 SystemThreadId, double Time, u
 	const int32 LastSampleStackSize = Thread.LastSampleStack.Num();
 	const int32 NewSampleStackSize = Callstack.Num();
 
-	if (NewSampleStackSize > Thread.MaxStackSize)
-	{
-		Thread.MaxStackSize = NewSampleStackSize;
-	}
-
-	// Ensure positive time values.
-	if (!ensure(Time >= 0.0))
-	{
-		Time = 0.0;
-	}
-
 	// Ensure monotonic increase of time values.
 	if (!ensure(Time >= Thread.LastSampleTime))
 	{
+		UE_LOG(LogTraceServices, Warning, TEXT("[StackSampling] Time should increase monotonically (SysThreadId=0x%X, Time=%f < %f)!"), SystemThreadId, Time, Thread.LastSampleTime);
 		Time = Thread.LastSampleTime;
 	}
 
@@ -246,7 +237,10 @@ void FStackSamplesProvider::AddStackSample(uint32 SystemThreadId, double Time, u
 	}
 
 	Thread.LastSampleTime = Time;
-	Thread.NumSamples++;
+
+	// Update stats...
+	MaxStackSize = FMath::Max(MaxStackSize, NewSampleStackSize);
+	TotalNumFrames += static_cast<uint32>(NewSampleStackSize);
 	TotalNumSamples++;
 }
 
@@ -349,6 +343,33 @@ bool FStackSamplesProvider::UpdatePendingTimers()
 	return PendingTimers.Num() > 0;
 }
 
+double FStackSamplesProvider::OnAnalysisEnd()
+{
+	EditAccessCheck();
+
+	double MaxEndTime = 0.0;
+
+	for (FStackSamplesProviderThread* Thread : Timelines)
+	{
+		check(Thread->Timeline);
+
+		const double LastSampleEndTime = Thread->LastSampleTime + MaxSampleDuration / 2;
+		MaxEndTime = FMath::Max(MaxEndTime, LastSampleEndTime);
+
+		for (int32 Index = Thread->LastSampleStack.Num(); Index > 0; --Index)
+		{
+			Thread->Timeline->AppendEndEvent(LastSampleEndTime);
+		}
+		Thread->LastSampleStack.Reset();
+	}
+
+	UE_LOG(LogTraceServices, Log, TEXT("[StackSampling] max stack size: %d frames"), MaxStackSize);
+	UE_LOG(LogTraceServices, Log, TEXT("[StackSampling] Analysis completed (%llu timers, %d timelines, %llu threads, %llu stack samples, %llu stack frames)."),
+		StackFrames.Num(), Timelines.Num(), Threads.Num(), TotalNumSamples, TotalNumFrames);
+
+	return MaxEndTime;
+}
+
 FName GetStackSamplesProviderName()
 {
 	static const FName Name("StackSamplesProvider");
diff --git a/Engine/Source/Developer/TraceServices/Private/Model/StackSamplesPrivate.h b/Engine/Source/Developer/TraceServices/Private/Model/StackSamplesPrivate.h
--- a/Engine/Source/Developer/TraceServices/Private/Model/StackSamplesPrivate.h
+++ b/Engine/Source/Developer/TraceServices/Private/Model/StackSamplesPrivate.h
@@ -45,8 +45,6 @@ private:
 	TSharedPtr<FStackSamplesTimeline> Timeline;
 	TArray<uint64> LastSampleStack;
 	double LastSampleTime = 0.0;
-	uint32 NumSamples = 0;
-	int32 MaxStackSize = 0;
 };
 
 class FStackSamplesProvider
@@ -81,9 +79,10 @@ public:
 	virtual void EndEdit() const override { Lock.EndWrite(GStackSamplesProviderLockState); }
 	virtual void EditAccessCheck() const override { Lock.WriteAccessCheck(GStackSamplesProviderLockState); }
 
+	void SetSamplingInterval(uint32 Microseconds);
 	void AddThreadName(uint32 SystemThreadId, uint32 SystemProcessId, FStringView Name);
 	void AddStackSample(uint32 SystemThreadId, double Time, uint32 Count, const uint64* Addresses);
-	void SetSamplingInterval(uint32 Microseconds);
+	double OnAnalysisEnd();
 
 	bool UpdatePendingTimers();
 
@@ -110,11 +109,12 @@ private:
 	TMap<uint32, FStackSampleFrame*> StackFramesByTimerId; // TimerId --> FStackSampleFrame*
 	TArray<FStackSampleFrame*> PendingTimers;
 
-	uint32 UnknownTimerId = 0;
-	uint32 TotalNumSamples = 0;
-
-	uint32 SamplingIntervalUsec = 0;// [microseconds]
 	double MaxSampleDuration = 0.0;// [seconds]
+	uint32 SamplingIntervalUsec = 0;// [microseconds]
+
+	int32 MaxStackSize = 0;
+	uint64 TotalNumFrames = 0;
+	uint64 TotalNumSamples = 0;
 };
 
 } // namespace TraceServices
