[Insights] GpuProfiler: Improved warnings and errors emitted by analysis.

#rnx

diff --git a/Engine/Source/Developer/TraceServices/Private/Analyzers/GpuProfilerTraceAnalysis.cpp b/Engine/Source/Developer/TraceServices/Private/Analyzers/GpuProfilerTraceAnalysis.cpp
--- a/Engine/Source/Developer/TraceServices/Private/Analyzers/GpuProfilerTraceAnalysis.cpp
+++ b/Engine/Source/Developer/TraceServices/Private/Analyzers/GpuProfilerTraceAnalysis.cpp
@@ -68,22 +68,22 @@ void FGpuProfilerAnalyzer::OnAnalysisEnd()
 
 	if (ErrorData.NumInterleavedEvents > 0)
 	{
-		UE_LOG(LogTraceServices, Error, TEXT("[GpuProfiler] Number of interleaved events: %u, max error %f"), ErrorData.NumInterleavedEvents, ErrorData.InterleavedEventsMaxDelta);
+		UE_LOG(LogTraceServices, Warning, TEXT("[GpuProfiler] Number of interleaved events: %u, max error %f"), ErrorData.NumInterleavedEvents, ErrorData.InterleavedEventsMaxDelta);
 	}
 
 	if (ErrorData.NumInterleavedAndReversedEvents > 0)
 	{
-		UE_LOG(LogTraceServices, Error, TEXT("[GpuProfiler] Number of interleaved and reversed events: %u, max error %f"), ErrorData.NumInterleavedAndReversedEvents, ErrorData.InterleavedAndReversedEventsMaxDelta);
+		UE_LOG(LogTraceServices, Warning, TEXT("[GpuProfiler] Number of interleaved and reversed events: %u, max error %f"), ErrorData.NumInterleavedAndReversedEvents, ErrorData.InterleavedAndReversedEventsMaxDelta);
 	}
 
 	if (ErrorData.NumMismatchedEvents > 0)
 	{
-		UE_LOG(LogTraceServices, Error, TEXT("[GpuProfiler] Number of mismatched events: %u"), ErrorData.NumMismatchedEvents);
+		UE_LOG(LogTraceServices, Warning, TEXT("[GpuProfiler] Number of mismatched events: %u"), ErrorData.NumMismatchedEvents);
 	}
 
 	if (ErrorData.NumNegativeDurationEvents > 0)
 	{
-		UE_LOG(LogTraceServices, Error, TEXT("[GpuProfiler] Number of negative duration events: %u, max error %f"), ErrorData.NumNegativeDurationEvents, ErrorData.NegativeDurationEventsMaxDelta);
+		UE_LOG(LogTraceServices, Warning, TEXT("[GpuProfiler] Number of negative duration events: %u, max error %f"), ErrorData.NumNegativeDurationEvents, ErrorData.NegativeDurationEventsMaxDelta);
 	}
 
 	if (Queues.Num() > 0 || TimerMap.Num() > 0)
@@ -108,7 +108,7 @@ bool FGpuProfilerAnalyzer::OnEvent(uint16 RouteId, EStyle Style, const FOnEventC
 	case RouteId_Init:
 	{
 		Version = (uint32)EventData.GetValue<uint8>("Version");
-		UE_LOG(LogTraceServices, Log, TEXT("[GPU] Init Version=%u"), Version);
+		UE_LOG(LogTraceServices, Log, TEXT("[GpuProfiler] Init Version=%u"), Version);
 		break;
 	}
 	case RouteId_QueueSpec:
@@ -122,7 +122,7 @@ bool FGpuProfilerAnalyzer::OnEvent(uint16 RouteId, EStyle Style, const FOnEventC
 		FString Name;
 		EventData.GetString("TypeString", Name);
 
-		UE_LOG(LogTraceServices, Log, TEXT("[GPU] QueueSpec QueueId=%u, GPU=%u, Index=%u, Type=%d \"%s\""), QueueId, (uint32)GPU, (uint32)Index, (uint32)Type, *Name);
+		UE_LOG(LogTraceServices, Log, TEXT("[GpuProfiler] QueueSpec QueueId=%u, GPU=%u, Index=%u, Type=%d \"%s\""), QueueId, (uint32)GPU, (uint32)Index, (uint32)Type, *Name);
 
 		FAnalysisSessionEditScope _(Session);
 		const TCHAR* PersistentName = Session.StoreString(Name);
@@ -139,7 +139,7 @@ bool FGpuProfilerAnalyzer::OnEvent(uint16 RouteId, EStyle Style, const FOnEventC
 		const uint32 QueueId = EventData.GetValue<uint32>("QueueId");
 		const uint32 FrameNumber = EventData.GetValue<uint32>("FrameNumber");
 
-		UE_LOG(LogTraceServices, VeryVerbose, TEXT("[GPU] EventFrameBoundary QueueId=%u, FrameNumber=%u"), QueueId, FrameNumber);
+		UE_LOG(LogTraceServices, VeryVerbose, TEXT("[GpuProfiler] EventFrameBoundary QueueId=%u, FrameNumber=%u"), QueueId, FrameNumber);
 
 		FQueue& Queue = GetOrAddQueue(QueueId);
 		Queue.FrameNumber = FrameNumber;
@@ -277,7 +277,7 @@ bool FGpuProfilerAnalyzer::OnEvent(uint16 RouteId, EStyle Style, const FOnEventC
 		TArray<uint8> Metadata(EventData.GetArrayView<uint8>("Metadata"));
 		const double Time = Context.EventTime.AsSeconds(GPUTimestampTOP);
 
-		UE_LOG(LogTraceServices, VeryVerbose, TEXT("[GPU] EventBeginBreadcrumb QueueId=%u, Time=%f, SpecId=\"%u\""), QueueId, Time, SpecId);
+		UE_LOG(LogTraceServices, VeryVerbose, TEXT("[GpuProfiler] EventBeginBreadcrumb QueueId=%u, Time=%f, SpecId=\"%u\""), QueueId, Time, SpecId);
 
 		uint32 TimerId = 0;
 		uint32* TimerIdPtr;
@@ -315,13 +315,10 @@ bool FGpuProfilerAnalyzer::OnEvent(uint16 RouteId, EStyle Style, const FOnEventC
 
 		const double Time = Context.EventTime.AsSeconds(GPUTimestampBOP);
 
-		FString Name;
-		EventData.GetString("Name", Name);
-
-		UE_LOG(LogTraceServices, VeryVerbose, TEXT("[GPU] EventEndBreadcrumb QueueId=%u, Time=%f, Name=\"%s\""), QueueId, Time, *Name);
+		UE_LOG(LogTraceServices, VeryVerbose, TEXT("[GpuProfiler] EventEndBreadcrumb QueueId=%u, Time=%f"), QueueId, Time);
 
 		FQueue& Queue = GetOrAddQueue(QueueId);
-		EndEvent(Queue, 0, Time, (uint32)- 1);
+		EndEvent(Queue, 0, Time, InvalidTimerId);
 		break;
 	}
 	case RouteId_EventBeginWork:
@@ -331,9 +328,9 @@ bool FGpuProfilerAnalyzer::OnEvent(uint16 RouteId, EStyle Style, const FOnEventC
 		double Time = Context.EventTime.AsSeconds(GPUTimestampTOP);
 		//const uint64 CPUTimestamp = EventData.GetValue<uint64>("CPUTimestamp"); // not used
 
-		UE_LOG(LogTraceServices, VeryVerbose, TEXT("[GPU] EventBeginWork QueueId=%u, GPUTimestampTOP=%f"), QueueId, Time);
+		UE_LOG(LogTraceServices, VeryVerbose, TEXT("[GpuProfiler] EventBeginWork QueueId=%u, GPUTimestampTOP=%f"), QueueId, Time);
 
-		if (GpuWorkTimerId == ~0u)
+		if (GpuWorkTimerId == InvalidTimerId)
 		{
 			const FString GpuWorkTimerName(TEXT("GpuWork"));
 			GpuWorkTimerId = GetOrAddTimerId(GpuWorkTimerName);
@@ -363,7 +360,7 @@ bool FGpuProfilerAnalyzer::OnEvent(uint16 RouteId, EStyle Style, const FOnEventC
 		const uint64 GPUTimestampBOP = EventData.GetValue<uint64>("GPUTimestampBOP");
 		const double Time = Context.EventTime.AsSeconds(GPUTimestampBOP);
 
-		UE_LOG(LogTraceServices, VeryVerbose, TEXT("[GPU] EventEndWork QueueId=%u, GPUTimestampBOP=%f"), QueueId, Time);
+		UE_LOG(LogTraceServices, VeryVerbose, TEXT("[GpuProfiler] EventEndWork QueueId=%u, GPUTimestampBOP=%f"), QueueId, Time);
 
 		FQueue& Queue = GetOrAddQueue(QueueId);
 		EndEvent(Queue, 1, Time, GpuWorkTimerId);
@@ -378,9 +375,9 @@ bool FGpuProfilerAnalyzer::OnEvent(uint16 RouteId, EStyle Style, const FOnEventC
 		const uint64 EndTimestamp = EventData.GetValue<uint64>("EndTime");
 		const double EndTime = Context.EventTime.AsSeconds(EndTimestamp);
 
-		UE_LOG(LogTraceServices, VeryVerbose, TEXT("[GPU] EventWait QueueId=%u, StartTime=%f, EndTime=%f"), QueueId, StartTime, EndTime);
+		UE_LOG(LogTraceServices, VeryVerbose, TEXT("[GpuProfiler] EventWait QueueId=%u, StartTime=%f, EndTime=%f"), QueueId, StartTime, EndTime);
 
-		if (GpuWaitTimerId == ~0u)
+		if (GpuWaitTimerId == InvalidTimerId)
 		{
 			const FString GpuWaitTimerName(TEXT("GpuWait"));
 			GpuWaitTimerId = GetOrAddTimerId(GpuWaitTimerName);
@@ -619,6 +616,10 @@ uint32 FGpuProfilerAnalyzer::GetOrAddTimerId(const FString& Breadcrumb)
 
 const TCHAR* FGpuProfilerAnalyzer::GetTimerName(uint32 TimerId) const
 {
+	if (TimerId == InvalidTimerId)
+	{
+		return TEXT("<invalid>");
+	}
 	const TCHAR* const* TimerNamePtr = TimerMap.Find(TimerId);
 	return TimerNamePtr ? *TimerNamePtr : TEXT("<unknown>");
 }
@@ -655,15 +656,15 @@ void FGpuProfilerAnalyzer::BeginEvent(FQueue& Queue, int32 StackIndex, double Be
 		{
 			++ErrorData.NumInterleavedEvents;
 			ErrorData.InterleavedEventsMaxDelta = FMath::Max(ErrorData.InterleavedEventsMaxDelta, LastTimestamp - BeginEventTime);
-
-			if (ErrorData.NumInterleavedEvents < ErrorData.NumMaxErrors)
+			if (++ErrorData.NumWarnings < ErrorData.NumMaxWarnings)
 			{
-				UE_LOG(LogTraceServices, Error, TEXT("[GPU] %s BEGIN %u \"%s\" : Time %f < %f !!!"),
+				uint32 OriginalBeginEventTimerId = TimingProfilerProvider.GetOriginalTimerIdFromMetadata(BeginEventTimerId);
+				UE_LOG(LogTraceServices, Warning, TEXT("[GpuProfiler] Queue %u %s BEGIN %d \"%s\" : Time %f < %f (last) !!!"),
+					Queue.Id,
 					StackIndex == 1 ? TEXT("WORK") : TEXT("BREADCRUMB"),
-					BeginEventTimerId, GetTimerName(BeginEventTimerId),
+					int32(BeginEventTimerId), GetTimerName(OriginalBeginEventTimerId),
 					BeginEventTime, LastTimestamp);
 			}
-
 			BeginEventTime = LastTimestamp;
 		}
 		Timeline->AppendBeginEvent(BeginEventTime, Event);
@@ -685,29 +686,42 @@ void FGpuProfilerAnalyzer::EndEvent(FQueue& Queue, int32 StackIndex, double EndE
 		FAnalysisSessionReadScope _(Session);
 		BeginEventTimerId = TimingProfilerProvider.GetOriginalTimerIdFromMetadata(BeginEvent.TimerId);
 	}
-	if (EndEventTimerId != (uint32)-1 && BeginEventTimerId != EndEventTimerId)
+	if (EndEventTimerId != InvalidTimerId &&
+		EndEventTimerId != BeginEventTimerId)
 	{
 		++ErrorData.NumMismatchedEvents;
-
-		if (ErrorData.NumMismatchedEvents < ErrorData.NumMaxErrors)
+		if (++ErrorData.NumWarnings < ErrorData.NumMaxWarnings)
 		{
-			UE_LOG(LogTraceServices, Error, TEXT("[GPU] %s END %u \"%s\" (%f) != BEGIN %u \"%s\" (%f) !!!"),
+			UE_LOG(LogTraceServices, Warning, TEXT("[GpuProfiler] Queue %u %s END %d \"%s\" (%f) != BEGIN %d \"%s\" (%f) !!!"),
+				Queue.Id,
 				StackIndex == 1 ? TEXT("WORK") : TEXT("BREADCRUMB"),
-				EndEventTimerId, GetTimerName(EndEventTimerId), EndEventTime,
-				BeginEventTimerId, GetTimerName(BeginEventTimerId), BeginEvent.Time);
+				int32(EndEventTimerId), GetTimerName(EndEventTimerId), EndEventTime,
+				int32(BeginEvent.TimerId), GetTimerName(BeginEventTimerId), BeginEvent.Time);
 		}
 	}
 	if (BeginEvent.Time > EndEventTime)
 	{
 		++ErrorData.NumNegativeDurationEvents;
 		ErrorData.NegativeDurationEventsMaxDelta = FMath::Max(ErrorData.NegativeDurationEventsMaxDelta, BeginEvent.Time - EndEventTime);
-
-		if (ErrorData.NumNegativeDurationEvents < ErrorData.NumMaxErrors)
+		if (++ErrorData.NumWarnings < ErrorData.NumMaxWarnings)
 		{
-			UE_LOG(LogTraceServices, Error, TEXT("[GPU] %s END %u \"%s\" (%f) < BEGIN %u \"%s\" (%f) !!!"),
-				StackIndex == 1 ? TEXT("WORK") : TEXT("BREADCRUMB"),
-				EndEventTimerId, GetTimerName(EndEventTimerId), EndEventTime,
-				BeginEventTimerId, GetTimerName(BeginEventTimerId), BeginEvent.Time);
+			if (EndEventTimerId == InvalidTimerId)
+			{
+				UE_LOG(LogTraceServices, Warning, TEXT("[GpuProfiler] Queue %u %s END : Time %f < %f (BEGIN %d \"%s\") !!!"),
+					Queue.Id,
+					StackIndex == 1 ? TEXT("WORK") : TEXT("BREADCRUMB"),
+					EndEventTime, BeginEvent.Time,
+					int32(BeginEvent.TimerId), GetTimerName(BeginEventTimerId));
+			}
+			else
+			{
+				UE_LOG(LogTraceServices, Warning, TEXT("[GpuProfiler] Queue %u %s END %d \"%s\" : Time %f < %f (BEGIN %d \"%s\") !!!"),
+					Queue.Id,
+					StackIndex == 1 ? TEXT("WORK") : TEXT("BREADCRUMB"),
+					int32(EndEventTimerId), GetTimerName(EndEventTimerId),
+					EndEventTime, BeginEvent.Time,
+					int32(BeginEvent.TimerId), GetTimerName(BeginEventTimerId));
+			}
 		}
 	}
 
@@ -722,15 +736,24 @@ void FGpuProfilerAnalyzer::EndEvent(FQueue& Queue, int32 StackIndex, double EndE
 		{
 			++ErrorData.NumInterleavedAndReversedEvents;
 			ErrorData.InterleavedAndReversedEventsMaxDelta = FMath::Max(ErrorData.InterleavedAndReversedEventsMaxDelta, LastTimestamp - EndEventTime);
-
-			if (ErrorData.NumInterleavedAndReversedEvents < ErrorData.NumMaxErrors)
+			if (++ErrorData.NumWarnings < ErrorData.NumMaxWarnings)
 			{
-				UE_LOG(LogTraceServices, Error, TEXT("[GPU] %s END %u \"%s\" : Time %f < %f !!!"),
-					StackIndex == 1 ? TEXT("WORK") : TEXT("BREADCRUMB"),
-					EndEventTimerId, GetTimerName(EndEventTimerId),
-					EndEventTime, LastTimestamp);
+				if (EndEventTimerId == InvalidTimerId)
+				{
+					UE_LOG(LogTraceServices, Warning, TEXT("[GpuProfiler] Queue %u %s END : Time %f < %f (last) !!!"),
+						Queue.Id,
+						StackIndex == 1 ? TEXT("WORK") : TEXT("BREADCRUMB"),
+						EndEventTime, LastTimestamp);
+				}
+				else
+				{
+					UE_LOG(LogTraceServices, Warning, TEXT("[GpuProfiler] Queue %u %s END %d \"%s\" : Time %f < %f (last) !!!"),
+						Queue.Id,
+						StackIndex == 1 ? TEXT("WORK") : TEXT("BREADCRUMB"),
+						int32(EndEventTimerId), GetTimerName(EndEventTimerId),
+						EndEventTime, LastTimestamp);
+				}
 			}
-
 			EndEventTime = LastTimestamp;
 		}
 		Timeline->AppendEndEvent(EndEventTime);
diff --git a/Engine/Source/Developer/TraceServices/Private/Analyzers/GpuProfilerTraceAnalysis.h b/Engine/Source/Developer/TraceServices/Private/Analyzers/GpuProfilerTraceAnalysis.h
--- a/Engine/Source/Developer/TraceServices/Private/Analyzers/GpuProfilerTraceAnalysis.h
+++ b/Engine/Source/Developer/TraceServices/Private/Analyzers/GpuProfilerTraceAnalysis.h
@@ -57,7 +57,11 @@ private:
 
 	struct FErrorData
 	{
-		const uint32 NumMaxErrors = 100;
+		static constexpr uint32 NumMaxWarnings = 100;
+		static constexpr uint32 NumMaxErrors = 100;
+
+		uint32 NumWarnings = 0;
+		uint32 NumErrors = 0;
 
 		uint32 NumInterleavedEvents = 0;
 		uint32 NumInterleavedAndReversedEvents = 0;
@@ -94,8 +98,9 @@ private:
 	// The New GPU Profiler
 
 	uint32 Version = 0;
-	uint32 GpuWorkTimerId = ~0;
-	uint32 GpuWaitTimerId = ~0;
+	static constexpr uint32 InvalidTimerId = 0;
+	uint32 GpuWorkTimerId = InvalidTimerId;
+	uint32 GpuWaitTimerId = InvalidTimerId;
 	TMap<uint32, uint32> BreadcrumbSpecMap; // breadcrumb spec id --> GPU timer id
 	TMap<FString, uint32> BreadcrumbMap; // breadcrumb name --> GPU timer id
 	TMap<uint32, const TCHAR*> TimerMap; // GPU timer id --> persistent timer name
