[Insights] Timing Regions: Implemented and fixed region event navigation for Quick Find searching

diff --git a/Engine/Source/Developer/TraceInsights/Private/Insights/TimingProfiler/Tracks/RegionsTimingTrack.cpp b/Engine/Source/Developer/TraceInsights/Private/Insights/TimingProfiler/Tracks/RegionsTimingTrack.cpp
--- a/Engine/Source/Developer/TraceInsights/Private/Insights/TimingProfiler/Tracks/RegionsTimingTrack.cpp
+++ b/Engine/Source/Developer/TraceInsights/Private/Insights/TimingProfiler/Tracks/RegionsTimingTrack.cpp
@@ -280,7 +280,8 @@ bool FTimingRegionsTrack::FindRegionEvent(const FTimingEventSearchParameters& In
 
 		const TraceServices::IRegionTimeline* Timeline = RegionProvider.GetTimelineForCategory(RegionsCategory);
 		check(Timeline)
-		Timeline->EnumerateRegions(InContext.GetParameters().StartTime, InContext.GetParameters().EndTime, [&InContext](const TraceServices::FTimeRegion& Region)
+
+		auto Callback = [&InContext](const TraceServices::FTimeRegion& Region)
 		{
 			InContext.Check(Region.BeginTime, Region.EndTime, Region.Depth, Region);
 
@@ -290,7 +291,16 @@ bool FTimingRegionsTrack::FindRegionEvent(const FTimingEventSearchParameters& In
 			}
 
 			return true;
-		});
+		};
+			
+		if (InContext.GetParameters().SearchDirection == FTimingEventSearchParameters::ESearchDirection::Forward)
+		{
+			Timeline->EnumerateRegions(InContext.GetParameters().StartTime, InContext.GetParameters().EndTime, Callback);
+		}
+		else
+		{
+			Timeline->EnumerateRegionsBackwards(InContext.GetParameters().EndTime, InContext.GetParameters().StartTime, Callback);
+		}
 	},
 	[&FilterConfiguratorContext, &InParameters](double EventStartTime, double EventEndTime, uint32 EventDepth, const TraceServices::FTimeRegion& Region)
 	{
diff --git a/Engine/Source/Developer/TraceServices/Private/Model/Regions.cpp b/Engine/Source/Developer/TraceServices/Private/Model/Regions.cpp
--- a/Engine/Source/Developer/TraceServices/Private/Model/Regions.cpp
+++ b/Engine/Source/Developer/TraceServices/Private/Model/Regions.cpp
@@ -77,6 +77,27 @@ bool FRegionTimeline::EnumerateRegions(double IntervalStart, double IntervalEnd,
 	return true;
 }
 
+bool FRegionTimeline::EnumerateRegionsBackwards(double IntervalEnd, double IntervalStart,
+	TFunctionRef<bool(const FTimeRegion&)> Callback) const
+{
+	Provider.ReadAccessCheck();
+	
+	if (IntervalStart > IntervalEnd)
+	{
+		return false;
+	}
+
+	for (const FRegionLane& Lane : Lanes)
+	{
+		if (!Lane.EnumerateRegionsBackwards(IntervalEnd, IntervalStart, Callback))
+		{
+			return false;
+		}
+	}
+
+	return true;
+}
+
 void FRegionTimeline::EnumerateLanes(TFunctionRef<void(const FRegionLane&, const int32)> Callback) const
 {
 	Provider.ReadAccessCheck();
@@ -406,6 +427,28 @@ bool FRegionLane::EnumerateRegions(double IntervalStart, double IntervalEnd, TFu
 	return true;
 }
 
+bool FRegionLane::EnumerateRegionsBackwards(double IntervalEnd, double IntervalStart, TFunctionRef<bool(const FTimeRegion&)> Callback) const
+{
+	const FInt32Interval OverlapRange = GetElementRangeOverlappingGivenRange<FTimeRegion>(Regions, IntervalStart, IntervalEnd,
+		[](const FTimeRegion& r) { return r.BeginTime; },
+		[](const FTimeRegion& r) { return r.EndTime; });
+
+	if (OverlapRange.Min == -1)
+	{
+		return true;
+	}
+
+	for (int32 Index = OverlapRange.Max; Index >= OverlapRange.Min; --Index)
+	{
+		if (!Callback(Regions[Index]))
+		{
+			return false;
+		}
+	}
+
+	return true;
+}
+
 FName GetRegionProviderName()
 {
 	static const FName Name("RegionProvider");
diff --git a/Engine/Source/Developer/TraceServices/Private/Model/RegionsPrivate.h b/Engine/Source/Developer/TraceServices/Private/Model/RegionsPrivate.h
--- a/Engine/Source/Developer/TraceServices/Private/Model/RegionsPrivate.h
+++ b/Engine/Source/Developer/TraceServices/Private/Model/RegionsPrivate.h
@@ -62,6 +62,14 @@ public:
 	 */
 	UE_API bool EnumerateRegions(double IntervalStart, double IntervalEnd, TFunctionRef<bool(const FTimeRegion& /*Region*/)> Callback) const override;
 
+	/**
+	 * Enumerates all regions that overlap a certain time interval in reverse. Will enumerate by depth but does not expose lanes.
+	 * 
+	 * @param Callback a callback called for each region. Return false to abort iteration.
+	 * @returns true if the enumeration finished, false if it was aborted by the callback returning false
+	 */
+	UE_API bool EnumerateRegionsBackwards(double IntervalEnd, double IntervalStart, TFunctionRef<bool(const FTimeRegion& /*Region*/)> Callback) const override;
+
 	/**
 	 * Will call Callback(Lane, Depth) for each lane in order.
 	 */
@@ -113,6 +121,7 @@ public:
 	virtual const FRegionLane* GetLane(int32 Index) const override;
 
 	virtual bool EnumerateRegions(double IntervalStart, double IntervalEnd, TFunctionRef<bool(const FTimeRegion&)> Callback) const override;
+
 	virtual void EnumerateLanes(TFunctionRef<void(const FRegionLane&, const int32)> Callback) const override;
 
 	virtual uint64 GetUpdateCounter() const override { ReadAccessCheck(); return UpdateCounter; }
diff --git a/Engine/Source/Developer/TraceServices/Public/TraceServices/Model/Regions.h b/Engine/Source/Developer/TraceServices/Public/TraceServices/Model/Regions.h
--- a/Engine/Source/Developer/TraceServices/Public/TraceServices/Model/Regions.h
+++ b/Engine/Source/Developer/TraceServices/Public/TraceServices/Model/Regions.h
@@ -41,6 +41,13 @@ public:
 	 */
 	UE_API bool EnumerateRegions(double IntervalStart, double IntervalEnd, TFunctionRef<bool(const FTimeRegion&)> Callback) const;
 
+	/**
+	 * Call Callback for every region overlapping the interval defined by IntervalEnd and IntervalStart
+	 * @param Callback a callback called for each region. Return false to abort iteration.
+	 * @returns true if the enumeration finished, false if it was aborted by the callback returning false
+	 */
+	UE_API bool EnumerateRegionsBackwards(double IntervalEnd, double IntervalStart, TFunctionRef<bool(const FTimeRegion&)> Callback) const;
+
 private:
 	TPagedArray<FTimeRegion> Regions;
 };
@@ -83,6 +90,14 @@ public:
 	 * @returns true if the enumeration finished, false if it was aborted by the callback returning false
 	 */
 	virtual bool EnumerateRegions(double IntervalStart, double IntervalEnd, TFunctionRef<bool(const FTimeRegion&)> Callback) const = 0;
+
+	/**
+	 * Enumerates all regions that overlap a certain time interval in reverse. Will enumerate by depth but does not expose lanes.
+	 * 
+	 * @param Callback a callback called for each region. Return false to abort iteration.
+	 * @returns true if the enumeration finished, false if it was aborted by the callback returning false
+	 */
+	virtual bool EnumerateRegionsBackwards(double IntervalEnd, double IntervalStart, TFunctionRef<bool(const FTimeRegion&)> Callback) const = 0;
 	
 	/**
 	 * Will call Callback(Lane, Depth) for each lane in order.
