[Insights] Minor refactoring for MonotonicTimeline implementation (fixed typo; added more comments; fixed int/float/double constants).

#rnx

diff --git a/Engine/Source/Developer/TraceServices/Public/Model/MonotonicTimeline.h b/Engine/Source/Developer/TraceServices/Public/Model/MonotonicTimeline.h
--- a/Engine/Source/Developer/TraceServices/Public/Model/MonotonicTimeline.h
+++ b/Engine/Source/Developer/TraceServices/Public/Model/MonotonicTimeline.h
@@ -2,6 +2,7 @@
 
 #pragma once
 
+#include "CoreTypes.h"
 #include "Misc/QueuedThreadPool.h"
 
 #include "Common/PagedArray.h"
@@ -38,21 +39,26 @@ class IEditableTimeline
 public:
 	virtual ~IEditableTimeline() = default;
 
-	/*
-	* Begin a new timed event.
-	* 
-	* @param StartTime	The starting timestamp of the event in seconds.
-	* @param Event		The event information.
-	*/
+	/**
+	 * Begin a new timed event.
+	 *
+	 * @param StartTime	The starting timestamp of the event in seconds.
+	 * @param Event		The event information.
+	 */
 	virtual void AppendBeginEvent(double StartTime, const InEventType& Event) = 0;
 
-	/*
-	* End a new timed event. This ends the event started by the prior call to AppendBeginEvent.
-	* 
-	* @param EndTime	The ending timestamp of the event in seconds.
-	*/
+	/**
+	 * End a new timed event. This ends the event started by the prior call to AppendBeginEvent.
+	 *
+	 * @param EndTime	The ending timestamp of the event in seconds.
+	 */
 	virtual void AppendEndEvent(double EndTime) = 0;
 
+	/**
+	 * Gets the last time added to the timeline.
+	 *
+	 * @return The last timestamp added to the timeline.
+	 */
 	virtual double GetLastTimestamp() const { return 0.0; }
 };
 
@@ -66,7 +72,7 @@ class TMonotonicTimeline
 	typedef FEventInfoStackEntry<InEventType, SettingsType> FEventInfoStackEntry;
 	typedef FDetailLevel<InEventType, SettingsType> FDetailLevel;
 	typedef FDetailLevelDepthState<InEventType> FDetailLevelDepthState;
-	typedef FEnumerateAsyncTask<InEventType, SettingsType> FEnumarateAsyncTask;
+	typedef FEnumerateAsyncTask<InEventType, SettingsType> FEnumerateAsyncTask;
 
 public:
 	using EventType = InEventType;
@@ -83,11 +89,6 @@ public:
 
 	virtual ~TMonotonicTimeline() = default;
 
-	virtual double GetLastTimestamp() const override
-	{
-		return DetailLevels[0].InsertionState.LastTime;
-	}
-
 	virtual void AppendBeginEvent(double StartTime, const EventType& Event) override
 	{
 		int32 CurrentDepth = DetailLevels[0].InsertionState.CurrentDepth;
@@ -224,6 +225,11 @@ public:
 		return NumScopeEntries > 0 ? FMath::Abs(DetailLevels[0].ScopeEntries[NumScopeEntries - 1].Time) : 0.0;
 	}
 
+	virtual double GetLastTimestamp() const override
+	{
+		return DetailLevels[0].InsertionState.LastTime;
+	}
+
 	virtual void EnumerateEventsDownSampled(double IntervalStart, double IntervalEnd, double Resolution, typename ITimeline<EventType>::EventCallback Callback) const override
 	{
 		int32 DetailLevelIndex = SettingsType::DetailLevelsCount - 1;
@@ -241,10 +247,13 @@ public:
 			return;
 		}
 
-		uint64 FirstScopePageIndex = Algo::UpperBoundBy(DetailLevel.ScopeEntries, IntervalStart, [](const FEventScopeEntryPage& Page)
-		{
-			return Page.BeginTime;
-		});
+		// The page of the first scope...
+		// i.e. the previous page of the first page with Page.BeginTime > IntervalStart
+		uint64 FirstScopePageIndex = Algo::UpperBoundBy(DetailLevel.ScopeEntries, IntervalStart,
+			[](const FEventScopeEntryPage& Page)
+			{
+				return Page.BeginTime;
+			});
 		if (FirstScopePageIndex > 0)
 		{
 			--FirstScopePageIndex;
@@ -265,7 +274,10 @@ public:
 			EventType Event;
 		};
 		FEnumerationStackEntry EventStack[SettingsType::MaxDepth];
-		int32 CurrentStackDepth = ScopePage->InitialStackCount;
+		int32 CurrentStackDepth;
+
+		// Make a copy of the initial stack of the page.
+		CurrentStackDepth = ScopePage->InitialStackCount;
 		for (int32 InitialStackIndex = 0; InitialStackIndex < CurrentStackDepth; ++InitialStackIndex)
 		{
 			FEnumerationStackEntry& EnumerationStackEntry = EventStack[InitialStackIndex];
@@ -282,7 +294,7 @@ public:
 
 		while (ScopeEntry && FMath::Abs(ScopeEntry->Time) < IntervalStart)
 		{
-			if (ScopeEntry->Time < 0.0)
+			if (ScopeEntry->Time < 0.0) // "begin" event
 			{
 				check(CurrentStackDepth < SettingsType::MaxDepth);
 				FEnumerationStackEntry& StackEntry = EventStack[CurrentStackDepth++];
@@ -290,7 +302,7 @@ public:
 				StackEntry.StartTime = -ScopeEntry->Time;
 				Event = EventsIterator.NextItem();
 			}
-			else
+			else // "end" event
 			{
 				check(CurrentStackDepth > 0);
 				--CurrentStackDepth;
@@ -301,6 +313,7 @@ public:
 		{
 			return;
 		}
+
 		for (int32 StackIndex = 0; StackIndex < CurrentStackDepth; ++StackIndex)
 		{
 			FEnumerationStackEntry& StackEntry = EventStack[StackIndex];
@@ -311,7 +324,7 @@ public:
 		}
 		while (ScopeEntry && FMath::Abs(ScopeEntry->Time) <= IntervalEnd)
 		{
-			if (ScopeEntry->Time < 0.0)
+			if (ScopeEntry->Time < 0.0) // "begin" event
 			{
 				check(CurrentStackDepth < SettingsType::MaxDepth);
 				FEnumerationStackEntry& StackEntry = EventStack[CurrentStackDepth++];
@@ -322,7 +335,7 @@ public:
 				}
 				Event = EventsIterator.NextItem();
 			}
-			else
+			else // "end" event
 			{
 				check(CurrentStackDepth > 0);
 				FEnumerationStackEntry& StackEntry = EventStack[--CurrentStackDepth];
@@ -344,11 +357,11 @@ public:
 				bSearchEndTimeUsingPages = true;
 				break;
 			}
-			if (ScopeEntry->Time < 0.0)
+			if (ScopeEntry->Time < 0.0) // "begin" event
 			{
 				++ExitDepth;
 			}
-			else
+			else // "end" event
 			{
 				if (ExitDepth == 0)
 				{
@@ -374,7 +387,7 @@ public:
 			do
 			{
 				check(CurrentStackDepth <= CurrentScopePage->InitialStackCount);
-				while (CurrentStackDepth > 0 && CurrentScopePage->InitialStack[CurrentStackDepth - 1].EndTime > 0)
+				while (CurrentStackDepth > 0 && CurrentScopePage->InitialStack[CurrentStackDepth - 1].EndTime > 0.0)
 				{
 					--CurrentStackDepth;
 					EventType CurrentEvent = DetailLevel.GetEvent(CurrentScopePage->InitialStack[CurrentStackDepth].EventIndex);
@@ -398,6 +411,64 @@ public:
 		}
 	}
 
+	virtual void EnumerateEventsDownSampled(double IntervalStart, double IntervalEnd, double Resolution, typename ITimeline<EventType>::EventRangeCallback Callback) const override
+	{
+		struct FStackEntry
+		{
+			uint64 LocalEventIndex;
+		};
+		FStackEntry EventStack[SettingsType::MaxDepth];
+		uint32 CurrentDepth = 0;
+
+		struct FOutputEvent
+		{
+			double StartTime;
+			double EndTime;
+			uint32 Depth;
+			EventType Event;
+		};
+		TArray<FOutputEvent> OutputEvents;
+
+		EnumerateEventsDownSampled(IntervalStart, IntervalEnd, Resolution,
+			[&EventStack, &OutputEvents, &CurrentDepth, Callback]
+			(bool IsEnter, double Time, const EventType& Event)
+			{
+				if (IsEnter)
+				{
+					FStackEntry& StackEntry = EventStack[CurrentDepth];
+					StackEntry.LocalEventIndex = OutputEvents.Num();
+					FOutputEvent& OutputEvent = OutputEvents.AddDefaulted_GetRef();
+					OutputEvent.StartTime = Time;
+					OutputEvent.EndTime = Time;
+					OutputEvent.Depth = CurrentDepth;
+					OutputEvent.Event = Event;
+					++CurrentDepth;
+				}
+				else
+				{
+					{
+						FStackEntry& StackEntry = EventStack[--CurrentDepth];
+						FOutputEvent* OutputEvent = OutputEvents.GetData() + StackEntry.LocalEventIndex;
+						OutputEvent->EndTime = Time;
+					}
+					if (CurrentDepth == 0)
+					{
+						for (FOutputEvent& OutputEvent : OutputEvents)
+						{
+							EEventEnumerate Ret = Callback(OutputEvent.StartTime, OutputEvent.EndTime, OutputEvent.Depth, OutputEvent.Event);
+							if (Ret != EEventEnumerate::Continue)
+							{
+								return Ret;
+							}
+						}
+						OutputEvents.Empty(OutputEvents.Num());
+					}
+				}
+
+				return EEventEnumerate::Continue;
+			});
+	}
+
 	virtual void EnumerateEventsBackwardsDownSampled(double IntervalEnd, double IntervalStart, double Resolution, typename ITimeline<EventType>::EventCallback Callback) const override
 	{
 		int32 DetailLevelIndex = SettingsType::DetailLevelsCount - 1;
@@ -415,7 +486,10 @@ public:
 			return;
 		}
 
-		uint64 LastScopePageIndex = Algo::UpperBoundBy(DetailLevel.ScopeEntries, IntervalEnd, [](const FEventScopeEntryPage& Page)
+		// The page of the last scope...
+		// i.e. the first page with Page.BeginTime > IntervalEnd
+		uint64 LastScopePageIndex = Algo::UpperBoundBy(DetailLevel.ScopeEntries, IntervalEnd,
+			[](const FEventScopeEntryPage& Page)
 			{
 				return Page.BeginTime;
 			});
@@ -425,17 +499,18 @@ public:
 			double EndTime;
 			EventType Event;
 		};
+		FEnumerationStackEntry EventStack[SettingsType::MaxDepth];
+		int32 CurrentStackDepth = 0;
+
+		const FEventScopeEntry* ScopeEntry = nullptr;
 
 		// By default, we start from the very end of the session.
 		auto ScopeEntryIterator = DetailLevel.ScopeEntries.GetIteratorFromItem(DetailLevel.ScopeEntries.Num() - 1);
 		auto EventsIterator = DetailLevel.Events.GetIteratorFromItem(DetailLevel.Events.Num() - 1);
-		FEnumerationStackEntry EventStack[SettingsType::MaxDepth];
-		int32 CurrentStackDepth = 0;
-		const FEventScopeEntry* ScopeEntry = nullptr;
 
 		if (LastScopePageIndex > 0 && LastScopePageIndex < DetailLevel.ScopeEntries.NumPages())
 		{
-			// If we have a page we can start from, start enumerating backwards from the begining of that page.
+			// If we have a page we can start from, start enumerating backwards from the beginning of that page.
 			ScopeEntryIterator = DetailLevel.ScopeEntries.GetIteratorFromPage(LastScopePageIndex);
 			const FEventScopeEntryPage* ScopePage = DetailLevel.ScopeEntries.GetPage(LastScopePageIndex);
 
@@ -446,19 +521,19 @@ public:
 				FEnumerationStackEntry& EnumerationStackEntry = EventStack[InitialStackIndex];
 				const FEventStackEntry& EventStackEntry = ScopePage->InitialStack[InitialStackIndex];
 				EnumerationStackEntry.EndTime = EventStackEntry.EndTime;
-				if (EnumerationStackEntry.EndTime < 0)
+				if (EnumerationStackEntry.EndTime < 0.0)
 				{
 					// We need to search for the EndTime of the event using pages.
 					auto PageIterator = ScopeEntryIterator;
 					while (const FEventScopeEntryPage* Page = PageIterator.NextPage())
 					{
-						if (Page->InitialStack[InitialStackIndex].EndTime > 0)
+						if (Page->InitialStack[InitialStackIndex].EndTime > 0.0)
 						{
 							EnumerationStackEntry.EndTime = Page->InitialStack[InitialStackIndex].EndTime;
 							break;
 						}
 					}
-					if (EnumerationStackEntry.EndTime < 0)
+					if (EnumerationStackEntry.EndTime < 0.0)
 					{
 						EnumerationStackEntry.EndTime = DetailLevel.InsertionState.LastTime;
 					}
@@ -483,17 +558,17 @@ public:
 			ScopeEntry = ScopeEntryIterator.GetCurrentItem();
 		}
 
-		// Enumerate backwards until we reach IntervalEnd, without calling the Callback because thess events are not in the provided interval.
+		// Enumerate backwards until we reach IntervalEnd, without calling the Callback because these events are not in the provided interval.
 		const EventType* Event = EventsIterator.GetCurrentItem();
 		while (ScopeEntry && FMath::Abs(ScopeEntry->Time) > IntervalEnd)
 		{
-			if (ScopeEntry->Time < 0.0)
+			if (ScopeEntry->Time < 0.0) // "begin" event
 			{
 				check(CurrentStackDepth > 0);
 				--CurrentStackDepth;
 				Event = EventsIterator.PrevItem();
 			}
-			else
+			else // "end" event
 			{
 				check(CurrentStackDepth < SettingsType::MaxDepth);
 				FEnumerationStackEntry& StackEntry = EventStack[CurrentStackDepth++];
@@ -515,7 +590,7 @@ public:
 		// Enumerate backwards between IntervalEnd and IntervalStart.
 		while (ScopeEntry && FMath::Abs(ScopeEntry->Time) >= IntervalStart)
 		{
-			if (ScopeEntry->Time < 0.0)
+			if (ScopeEntry->Time < 0.0) // "begin" event
 			{
 				check(CurrentStackDepth > 0);
 				FEnumerationStackEntry& StackEntry = EventStack[--CurrentStackDepth];
@@ -526,7 +601,7 @@ public:
 				}
 				Event = EventsIterator.PrevItem();
 			}
-			else
+			else // "end" event
 			{
 				check(CurrentStackDepth < SettingsType::MaxDepth);
 				FEnumerationStackEntry& StackEntry = EventStack[CurrentStackDepth++];
@@ -551,7 +626,7 @@ public:
 				bSearchStartTimeUsingPages = true;
 				break;
 			}
-			if (ScopeEntry->Time < 0.0)
+			if (ScopeEntry->Time < 0.0) // "begin" event
 			{
 				if (ExitDepth == 0)
 				{
@@ -567,7 +642,7 @@ public:
 				}
 				EventsIterator.PrevItem();
 			}
-			else
+			else // "end" event
 			{
 				++ExitDepth;
 			}
@@ -605,27 +680,29 @@ public:
 		FStackEntry EventStack[SettingsType::MaxDepth];
 		uint32 CurrentDepth = 0;
 
-		EnumerateEventsBackwardsDownSampled(IntervalEnd, IntervalStart, Resolution, [&EventStack, &CurrentDepth, Callback](bool IsEnter, double Time, const EventType& Event)
-		{
-			if (IsEnter)
+		EnumerateEventsBackwardsDownSampled(IntervalEnd, IntervalStart, Resolution,
+			[&EventStack, &CurrentDepth, Callback]
+			(bool IsEnter, double Time, const EventType& Event)
 			{
-				FStackEntry& StackEntry = EventStack[CurrentDepth];
-				StackEntry.Event = Event;
-				StackEntry.EndTime = Time;
-				++CurrentDepth;
-			}
-			else
-			{
-				FStackEntry& StackEntry = EventStack[--CurrentDepth];
-				EEventEnumerate Ret = Callback(Time, StackEntry.EndTime, CurrentDepth, Event);
-				if (Ret != EEventEnumerate::Continue)
+				if (IsEnter)
 				{
-					return Ret;
+					FStackEntry& StackEntry = EventStack[CurrentDepth];
+					StackEntry.Event = Event;
+					StackEntry.EndTime = Time;
+					++CurrentDepth;
+				}
+				else
+				{
+					FStackEntry& StackEntry = EventStack[--CurrentDepth];
+					EEventEnumerate Ret = Callback(Time, StackEntry.EndTime, CurrentDepth, Event);
+					if (Ret != EEventEnumerate::Continue)
+					{
+						return Ret;
+					}
 				}
-			}
 
-			return EEventEnumerate::Continue;
-		});
+				return EEventEnumerate::Continue;
+			});
 	}
 
 	virtual void EnumerateEventsDownSampledAsync(const typename ITimeline<EventType>::EnumerateAsyncParams& EnumerateAsyncParams) const override
@@ -653,28 +730,33 @@ public:
 			return;
 		}
 
-		uint64 FirstScopePageIndex = Algo::UpperBoundBy(DetailLevel.ScopeEntries, EnumerateAsyncParams.IntervalStart, [](const FEventScopeEntryPage& Page)
+		// The page of the first scope...
+		// i.e. the previous page of the first page with Page.BeginTime > EnumerateAsyncParams.IntervalStart
+		uint64 FirstScopePageIndex = Algo::UpperBoundBy(DetailLevel.ScopeEntries, EnumerateAsyncParams.IntervalStart,
+			[](const FEventScopeEntryPage& Page)
 			{
 				return Page.BeginTime;
 			});
-
 		if (FirstScopePageIndex > 0)
 		{
 			--FirstScopePageIndex;
 		}
 
-		uint64 LastScopePageIndex = Algo::UpperBoundBy(DetailLevel.ScopeEntries, EnumerateAsyncParams.IntervalEnd, [](const FEventScopeEntryPage& Page)
+		// The page of the last scope...
+		// i.e. the first page with Page.BeginTime > EnumerateAsyncParams.IntervalEnd
+		uint64 LastScopePageIndex = Algo::UpperBoundBy(DetailLevel.ScopeEntries, EnumerateAsyncParams.IntervalEnd,
+			[](const FEventScopeEntryPage& Page)
 			{
 				return Page.BeginTime;
 			});
 
-		TArray<TSharedRef<FAsyncTask<FEnumarateAsyncTask>>> WorkerTasks;
+		TArray<TSharedRef<FAsyncTask<FEnumerateAsyncTask>>> WorkerTasks;
 
 		check(LastScopePageIndex >= FirstScopePageIndex);
 		uint32 NumPages = FMath::Max(static_cast<uint32>(LastScopePageIndex - FirstScopePageIndex), 1u);
 		uint32 NumThreads = GThreadPool->GetNumThreads();
 
-		if (EnumerateAsyncParams.MaxOccupancy > 0.0f)
+		if (EnumerateAsyncParams.MaxOccupancy > 0.0)
 		{
 			NumThreads = FMath::Max(static_cast<uint32>(NumThreads * EnumerateAsyncParams.MaxOccupancy), 1u);
 		}
@@ -702,7 +784,7 @@ public:
 			TaskData.EventCallback = EnumerateAsyncParams.EventCallback;
 			TaskData.EventRangeCallback = EnumerateAsyncParams.EventRangeCallback;
 
-			TSharedRef<FAsyncTask<FEnumarateAsyncTask>> AsyncTask = MakeShared<FAsyncTask<FEnumarateAsyncTask>>(TaskData);
+			TSharedRef<FAsyncTask<FEnumerateAsyncTask>> AsyncTask = MakeShared<FAsyncTask<FEnumerateAsyncTask>>(TaskData);
 			WorkerTasks.Add(AsyncTask);
 			AsyncTask->StartBackgroundTask();
 
@@ -715,62 +797,6 @@ public:
 		}
 	}
 
-	virtual void EnumerateEventsDownSampled(double IntervalStart, double IntervalEnd, double Resolution, typename ITimeline<EventType>::EventRangeCallback Callback) const override
-	{
-		struct FStackEntry
-		{
-			uint64 LocalEventIndex;
-		};
-		FStackEntry EventStack[SettingsType::MaxDepth];
-		uint32 CurrentDepth = 0;
-
-		struct FOutputEvent
-		{
-			double StartTime;
-			double EndTime;
-			uint32 Depth;
-			EventType Event;
-		};
-		TArray<FOutputEvent> OutputEvents;
-
-		EnumerateEventsDownSampled(IntervalStart, IntervalEnd, Resolution, [&EventStack, &OutputEvents, &CurrentDepth, Callback](bool IsEnter, double Time, const EventType& Event)
-		{
-			if (IsEnter)
-			{
-				FStackEntry& StackEntry = EventStack[CurrentDepth];
-				StackEntry.LocalEventIndex = OutputEvents.Num();
-				FOutputEvent& OutputEvent = OutputEvents.AddDefaulted_GetRef();
-				OutputEvent.StartTime = Time;
-				OutputEvent.EndTime = Time;
-				OutputEvent.Depth = CurrentDepth;
-				OutputEvent.Event = Event;
-				++CurrentDepth;
-			}
-			else
-			{
-				{
-					FStackEntry& StackEntry = EventStack[--CurrentDepth];
-					FOutputEvent* OutputEvent = OutputEvents.GetData() + StackEntry.LocalEventIndex;
-					OutputEvent->EndTime = Time;
-				}
-				if (CurrentDepth == 0)
-				{
-					for (FOutputEvent& OutputEvent : OutputEvents)
-					{
-						EEventEnumerate Ret = Callback(OutputEvent.StartTime, OutputEvent.EndTime, OutputEvent.Depth, OutputEvent.Event);
-						if (Ret != EEventEnumerate::Continue)
-						{
-							return Ret;
-						}
-					}
-					OutputEvents.Empty(OutputEvents.Num());
-				}
-			}
-
-			return EEventEnumerate::Continue;
-		});
-	}
-
 	virtual void EnumerateEvents(double IntervalStart, double IntervalEnd, typename ITimeline<EventType>::EventCallback Callback) const override
 	{
 		EnumerateEventsDownSampled(IntervalStart, IntervalEnd, 0.0, Callback);
@@ -810,11 +836,13 @@ public:
 			return false;
 		}
 
-		uint64 FirstScopePageIndex = Algo::UpperBoundBy(DetailLevel.ScopeEntries, InTime, [](const FEventScopeEntryPage& Page)
+		// The page of the first scope...
+		// i.e. the previous page of the first page with Page.BeginTime > InTime
+		uint64 FirstScopePageIndex = Algo::UpperBoundBy(DetailLevel.ScopeEntries, InTime,
+			[](const FEventScopeEntryPage& Page)
 			{
 				return Page.BeginTime;
 			});
-
 		if (FirstScopePageIndex > 0)
 		{
 			--FirstScopePageIndex;
@@ -831,7 +859,7 @@ public:
 			EventInfo.EndTime = OutScopeEntry.EndTime;
 			EventInfo.Event = DetailLevel.GetEvent(OutScopeEntry.EventIndex);
 
-			if (EventInfo.EndTime < 0)
+			if (EventInfo.EndTime < 0.0)
 			{
 				//The end of the event has not been reached by analysis
 				EventInfo.EndTime = DetailLevel.InsertionState.LastTime;
@@ -868,10 +896,10 @@ public:
 			LastTime = abs(LastTimeIterator.GetCurrentItem()->Time);
 		}
 
-		//Iterate from the start of the page to the start time of our interval
+		// Iterate from the start of the page to the start time of our interval.
 		while (ScopeEntry && FMath::Abs(ScopeEntry->Time) <= IntervalStart)
 		{
-			if (ScopeEntry->Time < 0.0)
+			if (ScopeEntry->Time < 0.0) // "begin" event
 			{
 				check(CurrentStackDepth < SettingsType::MaxDepth);
 				FEventInfoStackEntry& StackEntry = EventStack[CurrentStackDepth++];
@@ -886,7 +914,7 @@ public:
 				Event = EventsIterator.NextItem();
 				LastTime = -ScopeEntry->Time;
 			}
-			else
+			else // "end" event
 			{
 				check(CurrentStackDepth > 0);
 				--CurrentStackDepth;
@@ -902,7 +930,7 @@ public:
 			ScopeEntry = ScopeEntryIterator.NextItem();
 		}
 
-		//Check if we have an event between InTime - InDeltaTime and InTime + InDeltaTime
+		// Check if we have an event between InTime - InDeltaTime and InTime + InDeltaTime.
 		FFindBestMatchEventInParams InParams;
 		InParams.IterationState.ScopeIterator = ScopeEntryIterator;
 		InParams.IterationState.EventsIterator = EventsIterator;
@@ -932,7 +960,7 @@ public:
 			}
 		}
 
-		//If we have found both start time and end time for our event, we can return the result
+		// If we have found both start time and end time for our event, we can return the result.
 		if (OutParams.bHasEndTime && bMatchEventStartsInsidePage)
 		{
 			EventInfo.StartTime = OutParams.EventInfo.StartTime;
@@ -940,16 +968,16 @@ public:
 			EventInfo.ExclTime = OutParams.EventInfo.ExclTime;
 			EventInfo.Event = OutParams.EventInfo.Event;
 
-			if (EventInfo.EndTime < 0)
+			if (EventInfo.EndTime < 0.0)
 			{
-				//The end of the event has not been reached by analysis
+				// The end of the event has not been reached by analysis.
 				EventInfo.EndTime = DetailLevel.InsertionState.LastTime;
 			}
 
 			return true;
 		}
 
-		//We continue searching for the end scope event
+		// We continue searching for the end scope event.
 		ScopeEntryIterator = OutParams.IterationState.ScopeIterator;
 		EventsIterator = OutParams.IterationState.EventsIterator;
 		CurrentStackDepth = OutParams.IterationState.StackDepth;
@@ -959,7 +987,7 @@ public:
 
 		auto ScopeEntryIteratorAtEvent = ScopeEntryIterator;
 
-		//We find the page where the target event ends
+		// We find the page where the target event ends.
 		ScopeEntryIterator = DetailLevel.ScopeEntries.GetIteratorFromPage(FirstScopePageIndex);
 		auto EventLastPageIterator = ScopeEntryIterator;
 
@@ -981,7 +1009,7 @@ public:
 
 		if (EventLastPageIterator.GetCurrentPageIndex() != FirstScopePageIndex || !bMatchEventStartsInsidePage)
 		{
-			//If the end scope event is on a different page than the start scope one we can get the event info from the InitialStack of the last page
+			// If the end scope event is on a different page than the start scope one we can get the event info from the InitialStack of the last page.
 			ScopeEntryIterator = EventLastPageIterator;
 			ScopePage = ScopeEntryIterator.GetCurrentPage();
 
@@ -992,9 +1020,9 @@ public:
 			EventInfo.ExclTime = TargetStackEntry.ExclTime;
 			EventInfo.Event = DetailLevel.GetEvent(TargetStackEntry.EventIndex);
 
-			if (EventInfo.EndTime < 0)
+			if (EventInfo.EndTime < 0.0)
 			{
-				//The end of the event has not been reached by analysis
+				// The end of the event has not been reached by analysis.
 				EventInfo.EndTime = DetailLevel.InsertionState.LastTime;
 			}
 
@@ -1002,7 +1030,7 @@ public:
 		}
 		else
 		{
-			//If the end scope event is in the same page, we continue iterating from the point where FindBestMatchEvent stopped
+			// If the end scope event is in the same page, we continue iterating from the point where FindBestMatchEvent stopped.
 			ScopeEntryIterator = ScopeEntryIteratorAtEvent;
 			ScopeEntry = ScopeEntryIterator.GetCurrentItem();
 		}
@@ -1010,7 +1038,7 @@ public:
 		while (ScopeEntry &&
 			   Depth < CurrentStackDepth)
 		{
-			if (ScopeEntry->Time < 0.0)
+			if (ScopeEntry->Time < 0.0) // "begin" event
 			{
 				check(CurrentStackDepth < SettingsType::MaxDepth);
 				FEventInfoStackEntry& StackEntry = EventStack[CurrentStackDepth++];
@@ -1024,7 +1052,7 @@ public:
 				Event = EventsIterator.NextItem();
 				LastTime = -ScopeEntry->Time;
 			}
-			else
+			else // "end" event
 			{
 				check(CurrentStackDepth > 0);
 				--CurrentStackDepth;
@@ -1057,7 +1085,10 @@ public:
 			return 0;
 		}
 
-		uint64 FirstScopePageIndex = Algo::UpperBoundBy(DetailLevel.ScopeEntries, Time, [](const FEventScopeEntryPage& Page)
+		// The page of the first scope...
+		// i.e. the previous page of the first page with Page.BeginTime > Time
+		uint64 FirstScopePageIndex = Algo::UpperBoundBy(DetailLevel.ScopeEntries, Time,
+			[](const FEventScopeEntryPage& Page)
 			{
 				return Page.BeginTime;
 			});
@@ -1065,6 +1096,7 @@ public:
 		{
 			--FirstScopePageIndex;
 		}
+
 		const FEventScopeEntryPage* ScopePage = DetailLevel.ScopeEntries.GetPage(FirstScopePageIndex);
 		if (ScopePage->BeginTime > Time)
 		{
@@ -1077,11 +1109,11 @@ public:
 		const FEventScopeEntry* ScopeEntry = ScopeEntryIterator.GetCurrentItem();
 		while (ScopeEntry && FMath::Abs(ScopeEntry->Time) < Time)
 		{
-			if (ScopeEntry->Time < 0.0)
+			if (ScopeEntry->Time < 0.0) // "begin" event
 			{
 				CurrentStackDepth++;
 			}
-			else
+			else // "end" event
 			{
 				check(CurrentStackDepth > 0);
 				--CurrentStackDepth;
@@ -1262,7 +1294,7 @@ private:
 		if (CurrentScopePage->InitialStackCount > Depth)
 		{
 			FEventStackEntry& CurrentPageStackEntry = CurrentScopePage->InitialStack[Depth];
-			if (CurrentPageStackEntry.EndTime < 0 || CurrentPageStackEntry.EndTime > Time)
+			if (CurrentPageStackEntry.EndTime < 0.0 || CurrentPageStackEntry.EndTime > Time)
 			{
 				bIsInCurrentPageInitStack = true;
 			}
@@ -1289,7 +1321,7 @@ private:
 			return false;
 		}
 
-		while (CurrentScopePage->InitialStack[Depth].EndTime < 0)
+		while (CurrentScopePage->InitialStack[Depth].EndTime < 0.0)
 		{
 			const FEventScopeEntryPage* NextScopePage = ScopeEntryIterator.NextPage();
 			if (NextScopePage == nullptr)
@@ -1314,11 +1346,11 @@ private:
 		double BestMatchEndTime = 0.0;
 		bool bHasEndEvent = false;
 
-		//In the first step, we iterate up to TargetExactTime, storing the last event with our target depth that ended during iteration
-		//......]...]....]...TargetExactTime
+		// In the first step, we iterate up to TargetExactTime, storing the last event with our target depth that ended during iteration.
+		// ......]...]....]...TargetExactTime
 		while (ScopeEntry && FMath::Abs(ScopeEntry->Time) <= InParams.TargetExactTime)
 		{
-			if (ScopeEntry->Time < 0.0)
+			if (ScopeEntry->Time < 0.0) // "begin" event
 			{
 				check(IterationState.StackDepth < SettingsType::MaxDepth);
 				FEventInfoStackEntry& StackEntry = IterationState.EventStack[IterationState.StackDepth++];
@@ -1333,7 +1365,7 @@ private:
 				Event = IterationState.EventsIterator.NextItem();
 				IterationState.LastIterationTime = -ScopeEntry->Time;
 			}
-			else
+			else // "end" event
 			{
 				check(IterationState.StackDepth > 0);
 				--IterationState.StackDepth;
@@ -1352,10 +1384,10 @@ private:
 			ScopeEntry = IterationState.ScopeIterator.NextItem();
 		}
 
-		//If the iteration stack depth is as deep as our target depth, that we have a perfect match,
-		//an event that is ongoing at TargetExactTime, so we just return it
-		//....[..TargetExactTime.....]
-		//The strict '>' comparison is needed because CurrentStackDepth is "RealDepth + 1"
+		// If the iteration stack depth is as deep as our target depth, that we have a perfect match,
+		// an event that is ongoing at TargetExactTime, so we just return it.
+		// ....[..TargetExactTime.....]
+		// The strict '>' comparison is needed because CurrentStackDepth is "RealDepth + 1".
 		if (IterationState.StackDepth > InParams.TargetDepth)
 		{
 			OutParams.EventInfo = IterationState.EventStack[InParams.TargetDepth];
@@ -1365,12 +1397,12 @@ private:
 			return true;
 		}
 
-		//We continue iterating until TargetEndTime or until we find the start of an event with the target depth
-		//TargetExactTime.....[
+		// We continue iterating until TargetEndTime or until we find the start of an event with the target depth.
+		// TargetExactTime.....[
 		bool bHasStartEvent = false;
 		while (ScopeEntry && FMath::Abs(ScopeEntry->Time) <= InParams.TargetEndTime)
 		{
-			if (ScopeEntry->Time < 0.0)
+			if (ScopeEntry->Time < 0.0) // "begin" event
 			{
 				check(IterationState.StackDepth < SettingsType::MaxDepth);
 				FEventInfoStackEntry& StackEntry = IterationState.EventStack[IterationState.StackDepth++];
@@ -1393,7 +1425,7 @@ private:
 					break;
 				}
 			}
-			else
+			else // "end" event
 			{
 				check(IterationState.StackDepth > 0);
 				--IterationState.StackDepth;
@@ -1414,7 +1446,7 @@ private:
 			return false;
 		}
 
-		//We choose the event that is closer to TargetExactTime
+		// We choose the event that is closer to TargetExactTime.
 		if (bHasStartEvent)
 		{
 			double EndTimeDelta = InParams.TargetExactTime - BestMatchEndTime;
diff --git a/Engine/Source/Developer/TraceServices/Public/TraceServices/Containers/Timelines.h b/Engine/Source/Developer/TraceServices/Public/TraceServices/Containers/Timelines.h
--- a/Engine/Source/Developer/TraceServices/Public/TraceServices/Containers/Timelines.h
+++ b/Engine/Source/Developer/TraceServices/Public/TraceServices/Containers/Timelines.h
@@ -2,6 +2,7 @@
 
 #pragma once
 
+#include "CoreTypes.h"
 #include "Templates/Function.h"
 
 namespace TraceServices
@@ -32,8 +33,8 @@ public:
 
 	struct FTimelineEventInfo
 	{
-		double StartTime;
-		double EndTime;
+		double StartTime = 0.0;
+		double EndTime = 0.0;
 		double ExclTime = 0.0;
 		EventType Event;
 	};
@@ -42,7 +43,7 @@ public:
 	{
 		double IntervalStart = 0.0;
 		double IntervalEnd = -1.0;
-		double Resolution = 0;
+		double Resolution = 0.0;
 		double MaxOccupancy = 0.875; // Leave 1/8 threads free to avoid starvation.
 		EEventSortOrder SortOrder = EEventSortOrder::ByEndTime;
 		typename ITimeline<EventType>::TaskInfoCallback SetupCallback;
@@ -53,30 +54,34 @@ public:
 	};
 
 	virtual ~ITimeline() = default;
+
 	virtual uint64 GetModCount() const = 0;
+
 	virtual uint64 GetEventCount() const = 0;
 	virtual const InEventType& GetEvent(uint64 InIndex) const = 0;
+
 	virtual double GetStartTime() const = 0;
 	virtual double GetEndTime() const = 0;
+
 	virtual void EnumerateEventsDownSampled(double IntervalStart, double IntervalEnd, double Resolution, EventCallback Callback) const = 0;
 	virtual void EnumerateEventsDownSampled(double IntervalStart, double IntervalEnd, double Resolution, EventRangeCallback Callback) const = 0;
 
-	virtual void EnumerateEventsBackwardsDownSampled(double IntervalEnd, double IntervalStart, double Resolution, EventCallback Callback) const {};
-	virtual void EnumerateEventsBackwardsDownSampled(double IntervalEnd, double IntervalStart, double Resolution, EventRangeCallback Callback) const {};
+	virtual void EnumerateEventsBackwardsDownSampled(double IntervalEnd, double IntervalStart, double Resolution, EventCallback Callback) const {}
+	virtual void EnumerateEventsBackwardsDownSampled(double IntervalEnd, double IntervalStart, double Resolution, EventRangeCallback Callback) const {}
 
-	virtual void EnumerateEventsDownSampledAsync(const EnumerateAsyncParams& EnumerateAsyncParams) const {};
+	virtual void EnumerateEventsDownSampledAsync(const EnumerateAsyncParams& EnumerateAsyncParams) const {}
 
 	virtual void EnumerateEvents(double IntervalStart, double IntervalEnd, EventCallback Callback) const = 0;
 	virtual void EnumerateEvents(double IntervalStart, double IntervalEnd, EventRangeCallback Callback) const = 0;
 
-	virtual void EnumerateEventsBackwards(double IntervalEnd, double IntervalStart, EventCallback Callback) const {};
-	virtual void EnumerateEventsBackwards(double IntervalEnd, double IntervalStart, EventRangeCallback Callback) const {};
+	virtual void EnumerateEventsBackwards(double IntervalEnd, double IntervalStart, EventCallback Callback) const {}
+	virtual void EnumerateEventsBackwards(double IntervalEnd, double IntervalStart, EventRangeCallback Callback) const {}
+
+	virtual int32 GetDepthAt(double Time) const { return 0; }
 
-	virtual int32 GetDepthAt(double Time) const { return 0; };
-	
 	/**
 	 * Finds event information for the event closest to InTime from the interval [InTime - DeltaTime, InTime + DeltaTime]
-	 * @param InTime - The time used to query for the event 
+	 * @param InTime - The time used to query for the event
 	 * @param DeltaTime - Events from interval [InTime - DeltaTime, InTime + DeltaTime] will be considered. The one closest to InTime will be returned
 	 * @param Depth - The Depth used to query for the event
 	 * @return True if an event was found, False if no event was found for the specified input parameters
