[Insights] Region Timing Track: Improved Enumeration algorithm for region events

#rnx

diff --git a/Engine/Source/Developer/TraceServices/Private/Model/Regions.cpp b/Engine/Source/Developer/TraceServices/Private/Model/Regions.cpp
--- a/Engine/Source/Developer/TraceServices/Private/Model/Regions.cpp
+++ b/Engine/Source/Developer/TraceServices/Private/Model/Regions.cpp
@@ -66,12 +66,46 @@ bool FRegionTimeline::EnumerateRegions(double IntervalStart, double IntervalEnd,
 		return false;
 	}
 
-	for (const FRegionLane& Lane : Lanes)
+	TArray<FLaneEnumeration> SearchLanes;
+	for (const FRegionLane& Lane: Lanes)
 	{
-		if (!Lane.EnumerateRegions(IntervalStart, IntervalEnd, Callback))
+		const FInt32Interval OverlapRange = GetElementRangeOverlappingGivenRange<FTimeRegion>(Lane.Regions, IntervalStart, IntervalEnd,
+[](const FTimeRegion& r) { return r.BeginTime; },
+[](const FTimeRegion& r) { return r.EndTime; });
+		
+		if (OverlapRange.Min != -1 && OverlapRange.Size() >= 0)
+		{
+			FLaneEnumeration EnumLane = {Lane, OverlapRange.Min, OverlapRange.Max};
+			SearchLanes.Add(EnumLane);
+		}
+	}
+	
+	SearchLanes.Heapify([](const FLaneEnumeration& A, const FLaneEnumeration& B)
+	{
+		return A < B;
+	});
+
+	while (SearchLanes.Num() > 0)
+	{
+		FLaneEnumeration& TopLane = SearchLanes.HeapTop();
+
+		if (!Callback(TopLane.Lane.Regions[TopLane.IntervalStart]))
 		{
 			return false;
 		}
+
+		if (TopLane.IntervalStart >= TopLane.IntervalEnd)
+		{
+			SearchLanes.HeapPopDiscard();
+		}
+		else
+		{
+			TopLane.IntervalStart++;
+			SearchLanes.Heapify([](const FLaneEnumeration& A, const FLaneEnumeration& B)
+			{
+				return A < B;
+			});
+		}
 	}
 
 	return true;
@@ -87,12 +121,46 @@ bool FRegionTimeline::EnumerateRegionsBackwards(double IntervalEnd, double Inter
 		return false;
 	}
 
-	for (const FRegionLane& Lane : Lanes)
+	TArray<FLaneEnumeration> SearchLanes;
+	for (const FRegionLane& Lane: Lanes)
 	{
-		if (!Lane.EnumerateRegionsBackwards(IntervalEnd, IntervalStart, Callback))
+		const FInt32Interval OverlapRange = GetElementRangeOverlappingGivenRange<FTimeRegion>(Lane.Regions, IntervalStart, IntervalEnd,
+[](const FTimeRegion& r) { return r.BeginTime; },
+[](const FTimeRegion& r) { return r.EndTime; });
+		
+		if (OverlapRange.Min != -1 && OverlapRange.Size() >= 0)
+		{
+			FLaneEnumeration EnumLane = {Lane, OverlapRange.Min, OverlapRange.Max};
+			SearchLanes.Add(EnumLane);
+		}
+	}
+	
+	SearchLanes.Heapify([](const FLaneEnumeration& A, const FLaneEnumeration& B)
+	{
+		return A > B;
+	});
+
+	while (SearchLanes.Num() > 0)
+	{
+		FLaneEnumeration& TopLane = SearchLanes.HeapTop();
+
+		if (!Callback(TopLane.Lane.Regions[TopLane.IntervalEnd]))
 		{
 			return false;
 		}
+
+		if (TopLane.IntervalEnd <= TopLane.IntervalStart)
+		{
+			SearchLanes.HeapPopDiscard();
+		}
+		else
+		{
+			TopLane.IntervalEnd--;
+			SearchLanes.Heapify([](const FLaneEnumeration& A, const FLaneEnumeration& B)
+			{
+				return A > B;
+			});
+		}
 	}
 
 	return true;
diff --git a/Engine/Source/Developer/TraceServices/Public/TraceServices/Model/Regions.h b/Engine/Source/Developer/TraceServices/Public/TraceServices/Model/Regions.h
--- a/Engine/Source/Developer/TraceServices/Public/TraceServices/Model/Regions.h
+++ b/Engine/Source/Developer/TraceServices/Public/TraceServices/Model/Regions.h
@@ -28,6 +28,7 @@ struct FTimeRegion
 class FRegionLane
 {
 	friend class FRegionTimeline;
+	friend struct FLaneEnumeration;
 
 public:
 	FRegionLane(ILinearAllocator& InAllocator) : Regions(InAllocator, 512) {}
@@ -52,6 +53,23 @@ private:
 	TPagedArray<FTimeRegion> Regions;
 };
 
+struct FLaneEnumeration
+{
+	const FRegionLane& Lane;
+	int32 IntervalStart;
+	int32 IntervalEnd;
+
+	bool operator<(const FLaneEnumeration& rhs) const
+	{
+		return Lane.Regions[IntervalStart].BeginTime < rhs.Lane.Regions[rhs.IntervalStart].BeginTime;
+	}
+
+	bool operator>(const FLaneEnumeration& rhs) const
+	{
+		return Lane.Regions[IntervalStart].BeginTime > rhs.Lane.Regions[rhs.IntervalStart].BeginTime;
+	}
+};
+	
 /*
  *	Sorts a set of Timing Regions into a stack of individual lanes without overlaps for display
  */
