Iris
- Fixed issue with bad padding in struct FNetProfilerBunchInfo

diff --git a/Engine/Source/Developer/TraceServices/Public/TraceServices/Model/NetProfiler.h b/Engine/Source/Developer/TraceServices/Public/TraceServices/Model/NetProfiler.h
--- a/Engine/Source/Developer/TraceServices/Public/TraceServices/Model/NetProfiler.h
+++ b/Engine/Source/Developer/TraceServices/Public/TraceServices/Model/NetProfiler.h
@@ -120,13 +120,15 @@ union FNetProfilerBunchInfo
 
 		// Reserved
 		uint64 bIsValid : 1;
-		uint64 Padding : 18;
+		uint64 Padding : 17;
 	};
 	uint64 Value;
 
 	static FNetProfilerBunchInfo MakeBunchInfo(uint64 InValue) { FNetProfilerBunchInfo Info; Info.Value = InValue; return Info; }
 };
 
+static_assert(sizeof(FNetProfilerBunchInfo) == sizeof(uint64), "Unexpected size of FNetProfilerBunchInfo.");
+
 struct FNetProfilerContentEvent
 {
 	uint64 StartPos : 24;					// Inclusive start position in the packet
diff --git a/Engine/Source/Runtime/Net/Core/Public/Net/Core/Trace/Private/NetTraceInternal.h b/Engine/Source/Runtime/Net/Core/Public/Net/Core/Trace/Private/NetTraceInternal.h
--- a/Engine/Source/Runtime/Net/Core/Public/Net/Core/Trace/Private/NetTraceInternal.h
+++ b/Engine/Source/Runtime/Net/Core/Public/Net/Core/Trace/Private/NetTraceInternal.h
@@ -107,6 +107,8 @@ union FNetTraceBunchInfo
 	uint64 Value;
 };
 
+static_assert(sizeof(FNetTraceBunchInfo) == sizeof(uint64), "Unexpected size of FNetTraceBunchInfo.");
+
 struct FNetTracePacketContentEvent
 {
 	union
