[Insights]
* SamplingProfiler: Added support for aggregated stats (Timers panel).
* SamplingProfiler: Added support for Callers & Callees functionality.

diff --git a/Engine/Source/Developer/TraceInsights/Private/Insights/TimingProfiler/ViewModels/FrameStatsHelper.cpp b/Engine/Source/Developer/TraceInsights/Private/Insights/TimingProfiler/ViewModels/FrameStatsHelper.cpp
--- a/Engine/Source/Developer/TraceInsights/Private/Insights/TimingProfiler/ViewModels/FrameStatsHelper.cpp
+++ b/Engine/Source/Developer/TraceInsights/Private/Insights/TimingProfiler/ViewModels/FrameStatsHelper.cpp
@@ -18,9 +18,16 @@ void FFrameStatsHelper::ComputeFrameStatsForTimer(TArray<FFrameStatsCachedEvent>
 	TSharedPtr<const TraceServices::IAnalysisSession> Session = FInsightsManager::Get()->GetSession();
 	if (Session.IsValid())
 	{
-		for (uint32 TimelineIndex : Timelines)
+		const TraceServices::ITimingProfilerProvider* TimingProfilerProvider = TraceServices::ReadTimingProfilerProvider(*Session.Get());
+		if (TimingProfilerProvider)
 		{
-			ProcessTimeline(FrameStatsEvents, TimerId, TimelineIndex);
+			Session->ReadAccessCheck();
+			const double SessionDuration = Session->GetDurationSeconds();
+
+			for (uint32 TimelineIndex : Timelines)
+			{
+				ProcessTimeline(*TimingProfilerProvider, SessionDuration, FrameStatsEvents, TimerId, TimelineIndex);
+			}
 		}
 	}
 }
@@ -32,57 +39,55 @@ void FFrameStatsHelper::ComputeFrameStatsForTimer(TArray<FFrameStatsCachedEvent>
 	TSharedPtr<const TraceServices::IAnalysisSession> Session = FInsightsManager::Get()->GetSession();
 	if (Session.IsValid())
 	{
-		Session->ReadAccessCheck();
-
-		const TraceServices::ITimingProfilerProvider& TimingProfilerProvider = *TraceServices::ReadTimingProfilerProvider(*Session.Get());
-
-		for (uint32 TimelineIndex = 0; TimelineIndex < TimingProfilerProvider.GetTimelineCount(); ++TimelineIndex)
+		const TraceServices::ITimingProfilerProvider* TimingProfilerProvider = TraceServices::ReadTimingProfilerProvider(*Session.Get());
+		if (TimingProfilerProvider)
 		{
-			ProcessTimeline(FrameStatsEvents, TimerId, TimelineIndex);
+			Session->ReadAccessCheck();
+			const double SessionDuration = Session->GetDurationSeconds();
+
+			const uint32 TimelineCount = TimingProfilerProvider->GetTimelineCount();
+			for (uint32 TimelineIndex = 0; TimelineIndex < TimelineCount; ++TimelineIndex)
+			{
+				ProcessTimeline(*TimingProfilerProvider, SessionDuration, FrameStatsEvents, TimerId, TimelineIndex);
+			}
 		}
 	}
 }
 
 ////////////////////////////////////////////////////////////////////////////////////////////////////
 
-void FFrameStatsHelper::ProcessTimeline(TArray<FFrameStatsCachedEvent>& FrameStatsEvents, uint32 TimerId, uint32 TimelineIndex)
+void FFrameStatsHelper::ProcessTimeline(const TraceServices::ITimingProfilerProvider& TimingProfilerProvider, double SessionDuration, TArray<FFrameStatsCachedEvent>& FrameStatsEvents, uint32 TimerId, uint32 TimelineIndex)
 {
-	TSharedPtr<const TraceServices::IAnalysisSession> Session = FInsightsManager::Get()->GetSession();
-	if (Session.IsValid())
-	{
-		Session->ReadAccessCheck();
-
-		const double SessionDuration = Session->GetDurationSeconds();
+	const TraceServices::ITimingProfilerTimerReader* TimerReader;
+	TimingProfilerProvider.ReadTimers([&TimerReader](const TraceServices::ITimingProfilerTimerReader& Out) { TimerReader = &Out; });
 
-		const TraceServices::ITimingProfilerProvider& TimingProfilerProvider = *TraceServices::ReadTimingProfilerProvider(*Session.Get());
-
-		const TraceServices::ITimingProfilerTimerReader* TimerReader;
-		TimingProfilerProvider.ReadTimers([&TimerReader](const TraceServices::ITimingProfilerTimerReader& Out) { TimerReader = &Out; });
-
-		TimingProfilerProvider.ReadTimeline(TimelineIndex,
-			[SessionDuration, &FrameStatsEvents, TimerReader, TimerId](const TraceServices::ITimingProfilerProvider::Timeline& Timeline)
+	TimingProfilerProvider.ReadTimeline(TimelineIndex,
+		[SessionDuration, &FrameStatsEvents, TimerReader, TimerId]
+		(const TraceServices::ITimingProfilerTimeline& Timeline)
+		{
+			struct FEnumerateAsyncTaskData
 			{
-				struct TaskData
-				{
-					double StartTime;
-					int32 NestedDepth = 0;
-				};
-
-				TArray<TaskData> DataArray;
-
-				TraceServices::ITimeline<TraceServices::FTimingProfilerEvent>::EnumerateAsyncParams Params;
-				Params.IntervalStart = 0;
-				Params.IntervalEnd = SessionDuration;
-				Params.Resolution = 0.0;
-				Params.SetupCallback = [&DataArray](uint32 NumTasks) { DataArray.AddDefaulted(NumTasks); };
-				Params.EventCallback = [TimerReader, &FrameStatsEvents, TimerId, &DataArray](bool bIsEnter, double Time, const TraceServices::FTimingProfilerEvent& Event, uint32 TaskIndex)
+				double StartTime;
+				int32 NestedDepth = 0;
+			};
+
+			TArray<FEnumerateAsyncTaskData> DataArray;
+
+			TraceServices::ITimingProfilerTimeline::EnumerateAsyncParams Params;
+			Params.IntervalStart = 0;
+			Params.IntervalEnd = SessionDuration;
+			Params.Resolution = 0.0;
+			Params.SetupCallback = [&DataArray](uint32 NumTasks) { DataArray.AddDefaulted(NumTasks); };
+			Params.EventCallback =
+				[TimerReader, &FrameStatsEvents, TimerId, &DataArray]
+				(bool bIsEnter, double Time, const TraceServices::FTimingProfilerEvent& Event, uint32 TaskIndex)
 				{
 					const TraceServices::FTimingProfilerTimer* Timer = TimerReader->GetTimer(Event.TimerIndex);
 					if (ensure(Timer != nullptr))
 					{
 						if (Timer->Id == TimerId)
 						{
-							TaskData& CurrentTaskData = DataArray[TaskIndex];
+							FEnumerateAsyncTaskData& CurrentTaskData = DataArray[TaskIndex];
 							if (bIsEnter)
 							{
 								if (CurrentTaskData.NestedDepth == 0)
@@ -106,7 +111,7 @@ void FFrameStatsHelper::ProcessTimeline(TArray<FFrameStatsCachedEvent>& FrameSta
 									--Index;
 								}
 
-								// This can can happen when the event is between frames.
+								// This can happen when the event is between frames.
 								if (CurrentTaskData.StartTime > FrameStatsEvents[Index].FrameEndTime)
 								{
 									Index++;
@@ -136,16 +141,16 @@ void FFrameStatsHelper::ProcessTimeline(TArray<FFrameStatsCachedEvent>& FrameSta
 									for (double Value = Entry.Duration.load(); !Entry.Duration.compare_exchange_strong(Value, Value + Duration););
 
 									Index++;
-								} while (Index < FrameStatsEvents.Num());
+								}
+								while (Index < FrameStatsEvents.Num());
 							}
 						}
 					}
 					return TraceServices::EEventEnumerate::Continue;
 				};
 
-				Timeline.EnumerateEventsDownSampledAsync(Params);
-			});
-	}
+			Timeline.EnumerateEventsDownSampledAsync(Params);
+		});
 }
 
 ////////////////////////////////////////////////////////////////////////////////////////////////////
diff --git a/Engine/Source/Developer/TraceInsights/Private/Insights/TimingProfiler/ViewModels/FrameStatsHelper.h b/Engine/Source/Developer/TraceInsights/Private/Insights/TimingProfiler/ViewModels/FrameStatsHelper.h
--- a/Engine/Source/Developer/TraceInsights/Private/Insights/TimingProfiler/ViewModels/FrameStatsHelper.h
+++ b/Engine/Source/Developer/TraceInsights/Private/Insights/TimingProfiler/ViewModels/FrameStatsHelper.h
@@ -8,6 +8,11 @@
 
 #include <atomic>
 
+namespace TraceServices
+{
+	class ITimingProfilerProvider;
+}
+
 namespace UE::Insights::TimingProfiler
 {
 
@@ -41,7 +46,11 @@ public:
 	static void ComputeFrameStatsForTimer(TArray<FFrameStatsCachedEvent>& FrameStatsEvents, uint32 TimerId);
 
 private:
-	static void ProcessTimeline(TArray<FFrameStatsCachedEvent>& FrameStatsEvents, uint32 TimerId, uint32 TimelineIndex);
+	static void ProcessTimeline(const TraceServices::ITimingProfilerProvider& TimingProfilerProvider,
+								double SessionDuration,
+								TArray<FFrameStatsCachedEvent>& FrameStatsEvents,
+								uint32 TimerId,
+								uint32 TimelineIndex);
 };
 
 } // namespace UE::Insights::TimingProfiler
diff --git a/Engine/Source/Developer/TraceInsights/Private/Insights/TimingProfiler/ViewModels/ThreadTimingSharedState.cpp b/Engine/Source/Developer/TraceInsights/Private/Insights/TimingProfiler/ViewModels/ThreadTimingSharedState.cpp
--- a/Engine/Source/Developer/TraceInsights/Private/Insights/TimingProfiler/ViewModels/ThreadTimingSharedState.cpp
+++ b/Engine/Source/Developer/TraceInsights/Private/Insights/TimingProfiler/ViewModels/ThreadTimingSharedState.cpp
@@ -258,7 +258,7 @@ TSharedPtr<FCpuTimingTrack> FThreadTimingSharedState::GetCpuTrack(uint32 InThrea
 
 bool FThreadTimingSharedState::IsCpuTrackVisible(uint32 InThreadId) const
 {
-	const TSharedPtr<FCpuTimingTrack>*const TrackPtrPtr = CpuTracks.Find(InThreadId);
+	const TSharedPtr<FCpuTimingTrack>* const TrackPtrPtr = CpuTracks.Find(InThreadId);
 	return TrackPtrPtr && (*TrackPtrPtr)->IsVisible();
 }
 
@@ -279,17 +279,40 @@ void FThreadTimingSharedState::GetVisibleCpuThreads(TSet<uint32>& OutSet) const
 
 ////////////////////////////////////////////////////////////////////////////////////////////////////
 
-void FThreadTimingSharedState::GetVisibleTimelineIndexes(TSet<uint32>& OutSet) const
+TSharedPtr<FCpuStackSampleTimingTrack> FThreadTimingSharedState::GetCpuSamplingTrack(uint32 InSystemThreadId)
+{
+	TSharedPtr<FCpuStackSampleTimingTrack>* const TrackPtrPtr = CpuStackSampleTracks.Find(InSystemThreadId);
+	return TrackPtrPtr ? *TrackPtrPtr : nullptr;
+}
+
+////////////////////////////////////////////////////////////////////////////////////////////////////
+
+bool FThreadTimingSharedState::IsCpuSamplingTrackVisible(uint32 InSystemThreadId) const
+{
+	const TSharedPtr<FCpuStackSampleTimingTrack>* const TrackPtrPtr = CpuStackSampleTracks.Find(InSystemThreadId);
+	return TrackPtrPtr && (*TrackPtrPtr)->IsVisible();
+}
+
+////////////////////////////////////////////////////////////////////////////////////////////////////
+
+void FThreadTimingSharedState::GetVisibleCpuSamplingThreads(TSet<uint32>& OutSet) const
 {
 	OutSet.Reset();
-	for (const auto& KV : CpuTracks)
+	for (const auto& KV : CpuStackSampleTracks)
 	{
-		const FCpuTimingTrack& Track = *KV.Value;
+		const FCpuStackSampleTimingTrack& Track = *KV.Value;
 		if (Track.IsVisible())
 		{
-			OutSet.Add(Track.GetTimelineIndex());
+			OutSet.Add(KV.Key);
 		}
 	}
+}
+
+////////////////////////////////////////////////////////////////////////////////////////////////////
+
+void FThreadTimingSharedState::GetVisibleTimelineIndexes(TSet<uint32>& OutSet) const
+{
+	OutSet.Reset();
 
 	if (OldGpu1Track.IsValid() && OldGpu1Track->IsVisible())
 	{
@@ -309,6 +332,20 @@ void FThreadTimingSharedState::GetVisibleTimelineIndexes(TSet<uint32>& OutSet) c
 			OutSet.Add(Track.GetTimelineIndex());
 		}
 	}
+
+	if (VerseSamplingTrack.IsValid() && VerseSamplingTrack->IsVisible())
+	{
+		OutSet.Add(VerseSamplingTrack->GetTimelineIndex());
+	}
+
+	for (const auto& KV : CpuTracks)
+	{
+		const FCpuTimingTrack& Track = *KV.Value;
+		if (Track.IsVisible())
+		{
+			OutSet.Add(Track.GetTimelineIndex());
+		}
+	}
 }
 
 ////////////////////////////////////////////////////////////////////////////////////////////////////
@@ -343,7 +380,7 @@ void FThreadTimingSharedState::ResetToDefaultValues()
 	TimingViewSettings->SetToggleOption(FTimingViewSettings::Option_ShowGpuFencesExtendedLines, true);
 	SetGpuFencesTracksVisibility(true);
 }
-	
+
 ////////////////////////////////////////////////////////////////////////////////////////////////////
 
 bool FThreadTimingSharedState::AreCpuSamplingTracksAvailable() const
@@ -727,7 +764,7 @@ void FThreadTimingSharedState::ExtendCpuTracksFilterMenu(Timing::ITimingViewSess
 		InOutMenuBuilder.AddMenuEntry(FThreadTimingViewCommands::Get().ShowHideAllCpuSamplingTracks);
 	}
 	InOutMenuBuilder.EndSection();
-	
+
 	InOutMenuBuilder.BeginSection("VerseTracks", LOCTEXT("ContextMenu_Section_VerseTracks", "Verse Tracks"));
 	{
 		InOutMenuBuilder.AddMenuEntry(FThreadTimingViewCommands::Get().ShowHideAllVerseTracks);
@@ -899,7 +936,7 @@ void FThreadTimingSharedState::SetAllCpuTracksToggle(bool bOnOff)
 		TimingView->HandleTrackVisibilityChanged();
 	}
 }
-	
+
 ////////////////////////////////////////////////////////////////////////////////////////////////////
 
 void FThreadTimingSharedState::SetAllCpuSamplingTracksToggle(bool bOnOff)
diff --git a/Engine/Source/Developer/TraceInsights/Private/Insights/TimingProfiler/ViewModels/ThreadTimingSharedState.h b/Engine/Source/Developer/TraceInsights/Private/Insights/TimingProfiler/ViewModels/ThreadTimingSharedState.h
--- a/Engine/Source/Developer/TraceInsights/Private/Insights/TimingProfiler/ViewModels/ThreadTimingSharedState.h
+++ b/Engine/Source/Developer/TraceInsights/Private/Insights/TimingProfiler/ViewModels/ThreadTimingSharedState.h
@@ -156,47 +156,44 @@
 	explicit FThreadTimingSharedState(STimingView* InTimingView);
 	virtual ~FThreadTimingSharedState() override = default;
 
+	//////////////////////////////////////////////////
+	// GPU
+
 	TSharedPtr<FGpuTimingTrack> GetOldGpu1Track() { return OldGpu1Track; }
 	TSharedPtr<FGpuTimingTrack> GetOldGpu2Track() { return OldGpu2Track; }
-	TSharedPtr<FGpuQueueTimingTrack> GetGpuTrack(uint32 InQueueId);
-	TSharedPtr<FVerseTimingTrack> GetVerseSamplingTrack() { return VerseSamplingTrack; }
-	TSharedPtr<FCpuTimingTrack> GetCpuTrack(uint32 InThreadId);
-	const TMap<uint32, TSharedPtr<FCpuTimingTrack>> GetAllCpuTracks() { return CpuTracks; }
-
 	bool IsOldGpu1TrackVisible() const;
 	bool IsOldGpu2TrackVisible() const;
+
+	TSharedPtr<FGpuQueueTimingTrack> GetGpuTrack(uint32 InQueueId);
 	bool IsAnyGpuTrackVisible() const;
 	bool IsGpuTrackVisible(uint32 InQueueId) const;
-	bool IsVerseSamplingTrackVisible() const;
-	bool IsCpuTrackVisible(uint32 InThreadId) const;
-
 	void GetVisibleGpuQueues(TSet<uint32>& OutSet) const;
-	void GetVisibleCpuThreads(TSet<uint32>& OutSet) const;
-	void GetVisibleTimelineIndexes(TSet<uint32>& OutSet) const;
 
-	//////////////////////////////////////////////////
-	// ITimingViewExtender interface
-
-	virtual void OnBeginSession(Timing::ITimingViewSession& InSession) override;
-	virtual void OnEndSession(Timing::ITimingViewSession& InSession) override;
-	virtual void Tick(Timing::ITimingViewSession& InSession, const TraceServices::IAnalysisSession& InAnalysisSession) override;
-	virtual void ExtendGpuTracksFilterMenu(Timing::ITimingViewSession& InSession, FMenuBuilder& InMenuBuilder) override;
-	virtual void ExtendCpuTracksFilterMenu(Timing::ITimingViewSession& InSession, FMenuBuilder& InMenuBuilder) override;
-
-	//////////////////////////////////////////////////
-
-	void BindCommands();
-
-	void ResetToDefaultValues();
-
-	bool AreCpuSamplingTracksAvailable() const;
-	bool AreVerseTracksAvailable() const;
-
 	bool IsAllGpuTracksToggleOn() const { return TimingViewSettings->IsToggleOptionEnabled(FTimingViewSettings::Option_ShowAllGpuTracks); }
 	void SetAllGpuTracksToggle(bool bOnOff);
 	void ShowAllGpuTracks() { SetAllGpuTracksToggle(true); }
 	void HideAllGpuTracks() { SetAllGpuTracksToggle(false); }
 	void ShowHideAllGpuTracks() { SetAllGpuTracksToggle(!IsAllGpuTracksToggleOn()); }
 
+	bool AreOverlaysVisibleInGpuQueueTracks() const { return TimingViewSettings->IsToggleOptionEnabled(FTimingViewSettings::Option_ShowGpuWorkOverlays); }
+	bool AreExtendedLinesVisibleInGpuQueueTracks() const { return TimingViewSettings->IsToggleOptionEnabled(FTimingViewSettings::Option_ShowGpuWorkExtendedLines); }
+
+	bool AreGpuWorkTracksVisible() const { return TimingViewSettings->IsToggleOptionEnabled(FTimingViewSettings::Option_ShowGpuWorkTracks); }
+	void SetGpuWorkTracksVisibility(bool bOnOff);
+
+	bool AreGpuFencesTracksVisible() const { return TimingViewSettings->IsToggleOptionEnabled(FTimingViewSettings::Option_ShowGpuFencesTracks); }
+	void SetGpuFencesTracksVisibility(bool bOnOff);
+
+	bool AreGpuFencesExtendedLinesVisible() const { return TimingViewSettings->IsToggleOptionEnabled(FTimingViewSettings::Option_ShowGpuFencesExtendedLines); }
+	bool AreGpuFenceRelationsVisible() const { return TimingViewSettings->IsToggleOptionEnabled(FTimingViewSettings::Option_ShowGpuFencesRelations); }
+
+	//////////////////////////////////////////////////
+	// Verse Sampling
+
+	bool AreVerseTracksAvailable() const;
+
+	TSharedPtr<FVerseTimingTrack> GetVerseSamplingTrack() { return VerseSamplingTrack; }
+	bool IsVerseSamplingTrackVisible() const;
+
 	bool IsAllVerseTracksToggleOn() const { return TimingViewSettings->IsToggleOptionEnabled(FTimingViewSettings::Option_ShowAllVerseTracks); }
 	void SetAllVerseTracksToggle(bool bOnOff);
@@ -204,29 +201,52 @@
 	void HideAllVerseTracks() { SetAllVerseTracksToggle(false); }
 	void ShowHideAllVerseTracks() { SetAllVerseTracksToggle(!IsAllVerseTracksToggleOn()); }
 
+	//////////////////////////////////////////////////
+	// CPU
+
+	TSharedPtr<FCpuTimingTrack> GetCpuTrack(uint32 InThreadId);
+	const TMap<uint32, TSharedPtr<FCpuTimingTrack>>& GetAllCpuTracks() { return CpuTracks; }
+	bool IsCpuTrackVisible(uint32 InThreadId) const;
+	void GetVisibleCpuThreads(TSet<uint32>& OutSet) const;
+
 	bool IsAllCpuTracksToggleOn() const { return TimingViewSettings->IsToggleOptionEnabled(FTimingViewSettings::Option_ShowAllCpuTracks); }
 	void SetAllCpuTracksToggle(bool bOnOff);
 	void ShowAllCpuTracks() { SetAllCpuTracksToggle(true); }
 	void HideAllCpuTracks() { SetAllCpuTracksToggle(false); }
 	void ShowHideAllCpuTracks() { SetAllCpuTracksToggle(!IsAllCpuTracksToggleOn()); }
 
+	//////////////////////////////////////////////////
+	// CPU Sampling
+
+	bool AreCpuSamplingTracksAvailable() const;
+
+	TSharedPtr<FCpuStackSampleTimingTrack> GetCpuSamplingTrack(uint32 InSystemThreadId);
+	const TMap<uint32, TSharedPtr<FCpuStackSampleTimingTrack >>& GetAllCpuSamplingTracks() { return CpuStackSampleTracks; }
+	bool IsCpuSamplingTrackVisible(uint32 InSystemThreadId) const;
+	void GetVisibleCpuSamplingThreads(TSet<uint32>& OutSet) const;
+
 	bool IsAllCpuSamplingTracksToggleOn() const { return TimingViewSettings->IsToggleOptionEnabled(FTimingViewSettings::Option_ShowAllCpuSamplingTracks); }
 	void SetAllCpuSamplingTracksToggle(bool bOnOff);
 	void ShowAllCpuSamplingTracks() { SetAllCpuSamplingTracksToggle(true); }
 	void HideAllCpuSamplingTracks() { SetAllCpuSamplingTracksToggle(false); }
 	void ShowHideAllCpuSamplingTracks() { SetAllCpuSamplingTracksToggle(!IsAllCpuSamplingTracksToggleOn()); }
 
-	bool AreOverlaysVisibleInGpuQueueTracks() const { return TimingViewSettings->IsToggleOptionEnabled(FTimingViewSettings::Option_ShowGpuWorkOverlays); }
-	bool AreExtendedLinesVisibleInGpuQueueTracks() const { return TimingViewSettings->IsToggleOptionEnabled(FTimingViewSettings::Option_ShowGpuWorkExtendedLines); }
+	//////////////////////////////////////////////////
+	// ITimingViewExtender interface
 
-	bool AreGpuWorkTracksVisible() const { return TimingViewSettings->IsToggleOptionEnabled(FTimingViewSettings::Option_ShowGpuWorkTracks); }
-	void SetGpuWorkTracksVisibility(bool bOnOff);
+	virtual void OnBeginSession(Timing::ITimingViewSession& InSession) override;
+	virtual void OnEndSession(Timing::ITimingViewSession& InSession) override;
+	virtual void Tick(Timing::ITimingViewSession& InSession, const TraceServices::IAnalysisSession& InAnalysisSession) override;
+	virtual void ExtendGpuTracksFilterMenu(Timing::ITimingViewSession& InSession, FMenuBuilder& InMenuBuilder) override;
+	virtual void ExtendCpuTracksFilterMenu(Timing::ITimingViewSession& InSession, FMenuBuilder& InMenuBuilder) override;
 
-	bool AreGpuFencesTracksVisible() const { return TimingViewSettings->IsToggleOptionEnabled(FTimingViewSettings::Option_ShowGpuFencesTracks); }
-	void SetGpuFencesTracksVisibility(bool bOnOff);
+	//////////////////////////////////////////////////
 
-	bool AreGpuFencesExtendedLinesVisible() const { return TimingViewSettings->IsToggleOptionEnabled(FTimingViewSettings::Option_ShowGpuFencesExtendedLines); }
-	bool AreGpuFenceRelationsVisible() const { return TimingViewSettings->IsToggleOptionEnabled(FTimingViewSettings::Option_ShowGpuFencesRelations); }
+	void BindCommands();
 
+	void ResetToDefaultValues();
+
+	void GetVisibleTimelineIndexes(TSet<uint32>& OutSet) const;
+
 	TSharedPtr<const ITimingEvent> FindMaxEventInstance(uint32 TimerId, double StartTime, double EndTime);
 	TSharedPtr<const ITimingEvent> FindMinEventInstance(uint32 TimerId, double StartTime, double EndTime);
@@ -238,6 +258,9 @@
 	void ToggleTrackVisibilityByGroup_Execute(const TCHAR* InGroupName);
 	void ToggleOptionState(const TCHAR* CommandName);
 
+	//////////////////////////////////////////////////
+	// GPU
+
 	void Command_ShowGpuWorkTracks_Execute() { SetGpuWorkTracksVisibility(!AreGpuWorkTracksVisible()); }
 	bool Command_ShowGpuWorkTracks_CanExecute() { return GpuTracks.Num() > 0; }
 	bool Command_ShowGpuWorkTracks_IsChecked() { return AreGpuWorkTracksVisible(); }
@@ -268,6 +291,8 @@
 	void AddGpuFencesChildTracks();
 	void RemoveGpuFencesChildTracks();
 
+	//////////////////////////////////////////////////
+
 	void OnTimingEventSelected(TSharedPtr<const ITimingEvent> InSelectedEvent);
 
 private:
diff --git a/Engine/Source/Developer/TraceInsights/Private/Insights/TimingProfiler/ViewModels/TimerAggregator.cpp b/Engine/Source/Developer/TraceInsights/Private/Insights/TimingProfiler/ViewModels/TimerAggregator.cpp
--- a/Engine/Source/Developer/TraceInsights/Private/Insights/TimingProfiler/ViewModels/TimerAggregator.cpp
+++ b/Engine/Source/Developer/TraceInsights/Private/Insights/TimingProfiler/ViewModels/TimerAggregator.cpp
@@ -2,6 +2,11 @@
 
 #include "TimerAggregator.h"
 
+// TraceServices
+#include "Common/ProviderLock.h"
+#include "TraceServices/Model/AnalysisSession.h"
+#include "TraceServices/Model/StackSamples.h"
+
 // TraceInsights
 #include "Insights/TimingProfiler/TimingProfilerManager.h"
 #include "Insights/TimingProfiler/ViewModels/ThreadTimingSharedState.h"
@@ -22,7 +27,8 @@ struct FTimerAggregationParams
 
 	bool bIncludeVerseSampling = false;
 
-	TSet<uint32> CpuThreads;
+	TSet<uint32> CpuThreads;  // [Thread Id]
+	TSet<uint32> CpuSamplingThreads; // [System Thread Id]
 
 	ETraceFrameType FrameType = ETraceFrameType::TraceFrameType_Count;
 };
@@ -43,6 +49,7 @@ public:
 		, bIncludeOldGpu2(InParams.bIncludeOldGpu2)
 		, bIncludeVerseSampling(InParams.bIncludeVerseSampling)
 		, CpuThreads(InParams.CpuThreads)
+		, CpuSamplingThreads(InParams.CpuSamplingThreads)
 		, FrameType(InParams.FrameType)
 		, ResultTable()
 	{
@@ -63,7 +70,8 @@ private:
 	bool bIncludeOldGpu1;
 	bool bIncludeOldGpu2;
 	bool bIncludeVerseSampling;
-	TSet<uint32> CpuThreads;
+	TSet<uint32> CpuThreads; // [Thread Id]
+	TSet<uint32> CpuSamplingThreads; // [System Thread Id]
 	ETraceFrameType FrameType;
 	TUniquePtr<TraceServices::ITable<TraceServices::FTimingProfilerAggregatedStats>> ResultTable;
 };
@@ -72,23 +80,55 @@ private:
 
 void FTimerAggregationWorker::DoWork(TSharedPtr<TraceServices::FCancellationToken> CancellationToken)
 {
-	if (Session.IsValid() && TraceServices::ReadTimingProfilerProvider(*Session.Get()))
+	if (!Session.IsValid())
+	{
+		return;
+	}
+
+	const TraceServices::ITimingProfilerProvider* TimingProfilerProvider = TraceServices::ReadTimingProfilerProvider(*Session.Get());
+	if (!TimingProfilerProvider)
+	{
+		return;
+	}
+
+	TraceServices::FCreateAggregationParams Params;
+	Params.IntervalStart = StartTime;
+	Params.IntervalEnd = EndTime;
+	Params.GpuQueueFilter = [this](uint32 QueueId) { return GpuQueues.Contains(QueueId); };
+	Params.bIncludeOldGpu1 = bIncludeOldGpu1;
+	Params.bIncludeOldGpu2 = bIncludeOldGpu2;
+	Params.bIncludeVerseSampling = bIncludeVerseSampling;
+	Params.CpuThreadFilter = [this](uint32 ThreadId) { return CpuThreads.Contains(ThreadId); };
+	Params.FrameType = FrameType;
+	Params.CancellationToken = CancellationToken;
+
+	bool bHasCpuSamplingTimelines = false;
+	if (CpuSamplingThreads.Num() > 0)
+	{
+		if (const TraceServices::IStackSamplesProvider* StackSamplesProvider = TraceServices::ReadStackSamplesProvider(*Session.Get()))
+		{
+			TraceServices::FProviderReadScopeLock _(*StackSamplesProvider);
+			for (uint32 SystemThreadId : CpuSamplingThreads)
+			{
+				const TraceServices::FStackSampleTimeline* Timeline = StackSamplesProvider->GetTimeline(SystemThreadId);
+				if (Timeline)
+				{
+					Params.CustomTimelines.Add(Timeline);
+					bHasCpuSamplingTimelines = true;
+				}
+			}
+			// Aggregation needs to execute under the read locks of both Session and StackSamplesProvider.
+			if (bHasCpuSamplingTimelines)
+			{
+				TraceServices::FAnalysisSessionReadScope SessionReadScope(*Session.Get());
+				ResultTable.Reset(TimingProfilerProvider->CreateAggregation(Params));
+			}
+		}
+	}
+	if (!bHasCpuSamplingTimelines)
 	{
 		TraceServices::FAnalysisSessionReadScope SessionReadScope(*Session.Get());
-		const TraceServices::ITimingProfilerProvider& TimingProfilerProvider = *TraceServices::ReadTimingProfilerProvider(*Session.Get());
-
-		TraceServices::FCreateAggregationParams Params;
-		Params.IntervalStart = StartTime;
-		Params.IntervalEnd = EndTime;
-		Params.GpuQueueFilter = [this](uint32 QueueId) { return GpuQueues.Contains(QueueId); };
-		Params.bIncludeOldGpu1 = bIncludeOldGpu1;
-		Params.bIncludeOldGpu2 = bIncludeOldGpu2;
-		Params.bIncludeVerseSampling = bIncludeVerseSampling;
-		Params.CpuThreadFilter = [this](uint32 ThreadId) { return CpuThreads.Contains(ThreadId); };
-		Params.FrameType = FrameType;
-		Params.CancellationToken = CancellationToken;
-
-		ResultTable.Reset(TimingProfilerProvider.CreateAggregation(Params));
+		ResultTable.Reset(TimingProfilerProvider->CreateAggregation(Params));
 	}
 }
 
@@ -118,6 +158,7 @@ IStatsAggregationWorker* FTimerAggregator::CreateWorker(TSharedPtr<const TraceSe
 				Params.bIncludeOldGpu2 = ThreadTimingSharedState->IsOldGpu2TrackVisible();
 				Params.bIncludeVerseSampling = ThreadTimingSharedState->IsVerseSamplingTrackVisible();
 				ThreadTimingSharedState->GetVisibleCpuThreads(Params.CpuThreads);
+				ThreadTimingSharedState->GetVisibleCpuSamplingThreads(Params.CpuSamplingThreads);
 			}
 		}
 	}
diff --git a/Engine/Source/Developer/TraceInsights/Private/Insights/TimingProfiler/ViewModels/TimerButterflyAggregator.cpp b/Engine/Source/Developer/TraceInsights/Private/Insights/TimingProfiler/ViewModels/TimerButterflyAggregator.cpp
--- a/Engine/Source/Developer/TraceInsights/Private/Insights/TimingProfiler/ViewModels/TimerButterflyAggregator.cpp
+++ b/Engine/Source/Developer/TraceInsights/Private/Insights/TimingProfiler/ViewModels/TimerButterflyAggregator.cpp
@@ -3,7 +3,9 @@
 #include "TimerButterflyAggregator.h"
 
 // TraceServices
+#include "Common/ProviderLock.h"
 #include "TraceServices/Model/AnalysisSession.h"
+#include "TraceServices/Model/StackSamples.h"
 #include "TraceServices/Model/TimingProfiler.h"
 
 // TraceInsights
@@ -27,7 +29,8 @@ struct FTimerButterflyAggregationParams
 
 	bool bIncludeVerseSampling = false;
 
-	TSet<uint32> CpuThreads;
+	TSet<uint32> CpuThreads; // [Thread Id]
+	TSet<uint32> CpuSamplingThreads; // [System Thread Id]
 };
 
 ////////////////////////////////////////////////////////////////////////////////////////////////////
@@ -46,6 +49,7 @@ public:
 		, bIncludeOldGpu2(InParams.bIncludeOldGpu2)
 		, bIncludeVerseSampling(InParams.bIncludeVerseSampling)
 		, CpuThreads(InParams.CpuThreads)
+		, CpuSamplingThreads(InParams.CpuSamplingThreads)
 		, ResultButterfly()
 	{
 	}
@@ -65,7 +69,8 @@ private:
 	bool bIncludeOldGpu1;
 	bool bIncludeOldGpu2;
 	bool bIncludeVerseSampling;
-	TSet<uint32> CpuThreads;
+	TSet<uint32> CpuThreads; // [Thread Id]
+	TSet<uint32> CpuSamplingThreads; // [System Thread Id]
 	TUniquePtr<TraceServices::ITimingProfilerButterfly> ResultButterfly;
 };
 
@@ -73,35 +78,73 @@ private:
 
 void FTimerButterflyAggregationWorker::DoWork(TSharedPtr<TraceServices::FCancellationToken> CancellationToken)
 {
-	if (Session.IsValid() && TraceServices::ReadTimingProfilerProvider(*Session.Get()))
+	if (!Session.IsValid())
 	{
-		TraceServices::FAnalysisSessionReadScope SessionReadScope(*Session.Get());
-		const TraceServices::ITimingProfilerProvider& TimingProfilerProvider = *TraceServices::ReadTimingProfilerProvider(*Session.Get());
+		return;
+	}
+
+	const TraceServices::ITimingProfilerProvider* TimingProfilerProvider = TraceServices::ReadTimingProfilerProvider(*Session.Get());
+	if (!TimingProfilerProvider)
+	{
+		return;
+	}
 
-		const double AdjustedEndTime = FMath::Min(Session->GetDurationSeconds(), EndTime);
-		if (StartTime <= AdjustedEndTime)
+	double SessionDuration = 0.0;
+	{
+		TraceServices::FAnalysisSessionReadScope SessionReadScope(*Session.Get());
+		SessionDuration = Session->GetDurationSeconds();
+	}
+	const double AdjustedEndTime = FMath::Min(SessionDuration, EndTime);
+	if (StartTime <= AdjustedEndTime)
+	{
+		TraceServices::FCreateButterflyParams Params;
+		Params.IntervalStart = StartTime;
+		Params.IntervalEnd = AdjustedEndTime;
+		Params.GpuQueueFilter = [this](uint32 QueueId) { return GpuQueues.Contains(QueueId); };
+		Params.bIncludeOldGpu1 = bIncludeOldGpu1;
+		Params.bIncludeOldGpu2 = bIncludeOldGpu2;
+		Params.bIncludeVerseSampling = bIncludeVerseSampling;
+		Params.CpuThreadFilter = [this](uint32 ThreadId) { return CpuThreads.Contains(ThreadId); };
+
+		bool bHasCpuSamplingTimelines = false;
+		if (CpuSamplingThreads.Num() > 0)
 		{
-			TraceServices::FCreateButterflyParams Params;
-			Params.IntervalStart = StartTime;
-			Params.IntervalEnd = AdjustedEndTime;
-			Params.GpuQueueFilter = [this](uint32 QueueId) { return GpuQueues.Contains(QueueId); };
-			Params.bIncludeOldGpu1 = bIncludeOldGpu1;
-			Params.bIncludeOldGpu2 = bIncludeOldGpu2;
-			Params.bIncludeVerseSampling = bIncludeVerseSampling;
-			Params.CpuThreadFilter = [this](uint32 ThreadId) { return CpuThreads.Contains(ThreadId); };
-
-			ResultButterfly.Reset(TimingProfilerProvider.CreateButterfly(Params));
+			if (const TraceServices::IStackSamplesProvider* StackSamplesProvider = TraceServices::ReadStackSamplesProvider(*Session.Get()))
+			{
+				TraceServices::FProviderReadScopeLock _(*StackSamplesProvider);
+				for (uint32 SystemThreadId : CpuSamplingThreads)
+				{
+					const TraceServices::FStackSampleTimeline* Timeline = StackSamplesProvider->GetTimeline(SystemThreadId);
+					if (Timeline)
+					{
+						Params.CustomTimelines.Add(Timeline);
+						bHasCpuSamplingTimelines = true;
+					}
+				}
+				// Aggregation needs to execute under the read locks of both Session and StackSamplesProvider.
+				if (bHasCpuSamplingTimelines)
+				{
+					TraceServices::FAnalysisSessionReadScope SessionReadScope(*Session.Get());
+					ResultButterfly.Reset(TimingProfilerProvider->CreateButterfly(Params));
+				}
+			}
 		}
-		else
+		if (!bHasCpuSamplingTimelines)
 		{
-			// Generate empty results (invalid time range selection).
-			TraceServices::FCreateButterflyParams Params;
-			Params.IntervalStart = 0.0;
-			Params.IntervalEnd = -1.0;
-
-			ResultButterfly.Reset(TimingProfilerProvider.CreateButterfly(Params));
+			TraceServices::FAnalysisSessionReadScope SessionReadScope(*Session.Get());
+			ResultButterfly.Reset(TimingProfilerProvider->CreateButterfly(Params));
 		}
 	}
+	else
+	{
+		// Generate empty results (invalid time range selection).
+		TraceServices::FCreateButterflyParams Params;
+		Params.IntervalStart = 0.0;
+		Params.IntervalEnd = -1.0;
+
+		TraceServices::FAnalysisSessionReadScope SessionReadScope(*Session.Get());
+		ResultButterfly.Reset(TimingProfilerProvider->CreateButterfly(Params));
+	}
 }
 
 ////////////////////////////////////////////////////////////////////////////////////////////////////
@@ -129,6 +172,7 @@ IStatsAggregationWorker* FTimerButterflyAggregator::CreateWorker(TSharedPtr<cons
 				Params.bIncludeOldGpu2 = ThreadTimingSharedState->IsOldGpu2TrackVisible();
 				Params.bIncludeVerseSampling = ThreadTimingSharedState->IsVerseSamplingTrackVisible();
 				ThreadTimingSharedState->GetVisibleCpuThreads(Params.CpuThreads);
+				ThreadTimingSharedState->GetVisibleCpuSamplingThreads(Params.CpuSamplingThreads);
 			}
 		}
 	}
diff --git a/Engine/Source/Developer/TraceServices/Private/Model/TimingProfiler.cpp b/Engine/Source/Developer/TraceServices/Private/Model/TimingProfiler.cpp
--- a/Engine/Source/Developer/TraceServices/Private/Model/TimingProfiler.cpp
+++ b/Engine/Source/Developer/TraceServices/Private/Model/TimingProfiler.cpp
@@ -761,7 +761,7 @@ ITable<FTimingProfilerAggregatedStats>* FTimingProfilerProvider::CreateAggregati
 {
 	Session.ReadAccessCheck();
 
-	TArray<const TimelineInternal*> IncludedTimelines;
+	TArray<const ITimingProfilerTimeline*> IncludedTimelines;
 
 	//////////////////////////////////////////////////
 	// GPU
@@ -808,6 +808,21 @@ ITable<FTimingProfilerAggregatedStats>* FTimingProfilerProvider::CreateAggregati
 	}
 
 	//////////////////////////////////////////////////
+	// Custom Timelines
+
+	for (const ITimingProfilerTimeline* Timeline : Params.CustomTimelines)
+	{
+		IncludedTimelines.Add(Timeline);
+	}
+
+	//////////////////////////////////////////////////
+
+	return CreateAggregationInternal(IncludedTimelines, Params);
+}
+
+ITable<FTimingProfilerAggregatedStats>* FTimingProfilerProvider::CreateAggregationInternal(TArray<const ITimingProfilerTimeline*>& IncludedTimelines, const FCreateAggregationParams& Params) const
+{
+	Session.ReadAccessCheck();
 
 	auto BucketMappingFunc = [this](const TimelineInternal::EventType& Event) -> const FTimingProfilerTimer*
 	{
@@ -828,7 +843,9 @@ ITable<FTimingProfilerAggregatedStats>* FTimingProfilerProvider::CreateAggregati
 	{
 		TArray<FFrameData> Frames;
 		const IFrameProvider& FrameProvider = ReadFrameProvider(Session);
-		FrameProvider.EnumerateFrames(Params.FrameType, Params.IntervalStart, Params.IntervalEnd, [&Frames](const FFrame& Frame)
+		FrameProvider.EnumerateFrames(Params.FrameType, Params.IntervalStart, Params.IntervalEnd,
+			[&Frames]
+			(const FFrame& Frame)
 			{
 				FFrameData NewFrameData;
 				NewFrameData.StartTime = Frame.StartTime;
@@ -1054,10 +1071,9 @@ struct FTimingProfilerCallstackKey
 
 ITimingProfilerButterfly* FTimingProfilerProvider::CreateButterfly(const FCreateButterflyParams& Params) const
 {
-	FTimingProfilerButterfly* Butterfly = new FTimingProfilerButterfly();
-	Butterfly->TimerCallstacksMap.AddDefaulted(Timers.Num());
+	Session.ReadAccessCheck();
 
-	TArray<const TimelineInternal*> IncludedTimelines;
+	TArray<const ITimingProfilerTimeline*> IncludedTimelines;
 
 	//////////////////////////////////////////////////
 	// GPU
@@ -1104,6 +1120,26 @@ ITimingProfilerButterfly* FTimingProfilerProvider::CreateButterfly(const FCreate
 	}
 
 	//////////////////////////////////////////////////
+	// Custom Timelines
+
+	for (const ITimingProfilerTimeline* Timeline : Params.CustomTimelines)
+	{
+		IncludedTimelines.Add(Timeline);
+	}
+
+	//////////////////////////////////////////////////
+
+	return CreateButterflyInternal(IncludedTimelines, Params);
+}
+
+////////////////////////////////////////////////////////////////////////////////////////////////////
+
+ITimingProfilerButterfly* FTimingProfilerProvider::CreateButterflyInternal(TArray<const ITimingProfilerTimeline*>& IncludedTimelines, const FCreateButterflyParams& Params) const
+{
+	Session.ReadAccessCheck();
+
+	FTimingProfilerButterfly* Butterfly = new FTimingProfilerButterfly();
+	Butterfly->TimerCallstacksMap.AddDefaulted(Timers.Num());
 
 	FTimingProfilerCallstackKey CurrentCallstackKey;
 	CurrentCallstackKey.TimerStack.Reserve(1024);
@@ -1123,88 +1159,90 @@ ITimingProfilerButterfly* FTimingProfilerProvider::CreateButterfly(const FCreate
 	TMap<FTimingProfilerCallstackKey, TTuple<FTimingProfilerButterflyNode*, bool>> CallstackNodeMap;
 
 	double LastTime = 0.0;
-	for (const TimelineInternal* Timeline : IncludedTimelines)
+	for (const ITimingProfilerTimeline* Timeline : IncludedTimelines)
 	{
-		Timeline->EnumerateEvents(Params.IntervalStart, Params.IntervalEnd, [this, &Params, &CurrentCallstackKey, &CurrentCallstack, &CallstackNodeMap, &LastTime, Butterfly](bool IsEnter, double Time, const FTimingProfilerEvent& Event)
-		{
-			Time = FMath::Clamp(Time, Params.IntervalStart, Params.IntervalEnd);
-			FTimingProfilerButterflyNode* ParentNode = nullptr;
-			uint32 ParentCallstackHash = 17;
-			if (CurrentCallstack.Num())
+		Timeline->EnumerateEvents(Params.IntervalStart, Params.IntervalEnd,
+			[this, &Params, &CurrentCallstackKey, &CurrentCallstack, &CallstackNodeMap, &LastTime, Butterfly]
+			(bool IsEnter, double Time, const FTimingProfilerEvent& Event)
 			{
-				FLocalStackEntry& Stackentry = CurrentCallstack.Top();
-				ParentNode = Stackentry.Node;
-				ParentCallstackHash = Stackentry.CurrentCallstackHash;
-				Stackentry.ExclusiveTime += Time - LastTime;
-			}
-			LastTime = Time;
-			if (IsEnter)
-			{
-				const FTimingProfilerTimer* Timer = GetTimer(Event.TimerIndex);
-				check(Timer != nullptr);
-
-				FLocalStackEntry& StackEntry = CurrentCallstack.AddDefaulted_GetRef();
-				StackEntry.StartTime = Time;
-				StackEntry.CurrentCallstackHash = ParentCallstackHash * 17 + Timer->Id;
-
-				CurrentCallstackKey.TimerStack.Push(Timer->Id);
-				CurrentCallstackKey.Hash = StackEntry.CurrentCallstackHash;
-
-				TTuple<FTimingProfilerButterflyNode*, bool>* FindIt = CallstackNodeMap.Find(CurrentCallstackKey);
-				if (FindIt)
+				Time = FMath::Clamp(Time, Params.IntervalStart, Params.IntervalEnd);
+				FTimingProfilerButterflyNode* ParentNode = nullptr;
+				uint32 ParentCallstackHash = 17;
+				if (CurrentCallstack.Num())
 				{
-					StackEntry.Node = FindIt->Get<0>();
-					StackEntry.bIsRecursive = FindIt->Get<1>();
+					FLocalStackEntry& StackEntry = CurrentCallstack.Top();
+					ParentNode = StackEntry.Node;
+					ParentCallstackHash = StackEntry.CurrentCallstackHash;
+					StackEntry.ExclusiveTime += Time - LastTime;
 				}
-				else
+				LastTime = Time;
+				if (IsEnter)
 				{
-					for (int32 StackIndex = 0, StackEnd = CurrentCallstack.Num() - 1; StackIndex < StackEnd; ++StackIndex)
+					const FTimingProfilerTimer* Timer = GetTimer(Event.TimerIndex);
+					check(Timer != nullptr);
+
+					FLocalStackEntry& StackEntry = CurrentCallstack.AddDefaulted_GetRef();
+					StackEntry.StartTime = Time;
+					StackEntry.CurrentCallstackHash = ParentCallstackHash * 17 + Timer->Id;
+
+					CurrentCallstackKey.TimerStack.Push(Timer->Id);
+					CurrentCallstackKey.Hash = StackEntry.CurrentCallstackHash;
+
+					TTuple<FTimingProfilerButterflyNode*, bool>* FindIt = CallstackNodeMap.Find(CurrentCallstackKey);
+					if (FindIt)
+					{
+						StackEntry.Node = FindIt->Get<0>();
+						StackEntry.bIsRecursive = FindIt->Get<1>();
+					}
+					else
 					{
-						if (CurrentCallstack[StackIndex].Node->Timer == Timer)
+						for (int32 StackIndex = 0, StackEnd = CurrentCallstack.Num() - 1; StackIndex < StackEnd; ++StackIndex)
 						{
-							StackEntry.Node = CurrentCallstack[StackIndex].Node;
-							StackEntry.bIsRecursive = true;
-							break;
+							if (CurrentCallstack[StackIndex].Node->Timer == Timer)
+							{
+								StackEntry.Node = CurrentCallstack[StackIndex].Node;
+								StackEntry.bIsRecursive = true;
+								break;
+							}
 						}
-					}
 
-					if (!StackEntry.Node)
-					{
-						StackEntry.Node = &Butterfly->Nodes.PushBack();
-						StackEntry.Node->InclusiveTime = 0.0;
-						StackEntry.Node->ExclusiveTime = 0.0;
-						StackEntry.Node->Count = 0;
-						StackEntry.Node->Timer = Timer;
-						Butterfly->TimerCallstacksMap[Timer->Id].Add(StackEntry.Node);
-
-						StackEntry.Node->Parent = ParentNode;
-						if (ParentNode)
+						if (!StackEntry.Node)
 						{
-							ParentNode->Children.Add(StackEntry.Node);
+							StackEntry.Node = &Butterfly->Nodes.PushBack();
+							StackEntry.Node->InclusiveTime = 0.0;
+							StackEntry.Node->ExclusiveTime = 0.0;
+							StackEntry.Node->Count = 0;
+							StackEntry.Node->Timer = Timer;
+							Butterfly->TimerCallstacksMap[Timer->Id].Add(StackEntry.Node);
+
+							StackEntry.Node->Parent = ParentNode;
+							if (ParentNode)
+							{
+								ParentNode->Children.Add(StackEntry.Node);
+							}
 						}
+						CallstackNodeMap.Add(CurrentCallstackKey, MakeTuple(StackEntry.Node, StackEntry.bIsRecursive));
 					}
-					CallstackNodeMap.Add(CurrentCallstackKey, MakeTuple(StackEntry.Node, StackEntry.bIsRecursive));
 				}
-			}
-			else
-			{
-				FLocalStackEntry& StackEntry = CurrentCallstack.Top();
-				double InclusiveTime = Time - StackEntry.StartTime;
-				check(InclusiveTime >= 0.0);
-				check(StackEntry.ExclusiveTime >= 0.0 && StackEntry.ExclusiveTime <= InclusiveTime);
-				if (!StackEntry.bIsRecursive)
+				else
 				{
-					StackEntry.Node->InclusiveTime += InclusiveTime;
-				}
-				StackEntry.Node->ExclusiveTime += StackEntry.ExclusiveTime;
-				++StackEntry.Node->Count;
+					FLocalStackEntry& StackEntry = CurrentCallstack.Top();
+					double InclusiveTime = Time - StackEntry.StartTime;
+					check(InclusiveTime >= 0.0);
+					check(StackEntry.ExclusiveTime >= 0.0 && StackEntry.ExclusiveTime <= InclusiveTime);
+					if (!StackEntry.bIsRecursive)
+					{
+						StackEntry.Node->InclusiveTime += InclusiveTime;
+					}
+					StackEntry.Node->ExclusiveTime += StackEntry.ExclusiveTime;
+					++StackEntry.Node->Count;
 
-				CurrentCallstack.Pop(EAllowShrinking::No);
-				CurrentCallstackKey.TimerStack.Pop(EAllowShrinking::No);
-			}
+					CurrentCallstack.Pop(EAllowShrinking::No);
+					CurrentCallstackKey.TimerStack.Pop(EAllowShrinking::No);
+				}
 
-			return EEventEnumerate::Continue;
-		});
+				return EEventEnumerate::Continue;
+			});
 	}
 	return Butterfly;
 }
diff --git a/Engine/Source/Developer/TraceServices/Private/Model/TimingProfilerPrivate.h b/Engine/Source/Developer/TraceServices/Private/Model/TimingProfilerPrivate.h
--- a/Engine/Source/Developer/TraceServices/Private/Model/TimingProfilerPrivate.h
+++ b/Engine/Source/Developer/TraceServices/Private/Model/TimingProfilerPrivate.h
@@ -141,6 +141,8 @@ public:
 
 private:
 	FTimingProfilerTimer& AddTimerInternal(ETimingProfilerTimerType TimerType, FStringView Name, const TCHAR* File, uint32 Line);
+	ITable<FTimingProfilerAggregatedStats>* CreateAggregationInternal(TArray<const ITimingProfilerTimeline*>& IncludedTimelines, const FCreateAggregationParams& Params) const;
+	ITimingProfilerButterfly* CreateButterflyInternal(TArray<const ITimingProfilerTimeline*>& IncludedTimelines, const FCreateButterflyParams& Params) const;
 
 	struct FMetadata
 	{
diff --git a/Engine/Source/Developer/TraceServices/Public/TraceServices/Model/TimingProfiler.h b/Engine/Source/Developer/TraceServices/Public/TraceServices/Model/TimingProfiler.h
--- a/Engine/Source/Developer/TraceServices/Public/TraceServices/Model/TimingProfiler.h
+++ b/Engine/Source/Developer/TraceServices/Public/TraceServices/Model/TimingProfiler.h
@@ -2,23 +2,36 @@
 
 #pragma once
 
+#include "CoreTypes.h"
 #include "Containers/Array.h"
 #include "Containers/ArrayView.h"
 #include "Containers/StringFwd.h"
 #include "HAL/Platform.h"
 #include "HAL/PlatformMath.h"
+#include "UObject/NameTypes.h"
+
+// TraceServices
 #include "Model/MonotonicTimeline.h"
 #include "TraceServices/Common/CancellationToken.h"
 #include "TraceServices/Containers/Tables.h"
 #include "TraceServices/Containers/Timelines.h"
 #include "TraceServices/Model/AnalysisSession.h"
-#include "UObject/NameTypes.h"
 
 namespace TraceServices
 {
+
+template <typename InEventType> class ITimeline;
 template <typename InEventType> class IEditableTimeline;
 template <typename RowType> class ITable;
 
+struct FTimingProfilerEvent
+{
+	uint32 TimerIndex = uint32(-1);
+};
+
+typedef ITimeline<FTimingProfilerEvent> ITimingProfilerTimeline;
+typedef IEditableTimeline<FTimingProfilerEvent> ITimingProfilerEditableTimeline;
+
 struct FMetadataSpec
 {
 	const TCHAR* Format = nullptr;
@@ -64,11 +77,6 @@ struct FTimingProfilerTimer
 	}
 };
 
-struct FTimingProfilerEvent
-{
-	uint32 TimerIndex = uint32(-1);
-};
-
 struct FTimingProfilerAggregatedStats
 {
 	const FTimingProfilerTimer* Timer = nullptr;
@@ -118,7 +126,7 @@ struct FCreateAggregationParams
 	// The end timestamp in seconds.
 	double IntervalEnd = 0.0;
 
-	// A function to filter the GPU queues to aggregate.
+	// A function to filter (by GPU Queue Id) the GPU queues to aggregate.
 	TFunction<bool(uint32)> GpuQueueFilter;
 
 	// A boolean value to specify if aggregation should include the old GPU timeline.
@@ -130,13 +138,16 @@ struct FCreateAggregationParams
 	// A boolean value to specify if aggregation should include the Verse Sampling timeline.
 	bool bIncludeVerseSampling = false;
 
-	// A function to filter the CPU threads to aggregate.
+	// A function to filter (by Thread Id) the CPU threads to aggregate.
 	TFunction<bool(uint32)> CpuThreadFilter;
 
-	// Whether to sort the table by a field
+	// A list of additional custom timelines to aggregate.
+	TArray<const ITimingProfilerTimeline*> CustomTimelines;
+
+	// Whether to sort the table by a field.
 	ESortBy SortBy = ESortBy::DontSort;
 
-	// Descending or ascending
+	// The sorting order
 	ESortOrder SortOrder = ESortOrder::DontSort;
 
 	// Limit the entries in aggregation (e.g. "top 100")
@@ -170,6 +181,9 @@ struct FCreateButterflyParams
 
 	// A function to filter the CPU threads to aggregate.
 	TFunction<bool(uint32)> CpuThreadFilter;
+
+	// A list of additional custom timelines to aggregate.
+	TArray<const ITimingProfilerTimeline*> CustomTimelines;
 };
 
 class ITimingProfilerButterfly
