[Insights]
* Regions: Added analysis support for exposing region timers and region categories. This is a prerequisite for exposing unique regions (timers) and region categories to UI.
* Regions: Improved performance for Find First , Find Last when using region name or category filtering (by deducing a PriorityTrack automatically).

diff --git a/Engine/Source/Developer/TraceInsights/Private/Insights/TimingProfiler/Tracks/RegionsTimingTrack.cpp b/Engine/Source/Developer/TraceInsights/Private/Insights/TimingProfiler/Tracks/RegionsTimingTrack.cpp
--- a/Engine/Source/Developer/TraceInsights/Private/Insights/TimingProfiler/Tracks/RegionsTimingTrack.cpp
+++ b/Engine/Source/Developer/TraceInsights/Private/Insights/TimingProfiler/Tracks/RegionsTimingTrack.cpp
@@ -70,24 +70,28 @@ void FTimingRegionsTrack::InitTooltip(FTooltipDrawState& InOutTooltip, const ITi
 		};
 
 		FTimingEventSearchParameters SearchParameters(TooltipEvent.GetStartTime(), TooltipEvent.GetEndTime(), ETimingEventSearchFlags::StopAtFirstMatch, MatchEvent);
-		FindRegionEvent(SearchParameters, [this, &InOutTooltip, &TooltipEvent](double InFoundStartTime, double InFoundEndTime, uint32 InFoundDepth, const TraceServices::FTimeRegion& InRegion)
-		{
-			InOutTooltip.Reset();
-			InOutTooltip.AddTitle(InRegion.Text, FLinearColor::White);
-			InOutTooltip.AddNameValueTextLine(TEXT("Duration:"),  FormatTimeAuto(InRegion.EndTime - InRegion.BeginTime));
-			InOutTooltip.AddNameValueTextLine(TEXT("Depth:"),  FString::FromInt(InRegion.Depth));
-			if (InRegion.Category)
+		FindRegionEvent(SearchParameters,
+			[this, &InOutTooltip, &TooltipEvent]
+			(double InFoundStartTime, double InFoundEndTime, uint32 InFoundDepth, const TraceServices::FTimeRegion& InRegion)
 			{
-				InOutTooltip.AddNameValueTextLine(TEXT("Category:"),  InRegion.Category);
-			}
-			InOutTooltip.UpdateLayout();
-		});
+				check(InRegion.Timer);
+				InOutTooltip.Reset();
+				InOutTooltip.AddTitle(InRegion.Timer->Name, FLinearColor::White);
+				InOutTooltip.AddNameValueTextLine(TEXT("Duration:"), FormatTimeAuto(InRegion.EndTime - InRegion.BeginTime));
+				InOutTooltip.AddNameValueTextLine(TEXT("Depth:"), FString::FromInt(InRegion.Depth));
+				if (InRegion.Timer->Category && InRegion.Timer->Category->Name)
+				{
+					InOutTooltip.AddNameValueTextLine(TEXT("Category:"), InRegion.Timer->Category->Name);
+				}
+				InOutTooltip.UpdateLayout();
+			});
 	}
 }
 
 ////////////////////////////////////////////////////////////////////////////////////////////////////
 
-void FTimingRegionsTrack::BuildDrawState(ITimingEventsTrackDrawStateBuilder& Builder,
+void FTimingRegionsTrack::BuildDrawState(
+	ITimingEventsTrackDrawStateBuilder& Builder,
 	const ITimingTrackUpdateContext& Context)
 {
 	const FTimingTrackViewport& Viewport = Context.GetViewport();
@@ -103,20 +107,30 @@ void FTimingRegionsTrack::BuildDrawState(ITimingEventsTrackDrawStateBuilder& Bui
 	int32 CurDepth = 0;
 	const TraceServices::IRegionTimeline* Timeline = RegionProvider.GetTimelineForCategory(RegionsCategory);
 	check(Timeline)
-	
-	Timeline->EnumerateLanes([this, Viewport, &CurDepth, &Builder](const TraceServices::FRegionLane& Lane, const int32 Depth)
-	{
-		bool RegionHadEvents = false;
-		Lane.EnumerateRegions(Viewport.GetStartTime(), Viewport.GetEndTime(), [this, &Builder, &RegionHadEvents, &CurDepth](const TraceServices::FTimeRegion& Region) -> bool
+
+	Timeline->EnumerateLanes(
+		[this, Viewport, &CurDepth, &Builder]
+		(const TraceServices::FRegionLane& Lane, const int32 Depth)
 		{
-			RegionHadEvents = true;
-			uint32 EventColor = SharedState.bColorRegionsByCategory ? FTimingEvent::ComputeEventColor(Region.Category) : FTimingEvent::ComputeEventColor(Region.Text);
-			Builder.AddEvent(Region.BeginTime, Region.EndTime,CurDepth, Region.Text, 0 , EventColor);
-			return true;
-		});
+			bool RegionHadEvents = false;
+			Lane.EnumerateRegions(Viewport.GetStartTime(), Viewport.GetEndTime(),
+				[this, &Builder, &RegionHadEvents, &CurDepth]
+				(const TraceServices::FTimeRegion& Region) -> bool
+				{
+					check(Region.Timer);
+					RegionHadEvents = true;
+					uint32 EventColor = SharedState.bColorRegionsByCategory ?
+						FTimingEvent::ComputeEventColor(Region.Timer->Category ? Region.Timer->Category->Name : nullptr) :
+						FTimingEvent::ComputeEventColor(Region.Timer->Name);
+					Builder.AddEvent(Region.BeginTime, Region.EndTime,CurDepth, Region.Timer->Name, 0, EventColor);
+					return true;
+				});
 
-		if (RegionHadEvents) CurDepth++;
-	});
+			if (RegionHadEvents)
+			{
+				CurDepth++;
+			}
+		});
 
 	Stopwatch.Stop();
 	const double TotalTime = Stopwatch.GetAccumulatedTime();
@@ -125,7 +139,9 @@ void FTimingRegionsTrack::BuildDrawState(ITimingEventsTrackDrawStateBuilder& Bui
 
 ////////////////////////////////////////////////////////////////////////////////////////////////////
 
-void FTimingRegionsTrack::BuildFilteredDrawState(ITimingEventsTrackDrawStateBuilder& Builder, const ITimingTrackUpdateContext& Context)
+void FTimingRegionsTrack::BuildFilteredDrawState(
+	ITimingEventsTrackDrawStateBuilder& Builder,
+	const ITimingTrackUpdateContext& Context)
 {
 	const TSharedPtr<ITimingEventFilter> EventFilterPtr = Context.GetEventFilter();
 	if (EventFilterPtr.IsValid() && EventFilterPtr->FilterTrack(*this))
@@ -152,20 +168,28 @@ void FTimingRegionsTrack::BuildFilteredDrawState(ITimingEventsTrackDrawStateBuil
 				int32 CurDepth = 0;
 				const TraceServices::IRegionTimeline* Timeline = RegionProvider.GetTimelineForCategory(RegionsCategory);
 				check(Timeline)
-				Timeline->EnumerateLanes([this, Viewport, &CurDepth, &Builder, FilterEventType](const TraceServices::FRegionLane& Lane, const int32 Depth)
+				Timeline->EnumerateLanes(
+					[this, Viewport, &CurDepth, &Builder, FilterEventType]
+					(const TraceServices::FRegionLane& Lane, const int32 Depth)
 					{
 						bool RegionHadEvents = false;
-						Lane.EnumerateRegions(Viewport.GetStartTime(), Viewport.GetEndTime(), [&Builder, &RegionHadEvents, &CurDepth, FilterEventType](const TraceServices::FTimeRegion& Region) -> bool
+						Lane.EnumerateRegions(Viewport.GetStartTime(), Viewport.GetEndTime(),
+							[&Builder, &RegionHadEvents, &CurDepth, FilterEventType]
+							(const TraceServices::FTimeRegion& Region) -> bool
 							{
+								check(Region.Timer);
 								RegionHadEvents = true;
-								if (Region.Text == FilterEventType)
+								if (Region.Timer->Name == FilterEventType)
 								{
-									Builder.AddEvent(Region.BeginTime, Region.EndTime, CurDepth, Region.Text);
+									Builder.AddEvent(Region.BeginTime, Region.EndTime, CurDepth, Region.Timer->Name);
 								}
 								return true;
 							});
 
-						if (RegionHadEvents) CurDepth++;
+						if (RegionHadEvents)
+						{
+							CurDepth++;
+						}
 					});
 			}
 			else // generic filter
@@ -203,27 +227,35 @@ void FTimingRegionsTrack::BuildFilteredDrawState(ITimingEventsTrackDrawStateBuil
 			int32 CurDepth = 0;
 			const TraceServices::IRegionTimeline* Timeline = RegionProvider.GetTimelineForCategory(RegionsCategory);
 			check(Timeline)
-			Timeline->EnumerateLanes([this, Viewport, &CurDepth, &Builder, &FilterContext](const TraceServices::FRegionLane& Lane, const int32 Depth)
+			Timeline->EnumerateLanes(
+				[this, Viewport, &CurDepth, &Builder, &FilterContext]
+				(const TraceServices::FRegionLane& Lane, const int32 Depth)
 				{
 					bool RegionHadEvents = false;
-					Lane.EnumerateRegions(Viewport.GetStartTime(), Viewport.GetEndTime(), [&Builder, &RegionHadEvents, &CurDepth, &FilterContext, this](const TraceServices::FTimeRegion& Region) -> bool
+					Lane.EnumerateRegions(Viewport.GetStartTime(), Viewport.GetEndTime(),
+						[&Builder, &RegionHadEvents, &CurDepth, &FilterContext, this]
+						(const TraceServices::FTimeRegion& Region) -> bool
 						{
+							check(Region.Timer);
 							FilterContext.SetFilterData<double>(static_cast<int32>(EFilterField::StartTime), Region.BeginTime);
 							FilterContext.SetFilterData<double>(static_cast<int32>(EFilterField::EndTime), Region.EndTime);
 							FilterContext.SetFilterData<double>(static_cast<int32>(EFilterField::Duration), Region.EndTime - Region.BeginTime);
-							FilterContext.SetFilterData<FString>(static_cast<int32>(EFilterField::RegionName), Region.Text);
-							FilterContext.SetFilterData<FString>(static_cast<int32>(EFilterField::RegionCategory), Region.Category);
+							FilterContext.SetFilterData<FString>(static_cast<int32>(EFilterField::RegionName), Region.Timer->Name);
+							FilterContext.SetFilterData<FString>(static_cast<int32>(EFilterField::RegionCategory), Region.Timer->Category ? Region.Timer->Category->Name : nullptr);
 
 							RegionHadEvents = true;
 							if (FilterConfigurator->ApplyFilters(FilterContext))
 							{
-								Builder.AddEvent(Region.BeginTime, Region.EndTime, CurDepth, Region.Text);
+								Builder.AddEvent(Region.BeginTime, Region.EndTime, CurDepth, Region.Timer->Name);
 							}
 
 							return true;
 						});
 
-					if (RegionHadEvents) CurDepth++;
+					if (RegionHadEvents)
+					{
+						CurDepth++;
+					}
 				});
 		}
 	}
@@ -236,17 +268,21 @@ const TSharedPtr<const ITimingEvent> FTimingRegionsTrack::SearchEvent(
 {
 	TSharedPtr<const ITimingEvent> FoundEvent;
 
-	FindRegionEvent(InSearchParameters, [this, &FoundEvent](double InFoundStartTime, double InFoundEndTime, uint32 InFoundDepth, const TraceServices::FTimeRegion& InEvent)
-	{
-		FoundEvent = MakeShared<FTimingEvent>(SharedThis(this), InFoundStartTime, InFoundEndTime, InFoundDepth, reinterpret_cast<uint64>(InEvent.Text));
-	});
+	FindRegionEvent(InSearchParameters,
+		[this, &FoundEvent]
+		(double InFoundStartTime, double InFoundEndTime, uint32 InFoundDepth, const TraceServices::FTimeRegion& InRegion)
+		{
+			check(InRegion.Timer);
+			FoundEvent = MakeShared<FTimingEvent>(SharedThis(this), InFoundStartTime, InFoundEndTime, InFoundDepth, reinterpret_cast<uint64>(InRegion.Timer->Name));
+		});
 
 	return FoundEvent;
 }
 
 ////////////////////////////////////////////////////////////////////////////////////////////////////
 
-bool FTimingRegionsTrack::FindRegionEvent(const FTimingEventSearchParameters& InParameters,
+bool FTimingRegionsTrack::FindRegionEvent(
+	const FTimingEventSearchParameters& InParameters,
 	TFunctionRef<void(double, double, uint32, const TraceServices::FTimeRegion&)> InFoundPredicate) const
 {
 	// If the query start time is larger than the end of the session return false.
@@ -269,65 +305,76 @@ bool FTimingRegionsTrack::FindRegionEvent(const FTimingEventSearchParameters& In
 	FilterConfiguratorContext.AddFilterData<FString>(static_cast<int32>(EFilterField::RegionCategory), "");
 
 	return TTimingEventSearch<TraceServices::FTimeRegion>::Search(
-	InParameters,
-
-	// Search...
-	[this](TTimingEventSearch<TraceServices::FTimeRegion>::FContext& InContext)
-	{
-		TSharedPtr<const TraceServices::IAnalysisSession> Session = FInsightsManager::Get()->GetSession();
-		const TraceServices::IRegionProvider& RegionProvider = TraceServices::ReadRegionProvider(*Session);
-		TraceServices::FProviderReadScopeLock RegionProviderScopedLock(RegionProvider);
 
-		const TraceServices::IRegionTimeline* Timeline = RegionProvider.GetTimelineForCategory(RegionsCategory);
-		check(Timeline)
+		InParameters,
 
-		auto Callback = [&InContext](const TraceServices::FTimeRegion& Region)
+		// Search Predicate
+		[this]
+		(TTimingEventSearch<TraceServices::FTimeRegion>::FContext& InContext)
 		{
-			InContext.Check(Region.BeginTime, Region.EndTime, Region.Depth, Region);
+			TSharedPtr<const TraceServices::IAnalysisSession> Session = FInsightsManager::Get()->GetSession();
+			const TraceServices::IRegionProvider& RegionProvider = TraceServices::ReadRegionProvider(*Session);
+			TraceServices::FProviderReadScopeLock RegionProviderScopedLock(RegionProvider);
+
+			const TraceServices::IRegionTimeline* Timeline = RegionProvider.GetTimelineForCategory(RegionsCategory);
+			check(Timeline)
 
-			if (!InContext.ShouldContinueSearching())
+			auto Callback =
+				[&InContext]
+				(const TraceServices::FTimeRegion& Region)
+				{
+					InContext.Check(Region.BeginTime, Region.EndTime, Region.Depth, Region);
+
+					if (!InContext.ShouldContinueSearching())
+					{
+						return false;
+					}
+
+					return true;
+				};
+
+			if (InContext.GetParameters().SearchDirection == FTimingEventSearchParameters::ESearchDirection::Forward)
 			{
-				return false;
+				Timeline->EnumerateRegions(InContext.GetParameters().StartTime, InContext.GetParameters().EndTime, Callback);
+			}
+			else
+			{
+				Timeline->EnumerateRegionsBackwards(InContext.GetParameters().EndTime, InContext.GetParameters().StartTime, Callback);
 			}
+		},
 
-			return true;
-		};
-			
-		if (InContext.GetParameters().SearchDirection == FTimingEventSearchParameters::ESearchDirection::Forward)
-		{
-			Timeline->EnumerateRegions(InContext.GetParameters().StartTime, InContext.GetParameters().EndTime, Callback);
-		}
-		else
-		{
-			Timeline->EnumerateRegionsBackwards(InContext.GetParameters().EndTime, InContext.GetParameters().StartTime, Callback);
-		}
-	},
-	[&FilterConfiguratorContext, &InParameters](double EventStartTime, double EventEndTime, uint32 EventDepth, const TraceServices::FTimeRegion& Region)
-	{
-		if (!InParameters.FilterExecutor.IsValid())
+		// Filter Predicate
+		[&FilterConfiguratorContext, &InParameters]
+		(double EventStartTime, double EventEndTime, uint32 EventDepth, const TraceServices::FTimeRegion& Region)
 		{
-			return true;
-		}
+			if (!InParameters.FilterExecutor.IsValid())
+			{
+				return true;
+			}
 
-		TSharedPtr<const TraceServices::IAnalysisSession> Session = FInsightsManager::Get()->GetSession();
-		if (Session.IsValid())
-		{
+			TSharedPtr<const TraceServices::IAnalysisSession> Session = FInsightsManager::Get()->GetSession();
+			if (!Session.IsValid())
+			{
+				// Region.Timer->Name and Region.Category strings are only available during a valid session.
+				return false;
+			}
+
+			check(Region.Timer);
 			FilterConfiguratorContext.SetFilterData<double>(static_cast<int32>(EFilterField::StartTime), EventStartTime);
 			FilterConfiguratorContext.SetFilterData<double>(static_cast<int32>(EFilterField::EndTime), EventEndTime);
 			FilterConfiguratorContext.SetFilterData<double>(static_cast<int32>(EFilterField::Duration), EventEndTime - EventStartTime);
-			FilterConfiguratorContext.SetFilterData<FString>(static_cast<int32>(EFilterField::RegionName), Region.Text);
-			FilterConfiguratorContext.SetFilterData<FString>(static_cast<int32>(EFilterField::RegionCategory), Region.Category);
+			FilterConfiguratorContext.SetFilterData<FString>(static_cast<int32>(EFilterField::RegionName), Region.Timer->Name);
+			FilterConfiguratorContext.SetFilterData<FString>(static_cast<int32>(EFilterField::RegionCategory), Region.Timer->Category ? Region.Timer->Category->Name : nullptr);
 			return InParameters.FilterExecutor->ApplyFilters(FilterConfiguratorContext);
-		}
+		},
 
-		return false;
-	},
-	[&InFoundPredicate](double InFoundStartTime, double InFoundEndTime, uint32 InFoundDepth, const TraceServices::FTimeRegion& InEvent)
-	{
-		InFoundPredicate(InFoundStartTime, InFoundEndTime, InFoundDepth, InEvent);
-	},
+		// Found Predicate
+		InFoundPredicate,
+
+		// Payload Matched Predicate
+		TTimingEventSearch<TraceServices::FTimeRegion>::NoMatch
 
-	TTimingEventSearch<TraceServices::FTimeRegion>::NoMatch);
+	); // Search
 }
 
 ////////////////////////////////////////////////////////////////////////////////////////////////////
diff --git a/Engine/Source/Developer/TraceInsights/Private/Insights/TimingProfiler/ViewModels/TimingExporter.cpp b/Engine/Source/Developer/TraceInsights/Private/Insights/TimingProfiler/ViewModels/TimingExporter.cpp
--- a/Engine/Source/Developer/TraceInsights/Private/Insights/TimingProfiler/ViewModels/TimingExporter.cpp
+++ b/Engine/Source/Developer/TraceInsights/Private/Insights/TimingProfiler/ViewModels/TimingExporter.cpp
@@ -713,21 +713,23 @@ void FTimingExporter::GetRegions(const FString& InRegionNamePattern, TMap<FStrin
 		RegionProvider.GetDefaultTimeline().EnumerateRegions(0.0, std::numeric_limits<double>::max(),
 			[&RegionCount, &RegionNameSpec, &OutRegionGroups](const TraceServices::FTimeRegion& InRegion) -> bool
 			{
-				if (!RegionNameSpec.Match(InRegion.Text))
+				if (InRegion.Timer == nullptr || InRegion.Timer->Name == nullptr)
 				{
 					return true;
 				}
-
-				// Handle duplicate region names, individual regions may appear multiple times
-				// we append numbers to allow for unique export filenames.
-				FString RegionName = InRegion.Text;
-				FTimeRegionGroup* ExistingRegionGroup = OutRegionGroups.Find(RegionName);
-				if (!ExistingRegionGroup)
+				const FString RegionName = InRegion.Timer->Name;
+				if (RegionNameSpec.Match(RegionName))
 				{
-					ExistingRegionGroup = &OutRegionGroups.Add(RegionName, FTimeRegionGroup{});
+					// Handle duplicate region names. Regions with same name may appear multiple times.
+					// We append numbers to allow for unique export filenames.
+					FTimeRegionGroup* ExistingRegionGroup = OutRegionGroups.Find(RegionName);
+					if (!ExistingRegionGroup)
+					{
+						ExistingRegionGroup = &OutRegionGroups.Add(RegionName, FTimeRegionGroup{});
+					}
+					ExistingRegionGroup->Intervals.Add(FTimeRegionInterval{ InRegion.BeginTime, InRegion.EndTime });
+					++RegionCount;
 				}
-				ExistingRegionGroup->Intervals.Add(FTimeRegionInterval{ InRegion.BeginTime, InRegion.EndTime });
-				++RegionCount;
 				return true;
 			});
 	}
diff --git a/Engine/Source/Developer/TraceInsights/Private/Insights/Widgets/STimingView.cpp b/Engine/Source/Developer/TraceInsights/Private/Insights/Widgets/STimingView.cpp
--- a/Engine/Source/Developer/TraceInsights/Private/Insights/Widgets/STimingView.cpp
+++ b/Engine/Source/Developer/TraceInsights/Private/Insights/Widgets/STimingView.cpp
@@ -366,7 +366,7 @@ TSharedPtr<FFilterConfigurator> STimingView::MakeCustomFilterConfigurator(uint32
 		LOCTEXT("EndTime", "End Time"),
 		EFilterDataType::Double,
 		MakeShared<FTimeFilterValueConverter>(),
-		FFilterService::Get()->GetDoubleOperators()));  
+		FFilterService::Get()->GetDoubleOperators()));
 
 	CustomFilterConfigurator->Add(MakeShared<FFilter>(
 		static_cast<int32>(EFilterField::TimerId),
@@ -394,7 +394,7 @@ TSharedPtr<FFilterConfigurator> STimingView::MakeCustomFilterConfigurator(uint32
 	CustomFilterConfigurator->GetRootNode()->AddChildAndSetParent(EndTimeNode);
 	CustomFilterConfigurator->GetRootNode()->AddChildAndSetParent(TimerIdNode);
 	CustomFilterConfigurator->GetRootNode()->ProcessFilter();
-		
+
 	return CustomFilterConfigurator;
 }
 
@@ -5260,7 +5260,7 @@ void STimingView::QuickFind_Execute()
 
 		NewFilterConfigurator->Add(MakeShared<FTimerNameFilter>());
 		NewFilterConfigurator->Add(MakeShared<FMetadataFilter>());
-		
+
 		TSharedRef<FFilterWithSuggestions> RegionNameFilter = MakeShared<FFilterWithSuggestions>(
 			static_cast<int32>(EFilterField::RegionName),
 			LOCTEXT("RegionName", "Region Name"),
@@ -5465,33 +5465,38 @@ void STimingView::FindPrevEvent(bool bUseCustomFilter)
 	TSharedPtr<const ITimingEvent> BestMatchEvent;
 	double StartTime = SelectedEvent.IsValid() ? SelectedEvent->GetStartTime() : std::numeric_limits<double>::max();
 
+	TSharedPtr<FFilterConfigurator> Filter = bUseCustomFilter ? CustomFilterConfigurator : QuickFindVm->GetFilterConfigurator();
+
 	auto EventFilter = [StartTime](double EventStartTime, double EventEndTime, uint32 EventDepth)
 	{
 		return EventStartTime < StartTime;
 	};
+
 	FTimingEventSearchParameters Params(std::numeric_limits<double>::lowest(), StartTime, ETimingEventSearchFlags::StopAtFirstMatch, EventFilter);
-	Params.FilterExecutor = bUseCustomFilter ? CustomFilterConfigurator : QuickFindVm->GetFilterConfigurator();
+	Params.FilterExecutor = Filter;
 	Params.SearchDirection = FTimingEventSearchParameters::ESearchDirection::Backward;
 
 	TSharedPtr<const FBaseTimingTrack> PriorityTrack = SelectedEvent.IsValid() ? SelectedEvent->GetTrack().ToSharedPtr() : nullptr;
-	EnumerateFilteredTracks(bUseCustomFilter ? CustomFilterConfigurator : QuickFindVm->GetFilterConfigurator(), PriorityTrack, [&Params, &BestMatchEvent](TSharedPtr<const FBaseTimingTrack> Track)
+	if (!PriorityTrack.IsValid() && Filter.IsValid())
 	{
-		if (!Track->IsVisible())
-		{
-			return;
-		}
-
-		Params.StartTime = BestMatchEvent.IsValid() ? BestMatchEvent->GetStartTime() : std::numeric_limits<double>::lowest();
+		PriorityTrack = FindPriorityTrack(*Filter->GetRootNode());
+	}
 
-		TSharedPtr<const ITimingEvent> FoundEvent = Track->SearchEvent(Params);
-		if (FoundEvent.IsValid())
+	EnumerateFilteredTracks(Filter, PriorityTrack,
+		[&Params, &BestMatchEvent]
+		(TSharedPtr<const FBaseTimingTrack> Track)
 		{
-			if (!BestMatchEvent.IsValid() || BestMatchEvent->GetStartTime() < FoundEvent->GetStartTime())
+			Params.StartTime = BestMatchEvent.IsValid() ? BestMatchEvent->GetStartTime() : std::numeric_limits<double>::lowest();
+
+			TSharedPtr<const ITimingEvent> FoundEvent = Track->SearchEvent(Params);
+			if (FoundEvent.IsValid())
 			{
-				BestMatchEvent = FoundEvent;
+				if (!BestMatchEvent.IsValid() || BestMatchEvent->GetStartTime() < FoundEvent->GetStartTime())
+				{
+					BestMatchEvent = FoundEvent;
+				}
 			}
-		}
-	});
+		});
 
 	PostFindEvent(BestMatchEvent);
 }
@@ -5503,33 +5508,38 @@ void STimingView::FindNextEvent(bool bUseCustomFilter)
 	TSharedPtr<const ITimingEvent> BestMatchEvent;
 	double StartTime = SelectedEvent.IsValid() ? SelectedEvent->GetStartTime() : std::numeric_limits<double>::lowest();
 
+	TSharedPtr<FFilterConfigurator> Filter = bUseCustomFilter ? CustomFilterConfigurator : QuickFindVm->GetFilterConfigurator();
+
 	auto EventFilter = [StartTime](double EventStartTime, double EventEndTime, uint32 EventDepth)
 	{
 		return EventStartTime > StartTime;
 	};
+
 	FTimingEventSearchParameters Params(StartTime, std::numeric_limits<double>::max(), ETimingEventSearchFlags::StopAtFirstMatch, EventFilter);
-	Params.FilterExecutor = bUseCustomFilter ? CustomFilterConfigurator : QuickFindVm->GetFilterConfigurator();
+	Params.FilterExecutor = Filter;
 
 	TSharedPtr<const FBaseTimingTrack> PriorityTrack = SelectedEvent.IsValid() ? SelectedEvent->GetTrack().ToSharedPtr() : nullptr;
-	EnumerateFilteredTracks(bUseCustomFilter ? CustomFilterConfigurator : QuickFindVm->GetFilterConfigurator(), PriorityTrack, [&Params, &BestMatchEvent](TSharedPtr<const FBaseTimingTrack> Track)
+	if (!PriorityTrack.IsValid() && Filter.IsValid())
 	{
-		if (!Track->IsVisible())
-		{
-			return;
-		}
-
-		Params.EndTime = BestMatchEvent.IsValid() ? BestMatchEvent->GetStartTime() : std::numeric_limits<double>::max();
+		PriorityTrack = FindPriorityTrack(*Filter->GetRootNode());
+	}
 
-		TSharedPtr<const ITimingEvent> FoundEvent = Track->SearchEvent(Params);
-		if (FoundEvent.IsValid())
+	EnumerateFilteredTracks(Filter, PriorityTrack,
+		[&Params, &BestMatchEvent]
+		(TSharedPtr<const FBaseTimingTrack> Track)
 		{
-			if (BestMatchEvent.IsValid())
+			Params.EndTime = BestMatchEvent.IsValid() ? BestMatchEvent->GetStartTime() : std::numeric_limits<double>::max();
+
+			TSharedPtr<const ITimingEvent> FoundEvent = Track->SearchEvent(Params);
+			if (FoundEvent.IsValid())
 			{
-				ensure(FoundEvent->GetStartTime() < BestMatchEvent->GetStartTime());
+				if (BestMatchEvent.IsValid())
+				{
+					ensure(FoundEvent->GetStartTime() < BestMatchEvent->GetStartTime());
+				}
+				BestMatchEvent = FoundEvent;
 			}
-			BestMatchEvent = FoundEvent;
-		}
-	});
+		});
 
 	PostFindEvent(BestMatchEvent);
 }
@@ -5559,6 +5569,8 @@ void STimingView::FindMaxEvent()
 	TSharedPtr<const ITimingEvent> BestMatchEvent;
 	double CurrentDuration = -1.0;
 
+	TSharedPtr<FFilterConfigurator> Filter = QuickFindVm->GetFilterConfigurator();
+
 	auto EventFilter = [&CurrentDuration](double EventStartTime, double EventEndTime, uint32 EventDepth)
 	{
 		const double Duration = EventEndTime - EventStartTime;
@@ -5572,30 +5584,32 @@ void STimingView::FindMaxEvent()
 	};
 
 	FTimingEventSearchParameters Params(0.0, std::numeric_limits<double>::max(), ETimingEventSearchFlags::SearchAll, EventFilter, EventMatched);
-	Params.FilterExecutor = QuickFindVm->GetFilterConfigurator();
-	
+	Params.FilterExecutor = Filter;
+
 	TSharedPtr<const FBaseTimingTrack> PriorityTrack = SelectedEvent.IsValid() ? SelectedEvent->GetTrack().ToSharedPtr() : nullptr;
-	EnumerateFilteredTracks(QuickFindVm->GetFilterConfigurator(), PriorityTrack, [&Params, &BestMatchEvent](TSharedPtr<const FBaseTimingTrack> Track)
+	if (!PriorityTrack.IsValid() && Filter.IsValid())
 	{
-		if (!Track->IsVisible())
-		{
-			return;
-		}
+		PriorityTrack = FindPriorityTrack(*Filter->GetRootNode());
+	}
 
-		TSharedPtr<const ITimingEvent> FoundEvent = Track->SearchEvent(Params);
-		if (FoundEvent.IsValid())
+	EnumerateFilteredTracks(Filter, PriorityTrack,
+		[&Params, &BestMatchEvent]
+		(TSharedPtr<const FBaseTimingTrack> Track)
 		{
-			if (!BestMatchEvent.IsValid() || BestMatchEvent->GetDuration() < FoundEvent->GetDuration())
+			TSharedPtr<const ITimingEvent> FoundEvent = Track->SearchEvent(Params);
+			if (FoundEvent.IsValid())
 			{
-				BestMatchEvent = FoundEvent;
+				if (!BestMatchEvent.IsValid() || BestMatchEvent->GetDuration() < FoundEvent->GetDuration())
+				{
+					BestMatchEvent = FoundEvent;
+				}
 			}
-		}
-	});
+		});
 
 	PostFindEvent(BestMatchEvent);
 }
 
-////////////////////////////////////////////////////////////////////////////////////////////////////	
+////////////////////////////////////////////////////////////////////////////////////////////////////
 
 void STimingView::FindMinEvent()
 {
@@ -5608,6 +5622,8 @@ void STimingView::FindMinEvent()
 	TSharedPtr<const ITimingEvent> BestMatchEvent;
 	double CurrentDuration = std::numeric_limits<double>::max();
 
+	TSharedPtr<FFilterConfigurator> Filter = QuickFindVm->GetFilterConfigurator();
+
 	auto EventFilter = [&CurrentDuration](double EventStartTime, double EventEndTime, uint32 EventDepth)
 	{
 		const double Duration = EventEndTime - EventStartTime;
@@ -5621,25 +5637,27 @@ void STimingView::FindMinEvent()
 	};
 
 	FTimingEventSearchParameters Params(0.0, std::numeric_limits<double>::max(), ETimingEventSearchFlags::SearchAll, EventFilter, EventMatched);
-	Params.FilterExecutor = QuickFindVm->GetFilterConfigurator();
-	
+	Params.FilterExecutor = Filter;
+
 	TSharedPtr<const FBaseTimingTrack> PriorityTrack = SelectedEvent.IsValid() ? SelectedEvent->GetTrack().ToSharedPtr() : nullptr;
-	EnumerateFilteredTracks(QuickFindVm->GetFilterConfigurator(), PriorityTrack, [&Params, &BestMatchEvent](TSharedPtr<const FBaseTimingTrack> Track)
+	if (!PriorityTrack.IsValid() && Filter.IsValid())
 	{
-		if (!Track->IsVisible())
-		{
-			return;
-		}
+		PriorityTrack = FindPriorityTrack(*Filter->GetRootNode());
+	}
 
-		TSharedPtr<const ITimingEvent> FoundEvent = Track->SearchEvent(Params);
-		if (FoundEvent.IsValid())
+	EnumerateFilteredTracks(Filter, PriorityTrack,
+		[&Params, &BestMatchEvent]
+		(TSharedPtr<const FBaseTimingTrack> Track)
 		{
-			if (!BestMatchEvent.IsValid() || BestMatchEvent->GetDuration() > FoundEvent->GetDuration())
+			TSharedPtr<const ITimingEvent> FoundEvent = Track->SearchEvent(Params);
+			if (FoundEvent.IsValid())
 			{
-				BestMatchEvent = FoundEvent;
+				if (!BestMatchEvent.IsValid() || BestMatchEvent->GetDuration() > FoundEvent->GetDuration())
+				{
+					BestMatchEvent = FoundEvent;
+				}
 			}
-		}
-	});
+		});
 
 	PostFindEvent(BestMatchEvent);
 }
@@ -5669,6 +5687,41 @@ void STimingView::PostFindEvent(TSharedPtr<const ITimingEvent> BestMatchEvent)
 
 ////////////////////////////////////////////////////////////////////////////////////////////////////
 
+TSharedPtr<const FBaseTimingTrack> STimingView::FindPriorityTrack(FFilterConfiguratorNode& Node)
+{
+	for (FBaseTreeNodePtr Child : Node.GetChildren())
+	{
+		check(Child.IsValid() && Child->Is<FFilterConfiguratorNode>());
+		FFilterConfiguratorNode& ChildNode = Child->As<FFilterConfiguratorNode>();
+		TSharedPtr<FFilter> Filter = ChildNode.GetSelectedFilter();
+		if (Filter.IsValid())
+		{
+			if (Filter->GetKey() == static_cast<int32>(EFilterField::RegionName) ||
+				Filter->GetKey() == static_cast<int32>(EFilterField::RegionCategory))
+			{
+				for (auto& Entry : AllTracks)
+				{
+					if (Entry.Value->IsVisible() &&
+						Entry.Value->Is<FTimingRegionsTrack>())
+					{
+						return Entry.Value;
+					}
+				}
+				return nullptr;
+			}
+		}
+
+		TSharedPtr<const FBaseTimingTrack> Result = FindPriorityTrack(ChildNode);
+		if (Result.IsValid())
+		{
+			return Result;
+		}
+	}
+	return nullptr;
+}
+
+////////////////////////////////////////////////////////////////////////////////////////////////////
+
 void STimingView::FilterAllTracks()
 {
 	LLM_SCOPE_BYTAG(Insights);
@@ -5752,18 +5805,17 @@ void STimingView::PopulateRegionNameSuggestionList(const FString& Text, TArray<F
 	{
 		const TraceServices::IRegionProvider& RegionProvider = TraceServices::ReadRegionProvider(*Session);
 		TraceServices::FProviderReadScopeLock RegionProviderScopedLock(RegionProvider);
-		
-		RegionProvider.EnumerateTimelinesByCategory([&Text, &OutSuggestions](const TraceServices::IRegionTimeline& Timeline, const TCHAR* Category)
-		{
-			Timeline.EnumerateUniqueRegionNames([&Text, &OutSuggestions](const TCHAR* RegionName) -> bool
+
+		RegionProvider.EnumerateTimers(
+			[&Text, &OutSuggestions]
+			(const TraceServices::FRegionTimer& Timer) -> bool
 			{
-				if (FCString::Stristr(RegionName, *Text) || Text.IsEmpty())
+				if (Timer.Name && (FCString::Stristr(Timer.Name, *Text) || Text.IsEmpty()))
 				{
-					OutSuggestions.Add(RegionName);
+					OutSuggestions.Add(Timer.Name);
 				}
 				return true;
 			});
-		});
 	}
 }
 
@@ -5776,14 +5828,17 @@ void STimingView::PopulateRegionCategorySuggestionList(const FString& Text, TArr
 	{
 		const TraceServices::IRegionProvider& RegionProvider = TraceServices::ReadRegionProvider(*Session);
 		TraceServices::FProviderReadScopeLock RegionProviderScopedLock(RegionProvider);
-		
-		RegionProvider.EnumerateTimelinesByCategory([&Text, &OutSuggestions](const TraceServices::IRegionTimeline& Timeline, const TCHAR* Category)
-		{
-			if (FCString::Stristr(Category, *Text) || Text.IsEmpty())
+
+		RegionProvider.EnumerateCategories(
+			[&Text, &OutSuggestions]
+			(const TraceServices::FRegionCategory& Category) -> bool
 			{
-				OutSuggestions.Add(Category);
-			}
-		});
+				if (Category.Name && (FCString::Stristr(Category.Name, *Text) || Text.IsEmpty()))
+				{
+					OutSuggestions.Add(Category.Name);
+				}
+				return true;
+			});
 	}
 }
 
@@ -5802,24 +5857,25 @@ void STimingView::EnumerateFilteredTracks(TSharedPtr<UE::Insights::FFilterConfig
 	if (PriorityTrack.IsValid())
 	{
 		SkipId = PriorityTrack->GetId();
-		FilterContext.SetFilterData(static_cast<int32>(EFilterField::TrackName), PriorityTrack->GetName());
-		if (InFilterConfigurator->ApplyFilters(FilterContext))
+		if (PriorityTrack->IsVisible())
 		{
-			Callback(PriorityTrack);
+			FilterContext.SetFilterData(static_cast<int32>(EFilterField::TrackName), PriorityTrack->GetName());
+			if (InFilterConfigurator->ApplyFilters(FilterContext))
+			{
+				Callback(PriorityTrack);
+			}
 		}
 	}
 
 	for (auto& Entry : AllTracks)
 	{
-		if (Entry.Value->GetId() == SkipId)
+		if (Entry.Value->GetId() != SkipId && Entry.Value->IsVisible())
 		{
-			continue;
-		}
-
-		FilterContext.SetFilterData(static_cast<int32>(EFilterField::TrackName), Entry.Value->GetName());
-		if (InFilterConfigurator->ApplyFilters(FilterContext))
-		{
-			Callback(Entry.Value);
+			FilterContext.SetFilterData(static_cast<int32>(EFilterField::TrackName), Entry.Value->GetName());
+			if (InFilterConfigurator->ApplyFilters(FilterContext))
+			{
+				Callback(Entry.Value);
+			}
 		}
 	}
 }
diff --git a/Engine/Source/Developer/TraceInsights/Private/Insights/Widgets/STimingView.h b/Engine/Source/Developer/TraceInsights/Private/Insights/Widgets/STimingView.h
--- a/Engine/Source/Developer/TraceInsights/Private/Insights/Widgets/STimingView.h
+++ b/Engine/Source/Developer/TraceInsights/Private/Insights/Widgets/STimingView.h
@@ -49,6 +49,7 @@ namespace UE::Insights
 	class FFilterConfigurator;
 	enum class ETimingEventsColoringMode : uint32;
 	class SLogView;
+	class FFilterConfiguratorNode;
 }
 
 namespace UE::Insights::Timing
@@ -562,6 +563,7 @@ PRAGMA_ENABLE_DEPRECATION_WARNINGS
 	void FindMaxEvent();
 	void FindMinEvent();
 	void PostFindEvent(TSharedPtr<const ITimingEvent> BestMatchEvent);
+	TSharedPtr<const FBaseTimingTrack> FindPriorityTrack(FFilterConfiguratorNode& Node);
 	void FilterAllTracks();
 	void ClearFilters();
 
diff --git a/Engine/Source/Developer/TraceServices/Private/Model/Regions.cpp b/Engine/Source/Developer/TraceServices/Private/Model/Regions.cpp
--- a/Engine/Source/Developer/TraceServices/Private/Model/Regions.cpp
+++ b/Engine/Source/Developer/TraceServices/Private/Model/Regions.cpp
@@ -3,11 +3,13 @@
 #include "TraceServices/Model/Regions.h"
 #include "Model/RegionsPrivate.h"
 
-#include "AnalysisServicePrivate.h"
 #include "Algo/ForEach.h"
+#include "Internationalization/Internationalization.h"
+
+// TraceServices
+#include "AnalysisServicePrivate.h"
 #include "Common/FormatArgs.h"
 #include "Common/Utils.h"
-#include "Internationalization/Internationalization.h"
 
 #define LOCTEXT_NAMESPACE "RegionProvider"
 
@@ -16,24 +18,38 @@ namespace TraceServices
 
 thread_local FProviderLock::FThreadLocalState GRegionsProviderLockState;
 
-// ####################### FRegionTimeline ##########################
+////////////////////////////////////////////////////////////////////////////////////////////////////
+// FRegionTimeline
+////////////////////////////////////////////////////////////////////////////////////////////////////
 
-const TCHAR* FRegionTimeline::GetCategory() const
+const TCHAR* FRegionTimeline::GetCategoryName() const
+{
+	Provider.ReadAccessCheck();
+	return Category ? Category->Name : nullptr;
+}
+
+////////////////////////////////////////////////////////////////////////////////////////////////////
+
+const FRegionCategory* FRegionTimeline::GetCategory() const
 {
 	Provider.ReadAccessCheck();
 	return Category;
 }
-	
+
+////////////////////////////////////////////////////////////////////////////////////////////////////
+
 int32 FRegionTimeline::GetLaneCount() const
 {
 	Provider.ReadAccessCheck();
 	return Lanes.Num();
 }
 
+////////////////////////////////////////////////////////////////////////////////////////////////////
+
 const FRegionLane* FRegionTimeline::GetLane(int32 Index) const
 {
 	Provider.ReadAccessCheck();
-	
+
 	if (Index >= 0 && Index < Lanes.Num())
 	{
 		return &(Lanes[Index]);
@@ -41,26 +57,13 @@ const FRegionLane* FRegionTimeline::GetLane(int32 Index) const
 	return nullptr;
 }
 
-const bool FRegionTimeline::EnumerateUniqueRegionNames(TFunctionRef<bool(const TCHAR*)> Callback) const
-{
-	Provider.ReadAccessCheck();
-
-	for (const TCHAR* RegionName : TimelineUniqueRegions)
-	{
-		if (!Callback(RegionName))
-		{
-			return false;
-		}
-	}
-
-	return true;
-}
+////////////////////////////////////////////////////////////////////////////////////////////////////
 
 bool FRegionTimeline::EnumerateRegions(double IntervalStart, double IntervalEnd,
 	TFunctionRef<bool(const FTimeRegion&)> Callback) const
 {
 	Provider.ReadAccessCheck();
-	
+
 	if (IntervalStart > IntervalEnd)
 	{
 		return false;
@@ -70,16 +73,16 @@ bool FRegionTimeline::EnumerateRegions(double IntervalStart, double IntervalEnd,
 	for (const FRegionLane& Lane: Lanes)
 	{
 		const FInt32Interval OverlapRange = GetElementRangeOverlappingGivenRange<FTimeRegion>(Lane.Regions, IntervalStart, IntervalEnd,
-[](const FTimeRegion& r) { return r.BeginTime; },
-[](const FTimeRegion& r) { return r.EndTime; });
-		
+			[](const FTimeRegion& r) { return r.BeginTime; },
+			[](const FTimeRegion& r) { return r.EndTime; });
+
 		if (OverlapRange.Min != -1 && OverlapRange.Size() >= 0)
 		{
-			FLaneEnumeration EnumLane = {Lane, OverlapRange.Min, OverlapRange.Max};
+			FLaneEnumeration EnumLane = { Lane, OverlapRange.Min, OverlapRange.Max };
 			SearchLanes.Add(EnumLane);
 		}
 	}
-	
+
 	SearchLanes.Heapify([](const FLaneEnumeration& A, const FLaneEnumeration& B)
 	{
 		return A < B;
@@ -111,11 +114,13 @@ bool FRegionTimeline::EnumerateRegions(double IntervalStart, double IntervalEnd,
 	return true;
 }
 
+////////////////////////////////////////////////////////////////////////////////////////////////////
+
 bool FRegionTimeline::EnumerateRegionsBackwards(double IntervalEnd, double IntervalStart,
 	TFunctionRef<bool(const FTimeRegion&)> Callback) const
 {
 	Provider.ReadAccessCheck();
-	
+
 	if (IntervalStart > IntervalEnd)
 	{
 		return false;
@@ -125,16 +130,16 @@ bool FRegionTimeline::EnumerateRegionsBackwards(double IntervalEnd, double Inter
 	for (const FRegionLane& Lane: Lanes)
 	{
 		const FInt32Interval OverlapRange = GetElementRangeOverlappingGivenRange<FTimeRegion>(Lane.Regions, IntervalStart, IntervalEnd,
-[](const FTimeRegion& r) { return r.BeginTime; },
-[](const FTimeRegion& r) { return r.EndTime; });
-		
+			[](const FTimeRegion& r) { return r.BeginTime; },
+			[](const FTimeRegion& r) { return r.EndTime; });
+
 		if (OverlapRange.Min != -1 && OverlapRange.Size() >= 0)
 		{
-			FLaneEnumeration EnumLane = {Lane, OverlapRange.Min, OverlapRange.Max};
+			FLaneEnumeration EnumLane = { Lane, OverlapRange.Min, OverlapRange.Max };
 			SearchLanes.Add(EnumLane);
 		}
 	}
-	
+
 	SearchLanes.Heapify([](const FLaneEnumeration& A, const FLaneEnumeration& B)
 	{
 		return A > B;
@@ -166,34 +171,40 @@ bool FRegionTimeline::EnumerateRegionsBackwards(double IntervalEnd, double Inter
 	return true;
 }
 
+////////////////////////////////////////////////////////////////////////////////////////////////////
+
 void FRegionTimeline::EnumerateLanes(TFunctionRef<void(const FRegionLane&, const int32)> Callback) const
 {
 	Provider.ReadAccessCheck();
-	
+
 	for (int32 LaneIndex = 0; LaneIndex < Lanes.Num(); ++LaneIndex)
 	{
 		Callback(Lanes[LaneIndex], LaneIndex);
 	}
 }
 
-FTimeRegion* FRegionTimeline::InsertNewRegion(double BeginTime, const TCHAR* Name, uint64 Id, const TCHAR* RegionCategory)
+////////////////////////////////////////////////////////////////////////////////////////////////////
+
+FTimeRegion* FRegionTimeline::InsertNewRegion(const FRegionTimerImpl& Timer, double BeginTime, uint64 Id)
 {
 	Provider.EditAccessCheck();
-	
+
 	FTimeRegion Region;
+	Region.Timer = &Timer;
 	Region.BeginTime = BeginTime;
-	Region.Text = Name;
 	Region.Id = Id;
-	Region.Category = RegionCategory;
+	PRAGMA_DISABLE_DEPRECATION_WARNINGS
+	Region.Text = Timer.Name;
+	Region.Category = Timer.Category ? Timer.Category->Name : nullptr;
+	PRAGMA_ENABLE_DEPRECATION_WARNINGS
 	Region.Depth = CalculateRegionDepth(Region);
 
 	if (Region.Depth == Lanes.Num())
 	{
 		Lanes.Emplace(Session.GetLinearAllocator());
 	}
-	FRegionLane& NewLane = Lanes[Region.Depth];	
+	FRegionLane& NewLane = Lanes[Region.Depth];
 	NewLane.Regions.EmplaceBack(Region);
-	TimelineUniqueRegions.FindOrAdd(Region.Text);
 
 	FTimeRegion* NewOpenRegion = &(Lanes[Region.Depth].Regions.Last());
 
@@ -212,7 +223,7 @@ int32 FRegionTimeline::CalculateRegionDepth(const FTimeRegion& Region) const
 		{
 			break;
 		}
-		
+
 		const FTimeRegion& LastRegion = Lanes[NewDepth].Regions.Last();
 		if (LastRegion.EndTime <= Region.BeginTime)
 		{
@@ -224,123 +235,253 @@ int32 FRegionTimeline::CalculateRegionDepth(const FTimeRegion& Region) const
 	ensureMsgf(NewDepth < DepthLimit, TEXT("Regions are nested too deep."));
 
 	return NewDepth;
-}	
-	
-// ##################### RegionProvider #####################	
+}
+
+////////////////////////////////////////////////////////////////////////////////////////////////////
+// FRegionProvider
+////////////////////////////////////////////////////////////////////////////////////////////////////
+
 
 FRegionProvider::FRegionProvider(IAnalysisSession& InSession)
-	: Session(InSession), AllRegions(*this, InSession, nullptr)
+	: Session(InSession)
+	, Categories(InSession.GetLinearAllocator(), 1024)
+	, Timers(InSession.GetLinearAllocator(), 1024)
+	, AllRegionsTimeline(*this, InSession, nullptr)
+{
+}
+
+////////////////////////////////////////////////////////////////////////////////////////////////////
+
+bool FRegionProvider::EnumerateCategories(TFunctionRef<bool(const FRegionCategory&)> Callback) const
+{
+	ReadAccessCheck();
+
+	auto CategoryIterator = Categories.GetIterator();
+	const FRegionCategory* Category = CategoryIterator.GetCurrentItem();
+	while (Category)
+	{
+		if (!Callback(*Category))
+		{
+			return false;
+		}
+		Category = CategoryIterator.NextItem();
+	}
+
+	return true;
+}
+
+////////////////////////////////////////////////////////////////////////////////////////////////////
+
+bool FRegionProvider::EnumerateTimers(TFunctionRef<bool(const FRegionTimer&)> Callback) const
 {
+	ReadAccessCheck();
+
+	auto TimerIterator = Timers.GetIterator();
+	const FRegionTimer* Timer = TimerIterator.GetCurrentItem();
+	while (Timer)
+	{
+		if (!Callback(*Timer))
+		{
+			return false;
+		}
+		Timer = TimerIterator.NextItem();
+	}
+
+	return true;
 }
 
-void FRegionProvider::EnumerateTimelinesByCategory(
-	TFunctionRef<void(const IRegionTimeline&, const TCHAR*)> Callback) const
+////////////////////////////////////////////////////////////////////////////////////////////////////
+
+void FRegionProvider::EnumerateTimelinesByCategory(TFunctionRef<void(const IRegionTimeline&, const TCHAR*)> Callback) const
 {
 	ReadAccessCheck();
-	
-	for (auto& KV: RegionsPerCategory)
+
+	auto CategoryIterator = Categories.GetIterator();
+	const FRegionCategoryImpl* Category = CategoryIterator.GetCurrentItem();
+	while (Category)
 	{
-		Callback(KV.Value, KV.Key);
+		Callback(Category->Timeline, Category->Name);
+		Category = CategoryIterator.NextItem();
 	}
 }
 
+////////////////////////////////////////////////////////////////////////////////////////////////////
+
 const IRegionTimeline& FRegionProvider::GetDefaultTimeline() const
 {
 	ReadAccessCheck();
-	
-	return AllRegions;
+
+	return AllRegionsTimeline;
 }
 
-const IRegionTimeline* FRegionProvider::GetTimelineForCategory(const TCHAR* Category) const
+////////////////////////////////////////////////////////////////////////////////////////////////////
+
+const IRegionTimeline* FRegionProvider::GetTimelineForCategory(const TCHAR* CategoryName) const
 {
 	ReadAccessCheck();
 
-	if (!Category)
+	if (!CategoryName)
 	{
-		return &AllRegions;
+		return &AllRegionsTimeline;
 	}
-	
-	if (const FRegionTimeline* Timeline = RegionsPerCategory.Find(Category))
+
+	FRegionCategoryImpl* const* FoundCategory = CategoriesByName.Find(CategoryName);
+	if (FoundCategory)
 	{
-		return Timeline;
+		return &(*FoundCategory)->Timeline;
 	}
+
 	return nullptr;
 }
-	
+
+////////////////////////////////////////////////////////////////////////////////////////////////////
+
 uint64 FRegionProvider::GetRegionCount() const
 {
 	ReadAccessCheck();
 
 	uint64 RegionCount = 0;
-	AllRegions.EnumerateLanes([&RegionCount](const FRegionLane& Lane, const int32 Index) {
-		RegionCount += Lane.Num();
-	});
+	AllRegionsTimeline.EnumerateLanes([&RegionCount](const FRegionLane& Lane, const int32 Index) { RegionCount += Lane.Num(); });
 	return RegionCount;
 }
 
+////////////////////////////////////////////////////////////////////////////////////////////////////
+
 const FRegionLane* FRegionProvider::GetLane(int32 index) const
 {
 	ReadAccessCheck();
-	return AllRegions.GetLane(index);
+
+	return AllRegionsTimeline.GetLane(index);
+}
+
+////////////////////////////////////////////////////////////////////////////////////////////////////
+
+FRegionCategoryImpl& FRegionProvider::GetOrAddRegionCategory(const TCHAR* CategoryName)
+{
+	check(CategoryName);
+
+	FRegionCategoryImpl** FoundCategory = CategoriesByName.Find(CategoryName);
+	if (FoundCategory)
+	{
+		return **FoundCategory;
+	}
+
+	FRegionCategoryImpl& Category = Categories.EmplaceBack(*this, Session);
+	Category.Name = CategoryName;
+	CategoriesByName.Add(CategoryName, &Category);
+
+	return Category;
+}
+
+////////////////////////////////////////////////////////////////////////////////////////////////////
+
+FRegionTimerImpl& FRegionProvider::GetOrAddRegionTimer(const TCHAR* TimerName, const FRegionCategory* Category)
+{
+	check(TimerName);
+
+	FRegionTimerImpl** FoundTimer = TimersByName.Find(TimerName);
+	if (FoundTimer)
+	{
+		return **FoundTimer;
+	}
+
+	FRegionTimerImpl& Timer = Timers.EmplaceBack();
+	Timer.Name = TimerName;
+	Timer.Category = Category;
+	TimersByName.Add(TimerName, &Timer);
+
+	return Timer;
+}
+
+////////////////////////////////////////////////////////////////////////////////////////////////////
+
+FRegionProvider::FOpenRegions FRegionProvider::AddRegionToTimelines(double BeginTime, const TCHAR* Name, uint64 Id, const TCHAR* CategoryName)
+{
+	const TCHAR* StoredCategoryName = CategoryName ? Session.StoreString(CategoryName) : UncategorizedName;
+	FRegionCategoryImpl& Category = GetOrAddRegionCategory(StoredCategoryName);
+
+	const TCHAR* StoredName = Session.StoreString(Name);
+	FRegionTimerImpl& Timer = GetOrAddRegionTimer(StoredName, &Category);
+
+	FOpenRegions OpenRegions;
+
+	// Insert a region in the AllRegions timeline.
+	OpenRegions.Region = AllRegionsTimeline.InsertNewRegion(Timer, BeginTime, Id);
+	check(OpenRegions.Region);
+
+	// Also insert a region in the category timeline.
+	OpenRegions.CategoryRegion = Category.Timeline.InsertNewRegion(Timer, BeginTime, Id);
+	check(OpenRegions.CategoryRegion);
+
+	return OpenRegions;
 }
 
-void FRegionProvider::AppendRegionBegin(const TCHAR* Name, double Time, const TCHAR* Category)
+////////////////////////////////////////////////////////////////////////////////////////////////////
+
+void FRegionProvider::AppendRegionBegin(const TCHAR* Name, double Time, const TCHAR* CategoryName)
 {
 	EditAccessCheck();
 
-	check(Name)
+	check(Name);
 
-	TArray<FTimeRegion*>* OpenRegions = OpenRegionsByName.Find(Name);
+	FOpenRegions* FoundOpenRegions = OpenRegionsByName.Find(Name);
 
-	if (OpenRegions)
+	if (FoundOpenRegions)
 	{
 		++NumWarnings;
 		if (NumWarnings <= MaxWarningMessages)
 		{
-			UE_LOG(LogTraceServices, Warning, TEXT("[Regions] A region begin event (BeginTime=%f, Name=\"%s\") was encountered while a region with same name is already open."), Time, Name)
+			UE_LOG(LogTraceServices, Warning, TEXT("[Regions] A region begin event (BeginTime=%f, Name=\"%s\", Category=\"%s\") was encountered while a region with same name is already open."),
+				Time, Name, CategoryName ? CategoryName : TEXT(""))
 		}
 
 		// Automatically end the previous region.
 		AppendRegionEnd(Name, Time);
 	}
 
-	TArray<FTimeRegion*> NewRegions = InsertNewRegion(Time, Name, 0, Category);
-	OpenRegionsByName.Add(NewRegions[0]->Text, NewRegions);
+	FOpenRegions OpenRegions = AddRegionToTimelines(Time, Name, 0, CategoryName);
+	OpenRegionsByName.Add(OpenRegions.Region->Timer->Name, OpenRegions);
 
 	UpdateCounter++;
 }
 
-void FRegionProvider::AppendRegionBeginWithId(const TCHAR* Name, uint64 Id, double Time, const TCHAR* Category)
+////////////////////////////////////////////////////////////////////////////////////////////////////
+
+void FRegionProvider::AppendRegionBeginWithId(const TCHAR* Name, uint64 Id, double Time, const TCHAR* CategoryName)
 {
 	EditAccessCheck();
 
-	check(Name && Id)
-	TArray<FTimeRegion*>* OpenRegions = OpenRegionsById.Find(Id);
+	check(Name && Id);
+
+	FOpenRegions* FoundOpenRegions = OpenRegionsById.Find(Id);
 
-	if (OpenRegions)
+	if (FoundOpenRegions)
 	{
 		++NumWarnings;
 		if (NumWarnings <= MaxWarningMessages)
 		{
-			UE_LOG(LogTraceServices, Warning, TEXT("[Regions] A region begin event (BeginTime=%f, Name=\"%s\", Id=%llu) was encountered while a region with same name is already open."), Time, Name, Id)
+			UE_LOG(LogTraceServices, Warning, TEXT("[Regions] A region begin event (BeginTime=%f, Id=%llu, Name=\"%s\", Category=\"%s\") was encountered while a region with same name is already open."),
+				Time, Id, Name, CategoryName ? CategoryName : TEXT(""))
 		}
 
 		// Automatically end the previous region.
 		AppendRegionEndWithId(Id, Time);
 	}
 
-	TArray<FTimeRegion*> NewRegions = InsertNewRegion(Time, Name, Id, Category);
-	OpenRegionsById.Add(Id, NewRegions);
-	
+	FOpenRegions OpenRegions = AddRegionToTimelines(Time, Name, Id, CategoryName);
+	OpenRegionsById.Add(Id, OpenRegions);
+
 	UpdateCounter++;
 }
 
+////////////////////////////////////////////////////////////////////////////////////////////////////
+
 void FRegionProvider::AppendRegionEnd(const TCHAR* Name, double Time)
 {
 	EditAccessCheck();
 
 	check(Name)
-	TArray<FTimeRegion*>* OpenRegionsPos = OpenRegionsByName.Find(Name);
+	FOpenRegions* OpenRegionsPos = OpenRegionsByName.Find(Name);
 
 	if (!OpenRegionsPos)
 	{
@@ -355,17 +496,19 @@ void FRegionProvider::AppendRegionEnd(const TCHAR* Name, double Time)
 		check(OpenRegionsPos);
 	}
 
-	for (FTimeRegion* OpenRegion: *OpenRegionsPos )
-	{
-		check(OpenRegion);
-		OpenRegion->EndTime = Time;
-	}	
+	check(OpenRegionsPos->Region);
+	OpenRegionsPos->Region->EndTime = Time;
+
+	check(OpenRegionsPos->CategoryRegion);
+	OpenRegionsPos->CategoryRegion->EndTime = Time;
 
 	OpenRegionsByName.Remove(Name);
-	
+
 	UpdateCounter++;
 }
 
+////////////////////////////////////////////////////////////////////////////////////////////////////
+
 void FRegionProvider::AppendRegionEndWithId(uint64 Id, double Time)
 {
 	EditAccessCheck();
@@ -379,7 +522,8 @@ void FRegionProvider::AppendRegionEndWithId(uint64 Id, double Time)
 		}
 		return;
 	}
-	TArray<FTimeRegion*>* OpenRegionsPos = OpenRegionsById.Find(Id);
+
+	FOpenRegions* OpenRegionsPos = OpenRegionsById.Find(Id);
 
 	if (!OpenRegionsPos)
 	{
@@ -397,34 +541,40 @@ void FRegionProvider::AppendRegionEndWithId(uint64 Id, double Time)
 		check(OpenRegionsPos);
 	}
 
-	for (FTimeRegion* OpenRegion: *OpenRegionsPos )
-	{
-		check(OpenRegion);
-		OpenRegion->EndTime = Time;
-	}
+	check(OpenRegionsPos->Region);
+	OpenRegionsPos->Region->EndTime = Time;
+
+	check(OpenRegionsPos->CategoryRegion);
+	OpenRegionsPos->CategoryRegion->EndTime = Time;
 
 	OpenRegionsById.Remove(Id);
-	
+
 	UpdateCounter++;
 }
 
+////////////////////////////////////////////////////////////////////////////////////////////////////
+
 void FRegionProvider::OnAnalysisSessionEnded()
 {
 	EditAccessCheck();
 
-	auto printOpenRegionMessage = [this](const auto& KV)
+	auto PrintOpenRegionMessage = [this](const auto& KV)
 	{
-		const TArray<FTimeRegion*> Regions = KV.Value;
-		const auto Region = Regions[0];
-		
 		++NumWarnings;
 		if (NumWarnings <= MaxWarningMessages)
 		{
-			UE_LOG(LogTraceServices, Warning, TEXT("[Regions] A region (BeginTime=%f, Name=\"%s\", Id=%llu) was never closed."), Region->BeginTime, Region->Text, Region->Id)
+			check(KV.Value.Region);
+			const FTimeRegion& Region = *KV.Value.Region;
+			check(Region.Timer && Region.Timer->Name);
+			UE_LOG(LogTraceServices, Warning, TEXT("[Regions] A region (BeginTime=%f, Id=%llu, Name=\"%s\", Category=\"%s\") was never closed."),
+				Region.BeginTime,
+				Region.Id,
+				Region.Timer->Name,
+				Region.Timer->Category && Region.Timer->Category->Name ? Region.Timer->Category->Name : TEXT(""));
 		}
 	};
-	Algo::ForEach(OpenRegionsById, printOpenRegionMessage);
-	Algo::ForEach(OpenRegionsByName, printOpenRegionMessage);
+	Algo::ForEach(OpenRegionsById, PrintOpenRegionMessage);
+	Algo::ForEach(OpenRegionsByName, PrintOpenRegionMessage);
 
 	if (NumWarnings > 0)
 	{
@@ -435,43 +585,30 @@ void FRegionProvider::OnAnalysisSessionEnded()
 		UE_LOG(LogTraceServices, Error, TEXT("[Regions] %u errors"), NumErrors);
 	}
 
-	uint64 TotalRegionCount = GetRegionCount();
-	UE_LOG(LogTraceServices, Log, TEXT("[Regions] Analysis completed (%llu regions, %d lanes)."), TotalRegionCount, AllRegions.Lanes.Num());
+	const uint64 TotalRegionCount = GetRegionCount();
+	UE_LOG(LogTraceServices, Log, TEXT("[Regions] Analysis completed (%llu timers, %llu categories, %llu timing regions, %d lanes)."),
+		Timers.Num(), Categories.Num(), TotalRegionCount, AllRegionsTimeline.Lanes.Num());
 }
 
-
-TArray<FTimeRegion*> FRegionProvider::InsertNewRegion(double BeginTime, const TCHAR* Name, uint64 Id,
-                                                      const TCHAR* Category)
-{
-	const TCHAR* StoredName = Session.StoreString(Name);
-	const TCHAR* StoredCategory = Category ? Session.StoreString(Category) : UncategorizedName;
-	
-	TArray<FTimeRegion*> Regions;
-	Regions.Add(AllRegions.InsertNewRegion(BeginTime, StoredName, Id, StoredCategory));
-	
-	FRegionTimeline* Timeline = RegionsPerCategory.Find(StoredCategory);
-	if (!Timeline)
-	{
-		FRegionTimeline NewTimeline{*this, Session, StoredCategory};
-		Timeline = &RegionsPerCategory.Emplace(StoredCategory, NewTimeline);
-	}
-	Regions.Add(Timeline->InsertNewRegion(BeginTime, StoredName, Id, StoredCategory));
-	return Regions;
-}
+////////////////////////////////////////////////////////////////////////////////////////////////////
 
 void FRegionProvider::EnumerateLanes(TFunctionRef<void(const FRegionLane&, int32)> Callback) const
 {
 	ReadAccessCheck();
-	AllRegions.EnumerateLanes(Callback);
+	AllRegionsTimeline.EnumerateLanes(Callback);
 }
 
+////////////////////////////////////////////////////////////////////////////////////////////////////
+
 bool FRegionProvider::EnumerateRegions(double IntervalStart, double IntervalEnd, TFunctionRef<bool(const FTimeRegion&)> Callback) const
 {
 	ReadAccessCheck();
-	return AllRegions.EnumerateRegions(IntervalStart, IntervalEnd, Callback);
+	return AllRegionsTimeline.EnumerateRegions(IntervalStart, IntervalEnd, Callback);
 }
 
-// ###################### FRegionLane #########################
+////////////////////////////////////////////////////////////////////////////////////////////////////
+// FRegionLane
+////////////////////////////////////////////////////////////////////////////////////////////////////
 
 bool FRegionLane::EnumerateRegions(double IntervalStart, double IntervalEnd, TFunctionRef<bool(const FTimeRegion&)> Callback) const
 {
@@ -495,6 +632,8 @@ bool FRegionLane::EnumerateRegions(double IntervalStart, double IntervalEnd, TFu
 	return true;
 }
 
+////////////////////////////////////////////////////////////////////////////////////////////////////
+
 bool FRegionLane::EnumerateRegionsBackwards(double IntervalEnd, double IntervalStart, TFunctionRef<bool(const FTimeRegion&)> Callback) const
 {
 	const FInt32Interval OverlapRange = GetElementRangeOverlappingGivenRange<FTimeRegion>(Regions, IntervalStart, IntervalEnd,
@@ -517,21 +656,31 @@ bool FRegionLane::EnumerateRegionsBackwards(double IntervalEnd, double IntervalS
 	return true;
 }
 
+////////////////////////////////////////////////////////////////////////////////////////////////////
+////////////////////////////////////////////////////////////////////////////////////////////////////
+
 FName GetRegionProviderName()
 {
 	static const FName Name("RegionProvider");
 	return Name;
 }
+
+////////////////////////////////////////////////////////////////////////////////////////////////////
+
 const IRegionProvider& ReadRegionProvider(const IAnalysisSession& Session)
 {
 	return *Session.ReadProvider<IRegionProvider>(GetRegionProviderName());
 }
 
+////////////////////////////////////////////////////////////////////////////////////////////////////
+
 IEditableRegionProvider& EditRegionProvider(IAnalysisSession& Session)
 {
 	return *Session.EditProvider<IEditableRegionProvider>(GetRegionProviderName());
 }
 
+////////////////////////////////////////////////////////////////////////////////////////////////////
+
 } // namespace TraceServices
 
 #undef LOCTEXT_NAMESPACE
diff --git a/Engine/Source/Developer/TraceServices/Private/Model/RegionsPrivate.h b/Engine/Source/Developer/TraceServices/Private/Model/RegionsPrivate.h
--- a/Engine/Source/Developer/TraceServices/Private/Model/RegionsPrivate.h
+++ b/Engine/Source/Developer/TraceServices/Private/Model/RegionsPrivate.h
@@ -2,99 +2,120 @@
 
 #pragma once
 
+#include "CoreTypes.h"
+#include "Templates/SharedPointer.h"
+
+// TraceServices
 #include "Common/PagedArray.h"
 #include "Common/ProviderLock.h"
 #include "TraceServices/Model/Regions.h"
-#include "Templates/SharedPointer.h"
-
-#define UE_API TRACESERVICES_API
 
 namespace TraceServices
 {
+
 class FRegionProvider;
-	
-extern thread_local FProviderLock::FThreadLocalState GRegionsProviderLockState;
+struct FRegionCategoryImpl;
+struct FRegionTimerImpl;
 
-class FAnalysisSessionLock;
-class FStringStore;
+extern thread_local FProviderLock::FThreadLocalState GRegionsProviderLockState;
 
 /**
- *	A View onto a set of Timing Regions, sorted into a stack of individual lanes without overlaps for display,
- *	optionally filtered by Category
+ * A View onto a set of Timing Regions, sorted into a stack of individual lanes without overlaps for display,
+ * optionally filtered by Category
  */
 class FRegionTimeline : public IRegionTimeline
 {
 public:
 	virtual ~FRegionTimeline() = default;
-	
-	FRegionTimeline(FRegionProvider& InProvider, IAnalysisSession& InSession, const TCHAR* Category) : Category(Category), Session(InSession), Provider(InProvider) {}
+
+	FRegionTimeline(FRegionProvider& InProvider, IAnalysisSession& InSession, const FRegionCategoryImpl* InCategory)
+		: Provider(InProvider)
+		, Session(InSession)
+		, Category(InCategory)
+	{}
+
+	/**
+	 * @return the category name of this timeline if filtered or nullptr
+	 */
+	const TCHAR* GetCategoryName() const override;
 
 	/**
 	 * @return the category of this timeline if filtered or nullptr
 	 */
-	UE_API const TCHAR* GetCategory() const override;
+	const FRegionCategory* GetCategory() const override;
 
 	/**
 	 * @return the number of lanes
 	 */
-	UE_API int32 GetLaneCount() const override;
+	int32 GetLaneCount() const override;
 
 	/**
 	 * Direct access to a certain lane at a given index/depth.
 	 * The pointer is valid only in the current read scope.
 	 * @return a pointer to the lane at the specified depth index or nullptr if Index is invalid
 	 */
-	UE_API const FRegionLane* GetLane(int32 Index) const override;
+	const FRegionLane* GetLane(int32 Index) const override;
 
-	/**
-	 * Enumerates all region names for a given timeline. Considers all lanes.
-	 * 
-	 * @param Callback a callback called for each region. Return false to abort iteration.
-	 * @returns true if the enumeration finished, false if it was aborted by the callback returning false
-	 */
-	UE_API const bool EnumerateUniqueRegionNames(TFunctionRef<bool(const TCHAR*)> Callback) const override;
-	
 	/**
 	 * Enumerates all regions that overlap a certain time interval. Will enumerate by depth but does not expose lanes.
-	 * 
+	 *
 	 * @param Callback a callback called for each region. Return false to abort iteration.
 	 * @returns true if the enumeration finished, false if it was aborted by the callback returning false
 	 */
-	UE_API bool EnumerateRegions(double IntervalStart, double IntervalEnd, TFunctionRef<bool(const FTimeRegion& /*Region*/)> Callback) const override;
+	bool EnumerateRegions(double IntervalStart, double IntervalEnd, TFunctionRef<bool(const FTimeRegion& /*Region*/)> Callback) const override;
 
 	/**
 	 * Enumerates all regions that overlap a certain time interval in reverse. Will enumerate by depth but does not expose lanes.
-	 * 
+	 *
 	 * @param Callback a callback called for each region. Return false to abort iteration.
 	 * @returns true if the enumeration finished, false if it was aborted by the callback returning false
 	 */
-	UE_API bool EnumerateRegionsBackwards(double IntervalEnd, double IntervalStart, TFunctionRef<bool(const FTimeRegion& /*Region*/)> Callback) const override;
+	bool EnumerateRegionsBackwards(double IntervalEnd, double IntervalStart, TFunctionRef<bool(const FTimeRegion& /*Region*/)> Callback) const override;
 
 	/**
 	 * Will call Callback(Lane, Depth) for each lane in order.
 	 */
-	UE_API void EnumerateLanes(TFunctionRef<void(const FRegionLane& /*Lane*/, const int32 /*Depth*/)> Callback) const override;
+	void EnumerateLanes(TFunctionRef<void(const FRegionLane& /*Lane*/, const int32 /*Depth*/)> Callback) const override;
 
 private:
 	friend class FRegionProvider;
-	/*
-	 * Inserts a new Region into the stack, returns the created Region to allow the update of the region end time when it closes
+
+	/**
+	 * Inserts a new Region into the stack, returns the created Region to allow the update of the region end time when it closes.
 	 */
-	UE_API FTimeRegion*  InsertNewRegion(double BeginTime, const TCHAR* Name, uint64 Id, const TCHAR* Category);
+	FTimeRegion* InsertNewRegion(const FRegionTimerImpl& Timer, double BeginTime, uint64 Id);
+
+	/**
+	 * Update the depth member of a region to allow overlapping regions to be displayed on separate lanes.
+	 */
+	int32 CalculateRegionDepth(const FTimeRegion& Item) const;
 
 private:
-	/// Update the depth member of a region to allow overlapping regions to be displayed on separate lanes.
-	UE_API int32 CalculateRegionDepth(const FTimeRegion& Item) const;
-	
+	// The parent provider. Required to engage the ReadAccessLock or WriteAccessLock when a Lane is modified.
+	TraceServices::FRegionProvider& Provider;
+
+	TraceServices::IAnalysisSession& Session;
+
+	const FRegionCategoryImpl* Category = nullptr;
+
 	// Closed regions
 	TArray<FRegionLane> Lanes;
-	TSet<const TCHAR*, TStringPointerSetKeyFuncs_DEPRECATED<const TCHAR*>> TimelineUniqueRegions;
-	const TCHAR* Category = nullptr;
-	TraceServices::IAnalysisSession& Session;
-	// the parent provider. Required to engage the ReadAccessLock or WriteAccessLock when a Lane is modified
-	TraceServices::FRegionProvider& Provider;
 };
-	
+
+struct FRegionCategoryImpl : public FRegionCategory
+{
+	FRegionTimeline Timeline;
+
+	FRegionCategoryImpl(FRegionProvider& InProvider, IAnalysisSession& InSession)
+		: Timeline(InProvider, InSession, this)
+	{
+	}
+};
+
+struct FRegionTimerImpl : public FRegionTimer
+{
+};
+
 class FRegionProvider
 	: public IRegionProvider
 	, public IEditableRegionProvider
@@ -109,14 +130,17 @@ public:
 	virtual void BeginRead() const override       { Lock.BeginRead(GRegionsProviderLockState); }
 	virtual void EndRead() const override         { Lock.EndRead(GRegionsProviderLockState); }
 	virtual void ReadAccessCheck() const override { Lock.ReadAccessCheck(GRegionsProviderLockState); }
-	
+
+	virtual bool EnumerateCategories(TFunctionRef<bool(const FRegionCategory&)> Callback) const override;
+	virtual bool EnumerateTimers(TFunctionRef<bool(const FRegionTimer&)> Callback) const override;
+
 	virtual void EnumerateTimelinesByCategory(TFunctionRef<void(const IRegionTimeline& /*Timeline*/, const TCHAR* /*Category*/)> Callback) const override;
 	virtual const IRegionTimeline& GetDefaultTimeline() const override;
 	virtual const IRegionTimeline* GetTimelineForCategory(const TCHAR* Category) const override;
 	virtual const TCHAR* GetUncategorizedRegionCategoryName() const override {return UncategorizedName;};
 
 	virtual uint64 GetRegionCount() const override;
-	virtual int32 GetLaneCount()  const override { ReadAccessCheck(); return AllRegions.GetLaneCount(); }
+	virtual int32 GetLaneCount()  const override { ReadAccessCheck(); return AllRegionsTimeline.GetLaneCount(); }
 
 	virtual const FRegionLane* GetLane(int32 Index) const override;
 
@@ -125,7 +149,7 @@ public:
 	virtual void EnumerateLanes(TFunctionRef<void(const FRegionLane&, const int32)> Callback) const override;
 
 	virtual uint64 GetUpdateCounter() const override { ReadAccessCheck(); return UpdateCounter; }
-	
+
 	//////////////////////////////////////////////////
 	// Edit operations
 
@@ -133,6 +157,17 @@ public:
 	virtual void EndEdit() const override         { Lock.EndWrite(GRegionsProviderLockState); }
 	virtual void EditAccessCheck() const override { Lock.WriteAccessCheck(GRegionsProviderLockState); }
 
+	virtual const FRegionCategory* RegisterRegionCategory(const TCHAR* Name) override
+	{
+		EditAccessCheck();
+		return &GetOrAddRegionCategory(Name);
+	}
+	virtual const FRegionTimer* RegisterRegionTimer(const TCHAR* Name, const FRegionCategory* Category = nullptr) override
+	{
+		EditAccessCheck();
+		return &GetOrAddRegionTimer(Name, Category);
+	}
+
 	virtual void AppendRegionBegin(const TCHAR* Name, double Time, const TCHAR* Category = nullptr) override;
 	virtual void AppendRegionBeginWithId(const TCHAR* Name, uint64 Id, double Time, const TCHAR* Category = nullptr) override;
 	virtual void AppendRegionEnd(const TCHAR* Name, double Time) override;
@@ -141,26 +176,39 @@ public:
 	virtual void OnAnalysisSessionEnded() override;
 
 	//////////////////////////////////////////////////
-	
+
 private:
+	FRegionCategoryImpl& GetOrAddRegionCategory(const TCHAR* CategoryName);
+	FRegionTimerImpl& GetOrAddRegionTimer(const TCHAR* TimerName, const FRegionCategory* Category = nullptr);
+
+	struct FOpenRegions
+	{
+		FTimeRegion* Region;
+		FTimeRegion* CategoryRegion;
+	};
 
-	/// Calculates depth, inserts a new region into the correct lane and updates session time.
-	TArray<FTimeRegion*> InsertNewRegion(double BeginTime, const TCHAR* Name, uint64 Id, const TCHAR* Category);
+	// Adds a new region to the AllRegions timeline and to the category timeline (which can be the uncategorized timeline).
+	// @return The two added regions.
+	FOpenRegions AddRegionToTimelines(double BeginTime, const TCHAR* Name, uint64 Id, const TCHAR* Category);
 
 private:
 	mutable FProviderLock Lock;
 
 	IAnalysisSession& Session;
 
+	TPagedArray<FRegionCategoryImpl> Categories;
+	TPagedArray<FRegionTimerImpl> Timers;
+
+	TMap<FStringView, FRegionCategoryImpl*> CategoriesByName;
+	TMap<FStringView, FRegionTimerImpl*> TimersByName;
+
 	// Open regions inside lanes
-	TMap<FStringView, TArray<FTimeRegion*>> OpenRegionsByName;
-	TMap<uint64, TArray<FTimeRegion*>> OpenRegionsById;
-	
-	FRegionTimeline AllRegions;
-	/// maps CategoryName to an FRegionTimeline.
-	/// Also contains a special entry with key UncategorizedName for any regions without a category specified
-	TMap<const TCHAR*, FRegionTimeline, FDefaultSetAllocator, TStringPointerMapKeyFuncs_DEPRECATED<const TCHAR*, FRegionTimeline>> RegionsPerCategory;
-	// string maps don't support nullptr, so use this key internally for all regions where category is nullptr
+	TMap<FStringView, FOpenRegions> OpenRegionsByName;
+	TMap<uint64, FOpenRegions> OpenRegionsById;
+
+	FRegionTimeline AllRegionsTimeline;
+
+	// String maps don't support nullptr, so use this key internally for all regions where category is nullptr.
 	const TCHAR* UncategorizedName = TEXT("Uncategorized");
 
 	// Counter incremented each time region data changes during analysis
@@ -174,5 +222,3 @@ private:
 };
 
 } // namespace TraceServices
-
-#undef UE_API
diff --git a/Engine/Source/Developer/TraceServices/Public/TraceServices/Model/Regions.h b/Engine/Source/Developer/TraceServices/Public/TraceServices/Model/Regions.h
--- a/Engine/Source/Developer/TraceServices/Public/TraceServices/Model/Regions.h
+++ b/Engine/Source/Developer/TraceServices/Public/TraceServices/Model/Regions.h
@@ -2,28 +2,49 @@
 
 #pragma once
 
-#include "Common/PagedArray.h" // TraceServices
+#include "CoreTypes.h"
 #include "HAL/Platform.h"
 #include "Templates/Function.h"
-#include "TraceServices/Model/AnalysisSession.h"
 #include "UObject/NameTypes.h"
 
+// TraceServices
+#include "Common/PagedArray.h"
+#include "TraceServices/Model/AnalysisSession.h"
+
 #define UE_API TRACESERVICES_API
 
 namespace TraceServices
 {
 
+struct FRegionCategory
+{
+	const TCHAR* Name = nullptr;
+};
+
+struct FRegionTimer
+{
+	const TCHAR* Name = nullptr;
+	const FRegionCategory* Category = nullptr;
+};
+
+PRAGMA_DISABLE_DEPRECATION_WARNINGS
 struct FTimeRegion
 {
+	const FRegionTimer* Timer = nullptr; // name and category of the region
+
 	double BeginTime = std::numeric_limits<double>::infinity();
 	double EndTime = std::numeric_limits<double>::infinity();
-	const TCHAR* Text = nullptr;
-	/// Id will be zero if the region is identified by Name only
-	uint64 Id = 0;
+
+	uint64 Id = 0; // 0 if the region is identified by Name only
 	int32 Depth = -1;
-	/// this will be nullptr for regions without a Category
-	const TCHAR* Category = nullptr;
+
+	UE_DEPRECATED(5.7, "Use instead Timer->Name")
+	const TCHAR* Text = nullptr;
+
+	UE_DEPRECATED(5.7, "Use instead Timer->Category->Name")
+	const TCHAR* Category = nullptr; // this will be nullptr for regions without a Category
 };
+PRAGMA_ENABLE_DEPRECATION_WARNINGS
 
 class FRegionLane
 {
@@ -69,18 +90,23 @@ struct FLaneEnumeration
 		return Lane.Regions[IntervalStart].BeginTime > rhs.Lane.Regions[rhs.IntervalStart].BeginTime;
 	}
 };
-	
+
 /*
  *	Sorts a set of Timing Regions into a stack of individual lanes without overlaps for display
  */
 class IRegionTimeline
 {
 public:
+	/**
+	 * @return the category name of this timeline if filtered or nullptr
+	 */
+	virtual const TCHAR* GetCategoryName() const = 0;
+
 	/**
 	 * @return the category of this timeline if filtered or nullptr
 	 */
-	virtual const TCHAR* GetCategory() const = 0;
-	
+	virtual const FRegionCategory* GetCategory() const = 0;
+
 	/**
 	 * @return the number of lanes
 	 */
@@ -93,17 +119,9 @@ public:
 	 */
 	virtual const FRegionLane* GetLane(int32 Index) const = 0;
 
-	/**
-	 * Enumerates all region names for a given timeline. Considers all lanes.
-	 * 
-	 * @param Callback a callback called for each region. Return false to abort iteration.
-	 * @returns true if the enumeration finished, false if it was aborted by the callback returning false
-	 */
-	virtual const bool EnumerateUniqueRegionNames(TFunctionRef<bool(const TCHAR*)> Callback) const = 0;
-
 	/**
 	 * Enumerates all regions that overlap a certain time interval. Will enumerate by depth but does not expose lanes.
-	 * 
+	 *
 	 * @param Callback a callback called for each region. Return false to abort iteration.
 	 * @returns true if the enumeration finished, false if it was aborted by the callback returning false
 	 */
@@ -111,18 +129,17 @@ public:
 
 	/**
 	 * Enumerates all regions that overlap a certain time interval in reverse. Will enumerate by depth but does not expose lanes.
-	 * 
+	 *
 	 * @param Callback a callback called for each region. Return false to abort iteration.
 	 * @returns true if the enumeration finished, false if it was aborted by the callback returning false
 	 */
 	virtual bool EnumerateRegionsBackwards(double IntervalEnd, double IntervalStart, TFunctionRef<bool(const FTimeRegion&)> Callback) const = 0;
-	
 	/**
 	 * Will call Callback(Lane, Depth) for each lane in order.
 	 */
 	virtual void EnumerateLanes(TFunctionRef<void(const FRegionLane&, const int32)> Callback) const = 0;
 };
-	
+
 class IRegionProvider
 	: public IProvider
 {
@@ -130,26 +147,40 @@ public:
 	virtual ~IRegionProvider() override = default;
 
 	/**
-	*	Enumerates all timelines, this includes the uncategorized timeline with uncategorized regions and individual timelines for each Category
-	*	@param Callback a callback called for each timeline. Second parameter contains Category or nullptr for the uncategorized timeline.
-	*/
+	 * Enumerates all region categories.
+	 * @param Callback a callback called for each region category. Return false to abort iteration.
+	 * @returns true if the enumeration finished, false if it was aborted by the callback returning false
+	 */
+	virtual bool EnumerateCategories(TFunctionRef<bool(const FRegionCategory&)> Callback) const = 0;
+
+	/**
+	 * Enumerates all region timers (unique regions by name).
+	 * @param Callback a callback called for each region timer. Return false to abort iteration.
+	 * @returns true if the enumeration finished, false if it was aborted by the callback returning false
+	 */
+	virtual bool EnumerateTimers(TFunctionRef<bool(const FRegionTimer&)> Callback) const = 0;
+
+	/**
+	 * Enumerates all timelines, this includes the uncategorized timeline with uncategorized regions and individual timelines for each Category
+	 * @param Callback a callback called for each timeline. Second parameter contains Category or nullptr for the uncategorized timeline.
+	 */
 	virtual void EnumerateTimelinesByCategory(TFunctionRef<void(const IRegionTimeline& /*Region*/, const TCHAR* /*Category*/)> Callback) const = 0;
 
 	/**
-	*	@return the default timeline containing all regions without filtering 
-	*/
+	 * @return the default timeline containing all regions without filtering
+	 */
 	virtual const IRegionTimeline& GetDefaultTimeline() const = 0;
-	
+
 	/**
-	*	@return the timeline for a given Category or the timeline with uncategorized regions if Category is nullptr. Returns nullptr if the category was invalid/not found
-	*/
+	 * @return the timeline for a given Category or the timeline with uncategorized regions if Category is nullptr. Returns nullptr if the category was invalid/not found
+	 */
 	virtual const IRegionTimeline* GetTimelineForCategory(const TCHAR* Category) const = 0;
 
 	/**
 	 * @return returns the string used to store regions with no explicit category set ("__Uncategorized__")
 	 */
 	virtual const TCHAR* GetUncategorizedRegionCategoryName() const = 0;
-	
+
 	/**
 	 * @return the amount of currently known regions (including open-ended ones)
 	 */
@@ -171,7 +202,7 @@ public:
 
 	/**
 	 * Enumerates all regions that overlap a certain time interval. Will enumerate by depth but does not expose lanes.
-	 * 
+	 *
 	 * @param Callback a callback called for each region. Return false to abort iteration.
 	 * @returns true if the enumeration finished, false if it was aborted by the callback returning false
 	 */
@@ -181,7 +212,6 @@ public:
 	/**
 	 * Will call Callback(Lane, Depth) for each lane in order.
 	 */
-	
 	UE_DEPRECATED(5.6, "Use GetDefaultTimeline().EnumerateLanes() or EnumerateTimelinesByCategory() to access timelines instead.")
 	virtual void EnumerateLanes(TFunctionRef<void(const FRegionLane& /*Lane*/, const int32 /*Depth*/)> Callback) const = 0;
 
@@ -201,6 +231,9 @@ class IEditableRegionProvider
 public:
 	virtual ~IEditableRegionProvider() override = default;
 
+	virtual const FRegionCategory* RegisterRegionCategory(const TCHAR* Name) = 0;
+	virtual const FRegionTimer* RegisterRegionTimer(const TCHAR* Name, const FRegionCategory* Category = nullptr) = 0;
+
 	/**
 	 * Append a new begin event of a region from the trace session.
 	 * Prefer opening/closing regions with an Id, since string names are not unique.
@@ -248,9 +281,9 @@ public:
 
 ////////////////////////////////////////////////////////////////////////////////////////////////////
 
-TRACESERVICES_API FName GetRegionProviderName();
-TRACESERVICES_API const IRegionProvider& ReadRegionProvider(const IAnalysisSession& Session);
-TRACESERVICES_API IEditableRegionProvider& EditRegionProvider(IAnalysisSession& Session);
+UE_API FName GetRegionProviderName();
+UE_API const IRegionProvider& ReadRegionProvider(const IAnalysisSession& Session);
+UE_API IEditableRegionProvider& EditRegionProvider(IAnalysisSession& Session);
 
 } // namespace TraceServices
 
