[Unreal Insights] Expose FSlabAllocator as public in the TraceServices module

diff --git a/Engine/Source/Developer/TraceServices/Private/AnalysisServicePrivate.h b/Engine/Source/Developer/TraceServices/Private/AnalysisServicePrivate.h
--- a/Engine/Source/Developer/TraceServices/Private/AnalysisServicePrivate.h
+++ b/Engine/Source/Developer/TraceServices/Private/AnalysisServicePrivate.h
@@ -4,8 +4,8 @@
 
 #include "AnalysisCache.h"
 #include "Containers/HashTable.h"
-#include "Common/SlabAllocator.h"
 #include "Common/StringStore.h"
+#include "TraceServices/Containers/SlabAllocator.h"
 
 // TraceAnalysis
 #include "Trace/Analysis.h"
diff --git a/Engine/Source/Developer/TraceServices/Private/Common/SlabAllocator.h b/Engine/Source/Developer/TraceServices/Private/Common/SlabAllocator.h
deleted file mode 100644
diff --git a/Engine/Source/Developer/TraceServices/Private/Common/StringStore.h b/Engine/Source/Developer/TraceServices/Private/Common/StringStore.h
--- a/Engine/Source/Developer/TraceServices/Private/Common/StringStore.h
+++ b/Engine/Source/Developer/TraceServices/Private/Common/StringStore.h
@@ -4,11 +4,11 @@
 
 #include "CoreTypes.h"
 
-#include "Common/SlabAllocator.h"
 #include "Containers/Array.h"
 #include "Containers/Map.h"
 #include "Containers/StringView.h"
 #include "HAL/CriticalSection.h"
+#include "TraceServices/Containers/SlabAllocator.h"
 
 namespace TraceServices
 {
diff --git a/Engine/Source/Developer/TraceServices/Private/Model/LoadTimeProfilerPrivate.h b/Engine/Source/Developer/TraceServices/Private/Model/LoadTimeProfilerPrivate.h
--- a/Engine/Source/Developer/TraceServices/Private/Model/LoadTimeProfilerPrivate.h
+++ b/Engine/Source/Developer/TraceServices/Private/Model/LoadTimeProfilerPrivate.h
@@ -4,7 +4,7 @@
 
 #include "TraceServices/Model/Counters.h"
 #include "TraceServices/Model/LoadTimeProfiler.h"
-#include "Common/SlabAllocator.h"
+#include "TraceServices/Containers/SlabAllocator.h"
 #include "Common/PagedArray.h"
 #include "Common/StringStore.h"
 #include "Model/MonotonicTimeline.h"
diff --git a/Engine/Source/Developer/TraceServices/Private/Model/NetProfilerProvider.h b/Engine/Source/Developer/TraceServices/Private/Model/NetProfilerProvider.h
--- a/Engine/Source/Developer/TraceServices/Private/Model/NetProfilerProvider.h
+++ b/Engine/Source/Developer/TraceServices/Private/Model/NetProfilerProvider.h
@@ -3,7 +3,7 @@
 #pragma once
 
 #include "TraceServices/Model/NetProfiler.h"
-#include "Common/SlabAllocator.h"
+#include "TraceServices/Containers/SlabAllocator.h"
 #include "Model/Tables.h"
 
 namespace TraceServices
diff --git a/Engine/Source/Developer/TraceServices/Private/Model/Tables.h b/Engine/Source/Developer/TraceServices/Private/Model/Tables.h
--- a/Engine/Source/Developer/TraceServices/Private/Model/Tables.h
+++ b/Engine/Source/Developer/TraceServices/Private/Model/Tables.h
@@ -6,7 +6,7 @@
 #include "TraceServices/Containers/Tables.h"
 #include "Templates/Function.h"
 #include "Common/PagedArray.h"
-#include "Common/SlabAllocator.h"
+#include "TraceServices/Containers/SlabAllocator.h"
 
 namespace TraceServices
 {
diff --git a/Engine/Source/Developer/TraceServices/Private/Model/TimingProfilerPrivate.h b/Engine/Source/Developer/TraceServices/Private/Model/TimingProfilerPrivate.h
--- a/Engine/Source/Developer/TraceServices/Private/Model/TimingProfilerPrivate.h
+++ b/Engine/Source/Developer/TraceServices/Private/Model/TimingProfilerPrivate.h
@@ -3,7 +3,7 @@
 #pragma once
 
 #include "TraceServices/Model/TimingProfiler.h"
-#include "Common/SlabAllocator.h"
+#include "TraceServices/Containers/SlabAllocator.h"
 #include "Model/MonotonicTimeline.h"
 #include "Model/Tables.h"
 
diff --git a/Engine/Source/Developer/TraceServices/Private/Tests/UnitTests.cpp b/Engine/Source/Developer/TraceServices/Private/Tests/UnitTests.cpp
--- a/Engine/Source/Developer/TraceServices/Private/Tests/UnitTests.cpp
+++ b/Engine/Source/Developer/TraceServices/Private/Tests/UnitTests.cpp
@@ -1,8 +1,8 @@
 // Copyright Epic Games, Inc. All Rights Reserved.
 #include "CoreMinimal.h"
-#include "Common/SlabAllocator.h"
 #include "Common/PagedArray.h"
 #include "Misc/AutomationTest.h"
+#include "TraceServices/Containers/SlabAllocator.h"
 
 IMPLEMENT_SIMPLE_AUTOMATION_TEST(FPagedArrayFilteringTest, "System.Insights.Analysis.PagedArrayFiltering", EAutomationTestFlags_ApplicationContextMask | EAutomationTestFlags::SmokeFilter)
 
diff --git a/Engine/Source/Developer/TraceServices/Public/TraceServices/Containers/SlabAllocator.h b/Engine/Source/Developer/TraceServices/Public/TraceServices/Containers/SlabAllocator.h
new file mode 100644
--- /dev/null
+++ b/Engine/Source/Developer/TraceServices/Public/TraceServices/Containers/SlabAllocator.h
@@ -0,0 +1,59 @@
+// Copyright Epic Games, Inc. All Rights Reserved.
+
+#pragma once
+
+#include "Containers/Array.h"
+#include "HAL/CriticalSection.h"
+#include "HAL/UnrealMemory.h"
+#include "Misc/ScopeLock.h"
+#include "TraceServices/Containers/Allocators.h"
+
+namespace TraceServices
+{
+
+class FSlabAllocator
+	: public ILinearAllocator
+{
+public:
+	FSlabAllocator(uint64 InSlabSize)
+		: SlabSize(InSlabSize)
+	{
+
+	}
+
+	~FSlabAllocator()
+	{
+		FScopeLock _(&Cs);
+		for (void* Slab : Slabs)
+		{
+			FMemory::Free(Slab);
+		}
+	}
+
+	virtual void* Allocate(uint64 Size) override
+	{
+		FScopeLock _(&Cs);
+		uint64 AllocationSize = Size + (-int64(Size) & 15);
+		if (!CurrentSlab || CurrentSlabAllocatedSize + AllocationSize > SlabSize)
+		{
+			TotalAllocatedSize += SlabSize;
+			void* Allocation = FMemory::Malloc(SlabSize, 16);
+			CurrentSlab = reinterpret_cast<uint8*>(Allocation);
+			CurrentSlabAllocatedSize = 0;
+			Slabs.Add(CurrentSlab);
+		}
+		void* Allocation = CurrentSlab + CurrentSlabAllocatedSize;
+		CurrentSlabAllocatedSize += AllocationSize;
+		return Allocation;
+	}
+
+private:
+	FCriticalSection Cs;
+	TArray<void*> Slabs;
+	uint8* CurrentSlab = nullptr;
+	const uint64 SlabSize;
+	uint64 CurrentSlabAllocatedSize = 0;
+	uint64 TotalAllocatedSize = 0;
+};
+
+} // namespace TraceServices
