[Insights] Symbol Helper: Fixed incorrect function comment

#rnx

diff --git a/Engine/Source/Developer/TraceServices/Private/Common/SymbolHelper.h b/Engine/Source/Developer/TraceServices/Private/Common/SymbolHelper.h
--- a/Engine/Source/Developer/TraceServices/Private/Common/SymbolHelper.h
+++ b/Engine/Source/Developer/TraceServices/Private/Common/SymbolHelper.h
@@ -23,7 +23,6 @@ namespace TraceServices
 		 * Get the list of symbol search paths declared in the UnrealInsightsSettings.ini configuration file.
 		 * @param LogCategory The category the messages should be emitted to.
 		 * @param OutSearchPaths An array that will hold any found paths.
-		 * @return Array of found search paths.
 		 */
 		static void GetPathsFromConfiguration(const FLogCategoryAlias& LogCategory, TArray<FString>& OutSearchPaths)
 		{
