[Insights] Changed Symbol Resolver Log categories to non-static

#rnx

diff --git a/Engine/Source/Developer/TraceServices/Private/Model/DbgHelpResolver.cpp b/Engine/Source/Developer/TraceServices/Private/Model/DbgHelpResolver.cpp
--- a/Engine/Source/Developer/TraceServices/Private/Model/DbgHelpResolver.cpp
+++ b/Engine/Source/Developer/TraceServices/Private/Model/DbgHelpResolver.cpp
@@ -30,7 +30,7 @@ THIRD_PARTY_INCLUDES_END
 
 ////////////////////////////////////////////////////////////////////////////////////////////////////
 
-DEFINE_LOG_CATEGORY_STATIC(LogDbgHelp, Log, All);
+DEFINE_LOG_CATEGORY(LogDbgHelp);
 
 namespace TraceServices {
 
diff --git a/Engine/Source/Developer/TraceServices/Private/Model/DbgHelpResolver.h b/Engine/Source/Developer/TraceServices/Private/Model/DbgHelpResolver.h
--- a/Engine/Source/Developer/TraceServices/Private/Model/DbgHelpResolver.h
+++ b/Engine/Source/Developer/TraceServices/Private/Model/DbgHelpResolver.h
@@ -16,6 +16,8 @@
 #include "TraceServices/Model/Modules.h"
 #include <atomic>
 
+DECLARE_LOG_CATEGORY_EXTERN(LogDbgHelp, Log, All);
+
 namespace TraceServices {
 
 class FDbgHelpResolver : public FRunnable
diff --git a/Engine/Source/Developer/TraceServices/Private/Model/PsymResolver.cpp b/Engine/Source/Developer/TraceServices/Private/Model/PsymResolver.cpp
--- a/Engine/Source/Developer/TraceServices/Private/Model/PsymResolver.cpp
+++ b/Engine/Source/Developer/TraceServices/Private/Model/PsymResolver.cpp
@@ -27,7 +27,7 @@
 namespace TraceServices {
 
 /////////////////////////////////////////////////////////////////////
-DEFINE_LOG_CATEGORY_STATIC(LogPsymResolver, Log, All);
+DEFINE_LOG_CATEGORY(LogPsymResolver);
 
 /////////////////////////////////////////////////////////////////////
 
diff --git a/Engine/Source/Developer/TraceServices/Private/Model/PsymResolver.h b/Engine/Source/Developer/TraceServices/Private/Model/PsymResolver.h
--- a/Engine/Source/Developer/TraceServices/Private/Model/PsymResolver.h
+++ b/Engine/Source/Developer/TraceServices/Private/Model/PsymResolver.h
@@ -12,6 +12,8 @@
 #include "TraceServices/Model/AnalysisSession.h"
 #include "TraceServices/Model/Modules.h"
 
+DECLARE_LOG_CATEGORY_EXTERN(LogPsymResolver, Log, All);
+
 namespace TraceServices
 {
 ////////////////////////////////////////////////////////////////////////////////
diff --git a/Engine/Source/Developer/TraceServices/Private/Model/SymslibResolver.cpp b/Engine/Source/Developer/TraceServices/Private/Model/SymslibResolver.cpp
--- a/Engine/Source/Developer/TraceServices/Private/Model/SymslibResolver.cpp
+++ b/Engine/Source/Developer/TraceServices/Private/Model/SymslibResolver.cpp
@@ -41,7 +41,7 @@
 
 ////////////////////////////////////////////////////////////////////////////////////////////////////
 
-DEFINE_LOG_CATEGORY_STATIC(LogSymslib, Log, All);
+DEFINE_LOG_CATEGORY(LogSymslib);
 
 namespace TraceServices {
 
diff --git a/Engine/Source/Developer/TraceServices/Private/Model/SymslibResolver.h b/Engine/Source/Developer/TraceServices/Private/Model/SymslibResolver.h
--- a/Engine/Source/Developer/TraceServices/Private/Model/SymslibResolver.h
+++ b/Engine/Source/Developer/TraceServices/Private/Model/SymslibResolver.h
@@ -23,6 +23,8 @@
 
 #include "symslib.h"
 
+DECLARE_LOG_CATEGORY_EXTERN(LogSymslib, Log, All);
+
 namespace TraceServices {
 
 class FSymslibResolver
