[Insights]
* TraceServices: Fixed FSlabAllocator when trying to allocate with size > slab size. Also added GetAllocatedSize() and GetUsedSize() to allow estimate used and wasted memory for this allocator.

diff --git a/Engine/Source/Developer/TraceServices/Public/TraceServices/Containers/SlabAllocator.h b/Engine/Source/Developer/TraceServices/Public/TraceServices/Containers/SlabAllocator.h
--- a/Engine/Source/Developer/TraceServices/Public/TraceServices/Containers/SlabAllocator.h
+++ b/Engine/Source/Developer/TraceServices/Public/TraceServices/Containers/SlabAllocator.h
@@ -18,7 +18,6 @@ public:
 	FSlabAllocator(uint64 InSlabSize)
 		: SlabSize(InSlabSize)
 	{
-
 	}
 
 	~FSlabAllocator()
@@ -33,26 +32,49 @@ public:
 	virtual void* Allocate(uint64 Size) override
 	{
 		FScopeLock _(&Cs);
+		TotalUsedSize += Size;
 		uint64 AllocationSize = Size + (-int64(Size) & 15);
-		if (!CurrentSlab || CurrentSlabAllocatedSize + AllocationSize > SlabSize)
+		if (AllocationSize <= SlabSize)
+		{
+			if (!CurrentSlab || CurrentSlabAllocatedSize + AllocationSize > SlabSize)
+			{
+				TotalAllocatedSize += SlabSize;
+				CurrentSlab = reinterpret_cast<uint8*>(FMemory::Malloc(SlabSize, 16));
+				CurrentSlabAllocatedSize = 0;
+				Slabs.Add(CurrentSlab);
+			}
+			void* Allocation = CurrentSlab + CurrentSlabAllocatedSize;
+			CurrentSlabAllocatedSize += AllocationSize;
+			return Allocation;
+		}
+		else
 		{
-			TotalAllocatedSize += SlabSize;
-			void* Allocation = FMemory::Malloc(SlabSize, 16);
-			CurrentSlab = reinterpret_cast<uint8*>(Allocation);
-			CurrentSlabAllocatedSize = 0;
-			Slabs.Add(CurrentSlab);
+			TotalAllocatedSize += AllocationSize;
+			void* Allocation = FMemory::Malloc(AllocationSize, 16);
+			Slabs.Add(Allocation);
+			return Allocation;
 		}
-		void* Allocation = CurrentSlab + CurrentSlabAllocatedSize;
-		CurrentSlabAllocatedSize += AllocationSize;
-		return Allocation;
+	}
+
+	uint64 GetUsedSize() const
+	{
+		FScopeLock _(&Cs);
+		return TotalUsedSize;
+	}
+
+	uint64 GetAllocatedSize() const
+	{
+		FScopeLock _(&Cs);
+		return TotalAllocatedSize + Slabs.GetAllocatedSize();
 	}
 
 private:
-	FCriticalSection Cs;
+	mutable FCriticalSection Cs;
 	TArray<void*> Slabs;
 	uint8* CurrentSlab = nullptr;
 	const uint64 SlabSize;
 	uint64 CurrentSlabAllocatedSize = 0;
+	uint64 TotalUsedSize = 0;
 	uint64 TotalAllocatedSize = 0;
 };
 
