[Insights] Regions: Fixed comment for GetUncategorizedRegionCategoryName.

#rnx

diff --git a/Engine/Source/Developer/TraceServices/Public/TraceServices/Model/Regions.h b/Engine/Source/Developer/TraceServices/Public/TraceServices/Model/Regions.h
--- a/Engine/Source/Developer/TraceServices/Public/TraceServices/Model/Regions.h
+++ b/Engine/Source/Developer/TraceServices/Public/TraceServices/Model/Regions.h
@@ -177,7 +177,7 @@ public:
 	virtual const IRegionTimeline* GetTimelineForCategory(const TCHAR* Category) const = 0;
 
 	/**
-	 * @return returns the string used to store regions with no explicit category set ("__Uncategorized__")
+	 * @return the string used to store regions with no explicit category set
 	 */
 	virtual const TCHAR* GetUncategorizedRegionCategoryName() const = 0;
 
