[Insights]
* Symbols: Unified code re reading symbol search paths from ini and environment variables.
* Symbols: Fixed the string allocator(s) to support strings longer than page size. Improved the debug stats (allocated/waste bytes). Removed code duplicate. Now both Psym and Symslib resolvers uses same string allocator implementation. Also increased allocator's page size for both Psym resolver and for Symslib resolver to 64 KiB.
* Symbols: Fixed the Psym resolver to correctly compute Discovered, Resolved and Failed module stats. Also fixed the status of resolver for fail cases.
* Symbols: Added "Available" per module stats and updated Psym and Symslib resolvers to compute it.
* Modules: Added "Available" column to show the total number of symbols available in the loaded symbols file.

diff --git a/Engine/Source/Developer/TraceInsights/Private/Insights/Widgets/SModulesView.cpp b/Engine/Source/Developer/TraceInsights/Private/Insights/Widgets/SModulesView.cpp
--- a/Engine/Source/Developer/TraceInsights/Private/Insights/Widgets/SModulesView.cpp
+++ b/Engine/Source/Developer/TraceInsights/Private/Insights/Widgets/SModulesView.cpp
@@ -77,11 +77,12 @@ public:
 	FText GetDiscoveredAsText() const
 	{
 		// Discovered != Cached + Resolved + Failed ?
-		if (GetDiscoveredSymbolCount() != Module->Stats.Cached + Module->Stats.Resolved + Module->Stats.Failed)
+		const uint32 DiscoveredSymbolCount = GetDiscoveredSymbolCount();
+		if (DiscoveredSymbolCount != Module->Stats.Cached + Module->Stats.Resolved + Module->Stats.Failed)
 		{
-			return FText::Format(LOCTEXT("DiscoveredStarredFmt", "{0} (!)"), FText::AsNumber(Module->Stats.Discovered + Module->Stats.Cached));
+			return FText::Format(LOCTEXT("DiscoveredStarredFmt", "{0} (!)"), FText::AsNumber(DiscoveredSymbolCount));
 		}
-		return FText::AsNumber(GetDiscoveredSymbolCount());
+		return FText::AsNumber(DiscoveredSymbolCount);
 	}
 
 	inline uint32 GetCachedSymbolCount() const
@@ -114,6 +115,16 @@ public:
 		return FText::AsNumber(GetFailedSymbolCount());
 	}
 
+	inline uint32 GetAvailableSymbolCount() const
+	{
+		return Module->Stats.Available;
+	}
+
+	FText GetAvailableAsText() const
+	{
+		return FText::AsNumber(GetAvailableSymbolCount());
+	}
+
 	FText GetStatusText() const
 	{
 		return FText::FromStringView(ModuleStatusToString(Module->Status.load()));
@@ -125,7 +136,7 @@ public:
 	}
 
 	FText GetToolTipText() const { return FText(); }
-	
+
 	inline uint32 GetSize() const
 	{
 		return Module->Size;
@@ -241,6 +252,16 @@ public:
 					.Text(this, &SModuleRow::GetFailedAsText)
 				];
 		}
+		else if (ColumnName == ModulesViewColumns::AvailableColumnName)
+		{
+			return SNew(SBox)
+				.Padding(FMargin(4.0, 0.0))
+				[
+					SNew(STextBlock)
+					.ColorAndOpacity(this, &SModuleRow::GetStatusColor)
+					.Text(this, &SModuleRow::GetAvailableAsText)
+				];
+		}
 		else if (ColumnName == ModulesViewColumns::StatusColumnName)
 		{
 			return SNew(SBox)
@@ -368,6 +389,19 @@ public:
 		}
 	}
 
+	FText GetAvailableAsText() const
+	{
+		const TSharedPtr<FModule> ModulePin = WeakModule.Pin();
+		if (ModulePin.IsValid())
+		{
+			return ModulePin->GetAvailableAsText();
+		}
+		else
+		{
+			return FText();
+		}
+	}
+
 	FText GetStatus() const
 	{
 		TSharedPtr<FModule> ModulePin = WeakModule.Pin();
@@ -406,7 +440,7 @@ public:
 			return FText();
 		}
 	}
-	
+
 	FText GetSizeAsText() const
 	{
 		const TSharedPtr<FModule> ModulePin = WeakModule.Pin();
@@ -576,7 +610,7 @@ public:
 			return GetBlackColor();
 		}
 	}
-	
+
 	FSlateColor GetSizeColor() const
 	{
 		return GetWhileColor();
@@ -668,7 +702,7 @@ void SModulesView::Construct(const FArguments& InArgs)
 							.InitialSortMode(EColumnSortMode::Ascending)
 							.SortMode(this, &SModulesView::GetSortModeForColumn, ModulesViewColumns::AddressRangeColumnName)
 							.OnSort(this, &SModulesView::OnSortModeChanged)
-							
+
 							+ SHeaderRow::Column(ModulesViewColumns::SizeColumnName)
 							.ManualWidth(80.0f)
 							.HAlignHeader(HAlign_Right)
@@ -690,7 +724,7 @@ void SModulesView::Construct(const FArguments& InArgs)
 							.HAlignHeader(HAlign_Right)
 							.HAlignCell(HAlign_Right)
 							.DefaultLabel(LOCTEXT("DiscoveredColumn", "Discovered"))
-							.DefaultTooltip(LOCTEXT("DiscoveredColumnTooltip", "The total number of symbols discovered in the trace.\nDiscovered = Cached + Resolved + Failed"))
+							.DefaultTooltip(LOCTEXT("DiscoveredColumnTooltip", "The total number of symbols discovered (used) in the trace.\nDiscovered = Cached + Resolved + Failed"))
 							.InitialSortMode(EColumnSortMode::Descending)
 							.SortMode(this, &SModulesView::GetSortModeForColumn, ModulesViewColumns::DiscoveredColumnName)
 							.OnSort(this, &SModulesView::OnSortModeChanged)
@@ -700,7 +734,7 @@ void SModulesView::Construct(const FArguments& InArgs)
 							.HAlignHeader(HAlign_Right)
 							.HAlignCell(HAlign_Right)
 							.DefaultLabel(LOCTEXT("CachedColumn", "Cached"))
-							.DefaultTooltip(LOCTEXT("CachedColumnTooltip", "The number of symbols loaded (pre-resolved) from the cache (*.ucache file)."))
+							.DefaultTooltip(LOCTEXT("CachedColumnTooltip", "The number of discovered symbols loaded from the cache (*.ucache file)."))
 							.InitialSortMode(EColumnSortMode::Descending)
 							.SortMode(this, &SModulesView::GetSortModeForColumn, ModulesViewColumns::CachedColumnName)
 							.OnSort(this, &SModulesView::OnSortModeChanged)
@@ -710,7 +744,7 @@ void SModulesView::Construct(const FArguments& InArgs)
 							.HAlignHeader(HAlign_Right)
 							.HAlignCell(HAlign_Right)
 							.DefaultLabel(LOCTEXT("ResolvedColumn", "Resolved"))
-							.DefaultTooltip(LOCTEXT("ResolvedColumnTooltip", "The number of symbols successfully resolved in the current analysis session."))
+							.DefaultTooltip(LOCTEXT("ResolvedColumnTooltip", "The number of discovered symbols successfully resolved in the current analysis session."))
 							.InitialSortMode(EColumnSortMode::Descending)
 							.SortMode(this, &SModulesView::GetSortModeForColumn, ModulesViewColumns::ResolvedColumnName)
 							.OnSort(this, &SModulesView::OnSortModeChanged)
@@ -720,11 +754,21 @@ void SModulesView::Construct(const FArguments& InArgs)
 							.HAlignHeader(HAlign_Right)
 							.HAlignCell(HAlign_Right)
 							.DefaultLabel(LOCTEXT("FailedColumn", "Failed"))
-							.DefaultTooltip(LOCTEXT("FailedColumnTooltip", "The number of symbols failed to resolve."))
+							.DefaultTooltip(LOCTEXT("FailedColumnTooltip", "The number of discovered symbols that failed to resolve."))
 							.InitialSortMode(EColumnSortMode::Descending)
 							.SortMode(this, &SModulesView::GetSortModeForColumn, ModulesViewColumns::FailedColumnName)
 							.OnSort(this, &SModulesView::OnSortModeChanged)
 
+							+ SHeaderRow::Column(ModulesViewColumns::AvailableColumnName)
+							.ManualWidth(76.0f)
+							.HAlignHeader(HAlign_Right)
+							.HAlignCell(HAlign_Right)
+							.DefaultLabel(LOCTEXT("AvailableColumn", "Available"))
+							.DefaultTooltip(LOCTEXT("AvailableColumnTooltip", "The total number of symbols available in the loaded symbols file."))
+							.InitialSortMode(EColumnSortMode::Descending)
+							.SortMode(this, &SModulesView::GetSortModeForColumn, ModulesViewColumns::AvailableColumnName)
+							.OnSort(this, &SModulesView::OnSortModeChanged)
+
 							+ SHeaderRow::Column(ModulesViewColumns::StatusColumnName)
 							.ManualWidth(100.0f)
 							.DefaultLabel(LOCTEXT("StatusColumn", "Status"))
@@ -734,7 +778,7 @@ void SModulesView::Construct(const FArguments& InArgs)
 
 							+ SHeaderRow::Column(ModulesViewColumns::SymbolsFileColumnName)
 							.ManualWidth(800.0f)
-							.DefaultLabel(LOCTEXT("SymbolsFileColumn", "Symbols"))
+							.DefaultLabel(LOCTEXT("SymbolsFileColumn", "Symbols File"))
 							.InitialSortMode(EColumnSortMode::Ascending)
 							.SortMode(this, &SModulesView::GetSortModeForColumn, ModulesViewColumns::SymbolsFileColumnName)
 							.OnSort(this, &SModulesView::OnSortModeChanged)
@@ -868,7 +912,7 @@ void SModulesView::Construct(const FArguments& InArgs)
 					SNew(SImage)
 					.Image(FInsightsStyle::GetBrush("Icons.Hint.TreeItem"))
 				]
-				
+
 				+ SHorizontalBox::Slot()
 				.VAlign(VAlign_Center)
 				.HAlign(HAlign_Center)
@@ -969,7 +1013,7 @@ FSlateColor SModulesView::GetStatsTextColor() const
 {
 	return FSlateColor(FLinearColor(1.0f, 1.0f, 1.0f, 1.0f));
 }
-	
+
 ////////////////////////////////////////////////////////////////////////////////////////////////////
 
 FText SModulesView::GetSymbolSearchPaths() const
@@ -986,14 +1030,14 @@ FText SModulesView::GetSymbolSearchPaths() const
 				CleanedPath.TrimQuotesInline();
 				PathsListStr << (TEXT("%s"), PathsListStr.Len() ? "\n" : "") << TEXT("\u2022 ") << CleanedPath;
 			});
-			
+
 			if (PathsListStr.Len())
 			{
 				return FText::FromStringView(PathsListStr);
 			}
 		}
 	}
-	
+
 	return LOCTEXT("SymbolResolutionNoPaths", " No symbol paths have been setup. \n Use configuration file, UE_INSIGHTS_SYMBOL_PATH environment variable or manually select a file or directory from the Modules view.");
 }
 
@@ -1011,7 +1055,7 @@ TSharedPtr<SWidget> SModulesView::ListView_GetContextMenu()
 		{
 			MenuBuilder.AddMenuEntry(
 				TAttribute<FText>(this, &SModulesView::LoadSymbolsFile_Label),
-				LOCTEXT("LoadSymbolsFile_TooltipReady", "Load symbols for a module by specifying a file. If succesful, tries to load other failed modules from the same directory."),
+				LOCTEXT("LoadSymbolsFile_TooltipReady", "Load symbols for a module by specifying a file.\nIf successful, tries to load other failed modules from the same directory."),
 				FSlateIcon(),
 				FUIAction(FExecuteAction::CreateSP(this, &SModulesView::LoadSymbols_Execute, SelectedModule, true),
 				          FCanExecuteAction::CreateSP(this, &SModulesView::LoadSymbols_CanExecute),
@@ -1021,7 +1065,7 @@ TSharedPtr<SWidget> SModulesView::ListView_GetContextMenu()
 			);
 			MenuBuilder.AddMenuEntry(
 				TAttribute<FText>(this, &SModulesView::LoadSymbolsDirectory_Label),
-				LOCTEXT("LoadSymbolsDirectory_TooltipReady", "Load symbols for a module by specifying a directory. If succesful, tries to load other failed modules from the same directory."),
+				LOCTEXT("LoadSymbolsDirectory_TooltipReady", "Load symbols for a module by specifying a directory.\nIf successful, tries to load other failed modules from the same directory."),
 				FSlateIcon(),
 				FUIAction(FExecuteAction::CreateSP(this, &SModulesView::LoadSymbols_Execute, SelectedModule, false),
 				          FCanExecuteAction::CreateSP(this, &SModulesView::LoadSymbols_CanExecute),
@@ -1239,6 +1283,19 @@ void SModulesView::UpdateSorting()
 				{ return A->GetFailedSymbolCount() > B->GetFailedSymbolCount(); });
 		}
 	}
+	else if (SortColumn == ModulesViewColumns::AvailableColumnName)
+	{
+		if (SortMode == EColumnSortMode::Ascending)
+		{
+			Modules.Sort([](const TSharedPtr<FModule>& A, const TSharedPtr<FModule>& B)
+				{ return A->GetAvailableSymbolCount() < B->GetAvailableSymbolCount(); });
+		}
+		else
+		{
+			Modules.Sort([](const TSharedPtr<FModule>& A, const TSharedPtr<FModule>& B)
+				{ return A->GetAvailableSymbolCount() > B->GetAvailableSymbolCount(); });
+		}
+	}
 	else if (SortColumn == ModulesViewColumns::StatusColumnName)
 	{
 		if (SortMode == EColumnSortMode::Ascending)
diff --git a/Engine/Source/Developer/TraceInsights/Private/Insights/Widgets/SModulesView.h b/Engine/Source/Developer/TraceInsights/Private/Insights/Widgets/SModulesView.h
--- a/Engine/Source/Developer/TraceInsights/Private/Insights/Widgets/SModulesView.h
+++ b/Engine/Source/Developer/TraceInsights/Private/Insights/Widgets/SModulesView.h
@@ -29,6 +29,7 @@ namespace ModulesViewColumns
 	static const FName CachedColumnName(TEXT("Cached"));
 	static const FName ResolvedColumnName(TEXT("Resolved"));
 	static const FName FailedColumnName(TEXT("Failed"));
+	static const FName AvailableColumnName(TEXT("Available"));
 	static const FName StatusColumnName(TEXT("Status"));
 	static const FName SymbolsFileColumnName(TEXT("SymbolsFile"));
 	static const FName SizeColumnName(TEXT("Size"));
diff --git a/Engine/Source/Developer/TraceServices/Private/Common/SymbolHelper.cpp b/Engine/Source/Developer/TraceServices/Private/Common/SymbolHelper.cpp
new file mode 100644
--- /dev/null
+++ b/Engine/Source/Developer/TraceServices/Private/Common/SymbolHelper.cpp
@@ -0,0 +1,82 @@
+// Copyright Epic Games, Inc. All Rights Reserved.
+
+#include "SymbolHelper.h"
+
+#include "Misc/ConfigContext.h"
+
+namespace TraceServices
+{
+
+////////////////////////////////////////////////////////////////////////////////////////////////////
+// FSymbolHelper
+////////////////////////////////////////////////////////////////////////////////////////////////////
+
+void FSymbolHelper::GetSymbolSearchPaths(const FLogCategoryAlias& LogCategory, TArray<FString>& OutSearchPaths)
+{
+	GetPathsFromEnvVars(LogCategory, OutSearchPaths);
+	GetPathsFromSettings(LogCategory, OutSearchPaths);
+}
+
+////////////////////////////////////////////////////////////////////////////////////////////////////
+
+void FSymbolHelper::GetPathsFromEnvVars(const FLogCategoryAlias& LogCategory, TArray<FString>& OutSearchPaths)
+{
+	auto SplitEnvPaths = [](FStringView EnvVariable, TArray<FString>& OutList)
+		{
+			FString SymbolPathPart, SymbolPathRemainder(EnvVariable);
+			while (SymbolPathRemainder.Split(TEXT(";"), &SymbolPathPart, &SymbolPathRemainder))
+			{
+				if (!SymbolPathPart.IsEmpty())
+				{
+					OutList.Emplace(SymbolPathPart);
+				}
+			}
+			if (!SymbolPathRemainder.IsEmpty())
+			{
+				OutList.Emplace(SymbolPathRemainder);
+			}
+		};
+
+	const FString InsightsSymbolPath = FPlatformMisc::GetEnvironmentVariable(TEXT("UE_INSIGHTS_SYMBOL_PATH"));
+	UE_LOG_REF(LogCategory, Log, TEXT("UE_INSIGHTS_SYMBOL_PATH: %s"), InsightsSymbolPath.IsEmpty() ? TEXT("not set") : *InsightsSymbolPath);
+	SplitEnvPaths(InsightsSymbolPath, OutSearchPaths);
+
+#if PLATFORM_WINDOWS
+	const FString NTSymbolPath = FPlatformMisc::GetEnvironmentVariable(TEXT("_NT_SYMBOL_PATH"));
+	UE_LOG_REF(LogCategory, Log, TEXT("_NT_SYMBOL_PATH: %s"), NTSymbolPath.IsEmpty() ? TEXT("not set") : *NTSymbolPath);
+	SplitEnvPaths(NTSymbolPath, OutSearchPaths);
+#endif
+}
+
+////////////////////////////////////////////////////////////////////////////////////////////////////
+
+void FSymbolHelper::GetPathsFromSettings(const FLogCategoryAlias& LogCategory, TArray<FString>& OutSearchPaths)
+{
+	FString SettingsIni;
+	TArray<FString> SymbolSearchPaths;
+
+	if (FConfigContext::ReadIntoGConfig().Load(TEXT("UnrealInsightsSettings"), SettingsIni))
+	{
+		GConfig->GetArray(TEXT("Insights"), TEXT("SymbolSearchPaths"), SymbolSearchPaths, SettingsIni);
+	}
+
+#if !NO_LOGGING
+	if (SymbolSearchPaths.IsEmpty())
+	{
+		UE_LOG_REF(LogCategory, Log, TEXT("SymbolSearchPaths not set in the UnrealInsightsSettings.ini file."));
+	}
+	else
+	{
+		for (const FString& Path : SymbolSearchPaths)
+		{
+			UE_LOG_REF(LogCategory, Log, TEXT("+SymbolSearchPaths=%s"), *Path);
+		}
+	}
+#endif
+
+	OutSearchPaths += SymbolSearchPaths;
+}
+
+////////////////////////////////////////////////////////////////////////////////////////////////////
+
+} // namespace TraceServices
diff --git a/Engine/Source/Developer/TraceServices/Private/Common/SymbolHelper.h b/Engine/Source/Developer/TraceServices/Private/Common/SymbolHelper.h
--- a/Engine/Source/Developer/TraceServices/Private/Common/SymbolHelper.h
+++ b/Engine/Source/Developer/TraceServices/Private/Common/SymbolHelper.h
@@ -2,56 +2,45 @@
 
 #pragma once
 
-#include "Internationalization/Text.h"
-#include "Misc/ConfigContext.h"
-
-#define LOCTEXT_NAMESPACE "UE::Insights::FSymbolHelper"
+#include "CoreTypes.h"
+#include "Containers/Array.h"
+#include "Containers/UnrealString.h"
+#include "Logging/LogCategory.h"
 
 namespace TraceServices
 {
-	class FSymbolHelper
-	{
-	public:
-	
-	#if NO_LOGGING
-    	typedef FNoLoggingCategory FLogCategoryAlias;
-    #else
-    	typedef FLogCategoryBase FLogCategoryAlias;
-    #endif
-
-		/**
-		 * Get the list of symbol search paths declared in the UnrealInsightsSettings.ini configuration file.
-		 * @param LogCategory The category the messages should be emitted to.
-		 * @param OutSearchPaths An array that will hold any found paths.
-		 */
-		static void GetPathsFromConfiguration(const FLogCategoryAlias& LogCategory, TArray<FString>& OutSearchPaths)
-		{
-			FString SettingsIni;
-			TArray<FString> SymbolSearchPaths;
-
-			if (FConfigContext::ReadIntoGConfig().Load(TEXT("UnrealInsightsSettings"), SettingsIni))
-			{
-				GConfig->GetArray(TEXT("Insights.Symbols"), TEXT("SymbolSearchPaths"), SymbolSearchPaths, SettingsIni);
-			}
-
-#if !NO_LOGGING
-			if (SymbolSearchPaths.IsEmpty())
-			{
-				UE_LOG_REF(LogCategory, Log, TEXT("[Insights.Symbols] SymbolSearchPaths not set"));
-			}
-			else
-			{
-				for (const FString& Path : SymbolSearchPaths)
-				{
-					UE_LOG_REF(LogCategory, Log, TEXT("[Insights.Symbols] +SymbolSearchPaths=%s"), *Path);
-				}
-			}
+
+class FSymbolHelper
+{
+public:
+#if NO_LOGGING
+	typedef FNoLoggingCategory FLogCategoryAlias;
+#else
+	typedef FLogCategoryBase FLogCategoryAlias;
 #endif
 
-			OutSearchPaths += SymbolSearchPaths;
-		}
-	};
+	/**
+	* Gets the list of symbol search paths declared by environment variables and
+	* the symbol search paths declared in the UnrealInsightsSettings.ini configuration file.
+	* @param LogCategory The log category the messages should be emitted to.
+	* @param OutSearchPaths An array where the found paths area added.
+	*/
+	static void GetSymbolSearchPaths(const FLogCategoryAlias& LogCategory, TArray<FString>& OutSearchPaths);
+
+private:
+	/**
+	* Gets the list of symbol search paths declared by environment variables.
+	* @param LogCategory The log category the messages should be emitted to.
+	* @param OutSearchPaths An array where the found paths area added.
+	*/
+	static void GetPathsFromEnvVars(const FLogCategoryAlias& LogCategory, TArray<FString>& OutSearchPaths);
+
+	/**
+	* Gets the list of symbol search paths declared in the UnrealInsightsSettings.ini configuration file.
+	* @param LogCategory The log category the messages should be emitted to.
+	* @param OutSearchPaths An array where the found paths area added.
+	*/
+	static void GetPathsFromSettings(const FLogCategoryAlias& LogCategory, TArray<FString>& OutSearchPaths);
+};
 
 } // namespace TraceServices
-
-#undef LOCTEXT_NAMESPACE
diff --git a/Engine/Source/Developer/TraceServices/Private/Common/SymbolStringAllocator.cpp b/Engine/Source/Developer/TraceServices/Private/Common/SymbolStringAllocator.cpp
new file mode 100644
--- /dev/null
+++ b/Engine/Source/Developer/TraceServices/Private/Common/SymbolStringAllocator.cpp
@@ -0,0 +1,51 @@
+// Copyright Epic Games, Inc. All Rights Reserved.
+
+#include "SymbolStringAllocator.h"
+
+namespace TraceServices
+{
+
+////////////////////////////////////////////////////////////////////////////////////////////////////
+// FSymbolStringAllocator
+////////////////////////////////////////////////////////////////////////////////////////////////////
+
+const TCHAR* FSymbolStringAllocator::Store(const FStringView InString)
+{
+	const uint32 StringSize = InString.Len() + 1;
+	TotalUsedSize += StringSize * sizeof(TCHAR);
+	if (StringSize <= BlockSize)
+	{
+		if (StringSize > BlockRemaining)
+		{
+			++NumAllocatedBlocks;
+			const uint32 AllocSize = BlockSize * sizeof(TCHAR);
+			TotalAllocatedSize += AllocSize;
+			// Allocate a new block. The allocated memory is owned by the linear allocator.
+			Block = (TCHAR*)Allocator.Allocate(AllocSize);
+			BlockRemaining = BlockSize;
+		}
+		const uint32 CopiedSize = InString.CopyString(Block, BlockRemaining - 1, 0);
+		check(StringSize == CopiedSize + 1);
+		Block[StringSize - 1] = TEXT('\0');
+		BlockRemaining -= StringSize;
+		const TCHAR* OutString = Block;
+		Block += StringSize;
+		return OutString;
+	}
+	else
+	{
+		++NumAllocatedBlocks;
+		const uint32 AllocSize = StringSize * sizeof(TCHAR);
+		TotalAllocatedSize += AllocSize;
+		// Allocate memory for the current string. The allocated memory is owned by the linear allocator.
+		TCHAR* OutString = (TCHAR*)Allocator.Allocate(AllocSize);
+		const uint32 CopiedSize = InString.CopyString(OutString, StringSize - 1, 0);
+		check(StringSize == CopiedSize + 1);
+		OutString[StringSize - 1] = TEXT('\0');
+		return OutString;
+	}
+}
+
+////////////////////////////////////////////////////////////////////////////////////////////////////
+
+} // namespace TraceServices
diff --git a/Engine/Source/Developer/TraceServices/Private/Common/SymbolStringAllocator.h b/Engine/Source/Developer/TraceServices/Private/Common/SymbolStringAllocator.h
new file mode 100644
--- /dev/null
+++ b/Engine/Source/Developer/TraceServices/Private/Common/SymbolStringAllocator.h
@@ -0,0 +1,58 @@
+// Copyright Epic Games, Inc. All Rights Reserved.
+
+#pragma once
+
+#include "CoreTypes.h"
+#include "Containers/StringView.h"
+
+// TraceServices
+#include "TraceServices/Containers/Allocators.h"
+
+namespace TraceServices
+{
+
+/// An allocator for strings as a wrapper over a linear allocator.
+/// The memory allocated is owned by the provided linear allocator.
+/// So, the lifetime of stored strings is the same as the linear allocator.
+class FSymbolStringAllocator
+{
+public:
+	FSymbolStringAllocator(ILinearAllocator& InAllocator, uint32 InBlockSize)
+		: Allocator(InAllocator)
+		, BlockSize(InBlockSize)
+	{
+	}
+
+	const TCHAR* Store(const TCHAR* InString)
+	{
+		return Store(FStringView(InString));
+	}
+
+	const TCHAR* Store(const FStringView InString);
+
+	uint64 GetUsedSize() const
+	{
+		return TotalUsedSize;
+	}
+
+	uint64 GetAllocatedSize() const
+	{
+		return TotalAllocatedSize;
+	}
+
+	uint32 GetNumAllocatedBlocks() const
+	{
+		return NumAllocatedBlocks;
+	}
+
+private:
+	ILinearAllocator& Allocator;
+	TCHAR* Block = nullptr;
+	uint32 BlockSize;
+	uint32 BlockRemaining = 0;
+	uint32 NumAllocatedBlocks = 0;
+	uint64 TotalAllocatedSize = 0;
+	uint64 TotalUsedSize = 0;
+};
+
+} // namespace TraceServices
diff --git a/Engine/Source/Developer/TraceServices/Private/Model/DbgHelpResolver.cpp b/Engine/Source/Developer/TraceServices/Private/Model/DbgHelpResolver.cpp
--- a/Engine/Source/Developer/TraceServices/Private/Model/DbgHelpResolver.cpp
+++ b/Engine/Source/Developer/TraceServices/Private/Model/DbgHelpResolver.cpp
@@ -2,9 +2,6 @@
 
 #include "DbgHelpResolver.h"
 
-// TraceServices
-#include "Common/SymbolHelper.h"
-
 #if PLATFORM_WINDOWS
 
 #include "Algo/ForEach.h"
@@ -20,7 +17,11 @@
 #include "Misc/ScopeLock.h"
 #include "Misc/ScopeRWLock.h"
 #include "Misc/StringBuilder.h"
+
+// TraceServices
+#include "Common/SymbolHelper.h"
 #include "TraceServices/Model/AnalysisSession.h"
+
 #include <atomic>
 
 #include "Windows/AllowWindowsPlatformTypes.h"
@@ -48,36 +49,12 @@ FDbgHelpResolver::FDbgHelpResolver(IAnalysisSession& InSession, IResolvedSymbolF
 	// Setup search paths. Paths are searched in the following order:
 	// 1. Any new path entered by the user this session
 	// 2. Path of the executable (if available)
-	// 3. Paths from UE_INSIGHTS_SYMBOL_PATH
-	// 4. Paths from _NT_SYMBOL_PATH
+	// 3. Paths from UE_INSIGHTS_SYMBOL_PATH environment variable
+	// 4. Paths from _NT_SYMBOL_PATH environment variable
 	// 5. Paths from the user configuration file
 
-	auto SplitEnvPaths = [](FStringView EnvVariable, TArray<FString>& OutList)
-	{
-		FString SymbolPathPart, SymbolPathRemainder(EnvVariable);
-		while (SymbolPathRemainder.Split(TEXT(";"), &SymbolPathPart, &SymbolPathRemainder))
-		{
-			OutList.Emplace(SymbolPathPart);
-		}
-		if (!SymbolPathRemainder.IsEmpty())
-		{
-			OutList.Emplace(SymbolPathRemainder);
-		}
-	};
-
-	// Paths from environment
-	const FString InsightsSymbolPath =  FPlatformMisc::GetEnvironmentVariable(TEXT("UE_INSIGHTS_SYMBOL_PATH"));
-	UE_LOG(LogDbgHelp, Log, TEXT("UE_INSIGHTS_SYMBOL_PATH: %s"), InsightsSymbolPath.IsEmpty() ? TEXT("Not set") : *InsightsSymbolPath);
-	SplitEnvPaths(InsightsSymbolPath, ConfigSymbolSearchPaths);
-
-#if PLATFORM_WINDOWS
-	const FString NTSymbolPath = FPlatformMisc::GetEnvironmentVariable(TEXT("_NT_SYMBOL_PATH"));
-	UE_LOG(LogDbgHelp, Log, TEXT("_NT_SYMBOL_PATH: %s"), NTSymbolPath.IsEmpty() ? TEXT("Not set") : *NTSymbolPath);
-	SplitEnvPaths(NTSymbolPath, ConfigSymbolSearchPaths);
-#endif
-
-	// Paths from configuration file
-	FSymbolHelper::GetPathsFromConfiguration(LogDbgHelp, ConfigSymbolSearchPaths);
+	// Get the pre-configured symbol search paths (environment variables + user configuration file).
+	FSymbolHelper::GetSymbolSearchPaths(LogDbgHelp, ConfigSymbolSearchPaths);
 
 	Start();
 }
diff --git a/Engine/Source/Developer/TraceServices/Private/Model/PsymResolver.cpp b/Engine/Source/Developer/TraceServices/Private/Model/PsymResolver.cpp
--- a/Engine/Source/Developer/TraceServices/Private/Model/PsymResolver.cpp
+++ b/Engine/Source/Developer/TraceServices/Private/Model/PsymResolver.cpp
@@ -5,71 +5,37 @@
 // TraceServices
 #include "Common/SymbolHelper.h"
 
-#include "Algo/Sort.h"
 #include "Algo/ForEach.h"
+#include "Algo/Sort.h"
 #include "Containers/Queue.h"
-#include "Containers/StringView.h"
 #include "Containers/StringFwd.h"
+#include "Containers/StringView.h"
 #include "HAL/PlatformProcess.h"
 #include "HAL/RunnableThread.h"
 #include "Logging/LogMacros.h"
 #include "Misc/ConfigCacheIni.h"
+#include "Misc/FileHelper.h"
 #include "Misc/Guid.h"
+#include "Misc/Parse.h"
 #include "Misc/PathViews.h"
 #include "Misc/ScopeLock.h"
+#include "Misc/ScopeRWLock.h"
 #include "Misc/StringBuilder.h"
+
+// TraceServices
+#include "Common/SymbolStringAllocator.h"
 #include "TraceServices/Model/AnalysisSession.h"
-#include "Misc/FileHelper.h"
-#include "Misc/Parse.h"
-#include <atomic>
 
-/////////////////////////////////////////////////////////////////////
-namespace TraceServices {
+#include <atomic>
 
-/////////////////////////////////////////////////////////////////////
 DEFINE_LOG_CATEGORY(LogPsymResolver);
 
 /////////////////////////////////////////////////////////////////////
 
-class FPsymSymbolStringAllocator
+namespace TraceServices
 {
-public:
-	FPsymSymbolStringAllocator(ILinearAllocator& InAllocator, uint32 InBlockSize)
-		: Allocator(InAllocator)
-		, BlockSize(InBlockSize)
-	{}
 
-	const TCHAR* Store(const TCHAR* InString)
-	{
-		return Store(FStringView(InString));
-	}
-
-	const TCHAR* Store(const FStringView InString)
-	{
-		const uint32 StringSize = InString.Len() + 1;
-		check(StringSize <= BlockSize);
-		if (StringSize > BlockRemaining)
-		{
-			Block = (TCHAR*)Allocator.Allocate(BlockSize * sizeof(TCHAR));
-			BlockRemaining = BlockSize;
-			++BlockUsed;
-		}
-		const uint32 CopiedSize = InString.CopyString(Block, BlockRemaining - 1, 0);
-		check(StringSize == CopiedSize + 1);
-		Block[StringSize - 1] = TEXT('\0');
-		BlockRemaining -= StringSize;
-		const TCHAR* OutString = Block;
-		Block += StringSize;
-		return OutString;
-	}
-
-private:
-	ILinearAllocator& Allocator;
-	TCHAR* Block = nullptr;
-	uint32 BlockSize;
-	uint32 BlockRemaining = 0;
-	uint32 BlockUsed = 0;
-};
+static const TCHAR* GUnknownModuleTextPsym = TEXT("Unknown");
 
 /////////////////////////////////////////////////////////////////////
 FPsymResolver::FPsymResolver(IAnalysisSession& InSession, IResolvedSymbolFilter& InSymbolFilter)
@@ -78,24 +44,15 @@ FPsymResolver::FPsymResolver(IAnalysisSession& InSession, IResolvedSymbolFilter&
 	, Session(InSession)
 	, SymbolFilter(InSymbolFilter)
 {
-	// Setup search paths. The SearchPaths array is a priority stack, which
-	// means paths are searched in reversed order.
-	// 1. Any new paths entered by the user this session
+	// Setup search paths. Paths are searched in the following order:
+	// 1. Any new path entered by the user this session
 	// 2. Path of the executable (if available)
-	// 3. Paths from UE_INSIGHTS_SYMBOLPATH
-	// 4. Paths from the user configuration file
-
-	// Paths from configuration
-	FSymbolHelper::GetPathsFromConfiguration(LogPsymResolver, SymbolSearchPaths);
+	// 3. Paths from UE_INSIGHTS_SYMBOL_PATH environment variable
+	// 4. Paths from _NT_SYMBOL_PATH environment variable
+	// 5. Paths from the user configuration file
 
-	// Paths from environment
-	FString SymbolPathEnvVar =  FPlatformMisc::GetEnvironmentVariable(TEXT("UE_INSIGHTS_SYMBOL_PATH"));
-	UE_LOG(LogPsymResolver, Log, TEXT("UE_INSIGHTS_SYMBOL_PATH: '%s'"), *SymbolPathEnvVar);
-	FString SymbolPathPart;
-	while (SymbolPathEnvVar.Split(TEXT(";"), &SymbolPathPart, &SymbolPathEnvVar))
-	{
-		SymbolSearchPaths.Emplace(SymbolPathPart);
-	}
+	// Get the pre-configured symbol search paths (environment variables + user configuration file).
+	FSymbolHelper::GetSymbolSearchPaths(LogPsymResolver, ConfigSymbolSearchPaths);
 
 	Start();
 }
@@ -124,33 +81,58 @@ void FPsymResolver::QueueModuleLoad(const uint8* ImageId, uint32 ImageIdSize, FM
 		Module->Base, Module->Size, Session.StoreString(ModuleName), Session.StoreString(Module->FullName),
 		Module, TArray(ImageId, ImageIdSize)
 	});
+	Algo::Sort(LoadedModules, [](const FModuleEntry& Lhs, const FModuleEntry& Rhs) { return Lhs.Base < Rhs.Base; });
 
-	// Queue up module to have symbols loaded
-	LoadSymbolsQueue.Enqueue(FQueuedModule{ Module, nullptr, LoadedModules[Index].ImageId});
+	// Reset stats
+	Module->Stats.Discovered.store(0u);
+	Module->Stats.Resolved.store(0u);
+	Module->Stats.Failed.store(0u);
+	Module->Stats.Available.store(0u);
 
-	// Sort list according to base address
-	Algo::Sort(LoadedModules, [](const FModuleEntry& Lhs, const FModuleEntry& Rhs) { return Lhs.Base < Rhs.Base; });
+#if 0
+	// TODO: search for psym file in the symbols search paths
+	const TCHAR* PsymPath = nullptr;
+
+	Module->Status.store(EModuleStatus::Pending);
+
+	// Queue up module to have symbols loaded
+	LoadSymbolsQueue.Enqueue(FQueuedModule{ Module, PsymPath, LoadedModules[Index].ImageId });
+#endif
 
 	++ModulesDiscovered;
 }
 
 /////////////////////////////////////////////////////////////////////
-void FPsymResolver::QueueModuleReload(const FModule* Module, const TCHAR* InPath, TFunction<void(SymbolArray&)> ResolveOnSuccess)
+void FPsymResolver::QueueModuleReload(FModule* Module, const TCHAR* InPath, TFunction<void(SymbolArray&)> ResolveOnSuccess)
 {
+	check(Module != nullptr);
+
 	FScopeLock _(&ModulesCs);
-	const uint64 ModuleBase = Module->Base;
-	const FModuleEntry* Entry = LoadedModules.FindByPredicate([ModuleBase](const FModuleEntry& Entry) { return Entry.Base == ModuleBase; });
-	if (Entry)
+
+	const FModuleEntry* Entry = GetEntryForModule(Module);
+	if (!Entry)
 	{
-		// Reset stats for reloaded module.
-		Entry->Module->Stats.Discovered.store(0u);
-		Entry->Module->Stats.Resolved.store(0u);
-		Entry->Module->Stats.Failed.store(0u);
+		return;
+	}
+	check(Entry->Module == Module);
 
-		Entry->Module->Status.store(EModuleStatus::Pending);
-		LoadSymbolsQueue.Enqueue(FQueuedModule{Module, Session.StoreString(InPath), TArrayView<const uint8>(Entry->ImageId)});
+	// No use in trying reload already loaded modules
+	if (Module->Status == EModuleStatus::Loaded)
+	{
+		return;
 	}
 
+	// Reset stats
+	Module->Stats.Discovered.store(0u);
+	Module->Stats.Resolved.store(0u);
+	Module->Stats.Failed.store(0u);
+	Module->Stats.Available.store(0u);
+
+	Module->Status.store(EModuleStatus::Pending);
+
+	// Queue up module to have symbols re-loaded
+	LoadSymbolsQueue.Enqueue(FQueuedModule{ Module, Session.StoreString(InPath), TArrayView<const uint8>(Entry->ImageId) });
+
 	SymbolArray SymbolsToResolve;
 	ResolveOnSuccess(SymbolsToResolve);
 	for(TTuple<uint64, FResolvedSymbol*> Pair : SymbolsToResolve)
@@ -196,8 +178,12 @@ bool FPsymResolver::HasFinishedResolving() const
 /////////////////////////////////////////////////////////////////////
 void FPsymResolver::EnumerateSymbolSearchPaths(TFunctionRef<void(FStringView Path)> Callback) const
 {
-	FScopeLock _(&SymbolSearchPathsLock);
-	Algo::ForEach(SymbolSearchPaths, Callback);
+	{
+		FReadScopeLock _(CustomSymbolSearchPathsLock);
+		Algo::ForEach(CustomSymbolSearchPaths, Callback);
+	}
+
+	Algo::ForEach(ConfigSymbolSearchPaths, Callback);
 }
 
 /////////////////////////////////////////////////////////////////////
@@ -269,22 +255,36 @@ void FPsymResolver::UpdateResolvedSymbol(FResolvedSymbol& Symbol, ESymbolQueryRe
 	Symbol.Result.store(Result, std::memory_order_release);
 }
 
+/////////////////////////////////////////////////////////////////////
+const FPsymResolver::FModuleEntry* FPsymResolver::GetEntryForModule(const FModule* Module) const
+{
+	const int32 EntryIdx = Algo::BinarySearchBy(LoadedModules, Module->Base, [](const FModuleEntry& Entry) { return Entry.Base; });
+	if (EntryIdx != INDEX_NONE)
+	{
+		return &LoadedModules[EntryIdx];
+	}
+	return nullptr;
+}
+
 /////////////////////////////////////////////////////////////////////
 const FPsymResolver::FModuleEntry* FPsymResolver::GetModuleForAddress(uint64 Address) const
 {
 	const int32 EntryIdx = Algo::LowerBoundBy(LoadedModules, Address, [](const FModuleEntry& Entry) { return Entry.Base; }) - 1;
-	if (EntryIdx < 0 || EntryIdx >= LoadedModules.Num())
+	if (EntryIdx >= 0 && EntryIdx < LoadedModules.Num())
 	{
-		return nullptr;
+		return &LoadedModules[EntryIdx];
 	}
-
-	return &LoadedModules[EntryIdx];
+	return nullptr;
 }
 
 /////////////////////////////////////////////////////////////////////
 
-static inline bool IsWhiteSpace(TCHAR Value) { return Value == TCHAR(' ') || Value == TCHAR('\n') || Value == TCHAR('\r'); }
+static inline bool IsWhiteSpace(TCHAR Value)
+{
+	return Value == TCHAR(' ') || Value == TCHAR('\n') || Value == TCHAR('\r');
+}
 
+/////////////////////////////////////////////////////////////////////
 static FString NextToken(const FStringView& InStringView, int32& InIndex)
 {
 	FString Result;
@@ -308,202 +308,244 @@ static FString NextToken(const FStringView& InStringView, int32& InIndex)
 	return Result;
 }
 
-// [REDACTED_PROFILE_URL]/breakpad/blob/master/docs/symbol_files.md
-//
-// Prefix	: Info								   : Number of spaces
-// ------------------------------------------------------------------
-// MODULE	: operatingsystem architecture id name : 4
-// FILE		: number name						   : 2
-// FUNC m	: address size parameter_size name	   : 5
-// FUNC		: address size parameter_size name	   : 4
-// address	: size line filenum					   : 3
-// PUBLIC m : address parameter_size name		   : 4
-// PUBLIC	: address parameter_size name		   : 3
-// STACK	:									   : 0 // Ignore
-// INFO		:									   : 0 // Ignore
-
+/////////////////////////////////////////////////////////////////////
 void FPsymResolver::LoadModuleSymbols(const FModule* Module, const TCHAR* Path, const TArrayView<const uint8> ImageId)
 {
 	check(Module);
-	if (Path)
-	{
-		// Find the module entry
-		FScopeLock _(&ModulesCs);
-		const int32 EntryIdx = Algo::BinarySearchBy(LoadedModules, Module->Base, [](const FModuleEntry& Entry) { return Entry.Base; });
-		check(EntryIdx != INDEX_NONE);
-		const FModuleEntry& Entry = LoadedModules[EntryIdx];
-
-		uint64 BaseAddress = Module->Base;
 
-		FPsymSymbolStringAllocator StringAllocator(Session.GetLinearAllocator(), (2 << 12) / sizeof(TCHAR));
+	if (!Path)
+	{
+		const_cast<FModule*>(Module)->Status.store(EModuleStatus::NotFound);
+		return;
+	}
 
-		constexpr uint32 BuildIdSize = 16;
-		uint8 BuildId[BuildIdSize] = { 0 };
-		auto LineVisitor = [this, BaseAddress, &BuildId, BuildIdSize, Entry, &StringAllocator](FStringView Line)
+	// Find the module entry
+	FScopeLock _(&ModulesCs);
+	const FModuleEntry* EntryPtr = GetEntryForModule(Module);
+	check(EntryPtr && EntryPtr->Module == Module);
+	const FModuleEntry& Entry = *EntryPtr;
+
+	Entry.Module->Stats.Available.store(0u);
+
+	uint64 BaseAddress = Module->Base;
+
+	FSymbolStringAllocator StringAllocator(Session.GetLinearAllocator(), (64 << 10) / sizeof(TCHAR)); // using 64 KiB page size
+
+	// [REDACTED_PROFILE_URL]/breakpad/blob/master/docs/symbol_files.md
+	//
+	// Prefix	: Info								   : Number of spaces
+	// ------------------------------------------------------------------
+	// MODULE	: operatingsystem architecture id name : 4
+	// FILE		: number name						   : 2
+	// FUNC m	: address size parameter_size name	   : 5
+	// FUNC		: address size parameter_size name	   : 4
+	// address	: size line filenum					   : 3
+	// PUBLIC m : address parameter_size name		   : 4
+	// PUBLIC	: address parameter_size name		   : 3
+	// STACK	:									   : 0 // Ignore
+	// INFO		:									   : 0 // Ignore
+
+	constexpr uint32 BuildIdSize = 16;
+	uint8 BuildId[BuildIdSize] = { 0 };
+	auto LineVisitor = [this, BaseAddress, &BuildId, BuildIdSize, Entry, &StringAllocator](FStringView Line)
+		{
+			int32 Index = 0;
+			FString Command = NextToken(Line, Index);
+			if (Command == TEXT("MODULE"))
 			{
-				int32 Index = 0;
-				FString Command = NextToken(Line, Index);
-				if (Command == TEXT("MODULE"))
+				FString OS = NextToken(Line, Index);
+				FString Architecture = NextToken(Line, Index);
+				FString BuildIdStr = NextToken(Line, Index);
+				int i = 0;
+				for (const TCHAR* Ptr = *BuildIdStr; *Ptr && *(Ptr+1) && i < BuildIdSize; Ptr+=2, i++)
 				{
-					FString OS = NextToken(Line, Index);
-					FString Architecture = NextToken(Line, Index);
-					FString BuildIdStr = NextToken(Line, Index);
-					int i = 0;
-					for (const TCHAR* Ptr = *BuildIdStr; *Ptr && *(Ptr+1) && i < BuildIdSize; Ptr+=2, i++)
-					{
-						BuildId[i] = (uint8)((FParse::HexDigit(*Ptr) << 4) + (FParse::HexDigit(*(Ptr+1))));
-					}
-					FString Name = NextToken(Line, Index);
+					BuildId[i] = (uint8)((FParse::HexDigit(*Ptr) << 4) + (FParse::HexDigit(*(Ptr+1))));
 				}
-				else if (Command == TEXT("FILE"))
-				{
-					FString FileIndex = NextToken(Line, Index);
-					FString FileName = NextToken(Line, Index);
+				FString Name = NextToken(Line, Index);
+			}
+			else if (Command == TEXT("FILE"))
+			{
+				FString FileIndex = NextToken(Line, Index);
+				FString FileName = NextToken(Line, Index);
 
-					PsymSourceFiles.Add(FCString::Atoi(*FileIndex), StringAllocator.Store(*FileName));
-				}
-				else if (Command == TEXT("FUNC"))
-				{
-					constexpr int32 MaxSymbolLength = 255;
-
-					FString Address = NextToken(Line, Index);
-					if (Address == "m")
-					{
-						Address = NextToken(Line, Index);
-					}
-
-					FString Size = NextToken(Line, Index);
-					FString ParamSize = NextToken(Line, Index);
-
-					// find first non-whitespace char
-					while (Index < Line.Len() && IsWhiteSpace(Line[Index]))
-					{
-						Index++;
-					}
-					FString Name = Index == Line.Len() ? FString(TEXT("")) : FString(Line.RightChop(Index).Left(MaxSymbolLength));
-
- 					++Entry.Module->Stats.Discovered;
-					PsymSymbols.Add(FPsymSymbol(FParse::HexNumber64(*Address) + BaseAddress, FParse::HexNumber(*Size), StringAllocator.Store(*Name)));
-				}
-				else if (Command == TEXT("PUBLIC"))
-				{
-					FString Address = NextToken(Line, Index);
-					if (Address == "m")
-					{
-						Address = NextToken(Line, Index);
-					}
-					FString ParamSize = NextToken(Line, Index);
-					FString Name = NextToken(Line, Index);
-
-					++Entry.Module->Stats.Discovered;
-					PsymSymbols.Add(FPsymSymbol(FParse::HexNumber64(*Address) + BaseAddress, FParse::HexNumber(*ParamSize), StringAllocator.Store(*Name)));
-				}
-				else if (Command == TEXT("STACK"))
+				PsymSourceFiles.Add(FCString::Atoi(*FileIndex), StringAllocator.Store(*FileName));
+			}
+			else if (Command == TEXT("FUNC"))
+			{
+				constexpr int32 MaxSymbolLength = 255;
+
+				FString Address = NextToken(Line, Index);
+				if (Address == "m")
 				{
-					// ignore
+					Address = NextToken(Line, Index);
 				}
-				else if (Command == TEXT("INFO"))
+
+				FString Size = NextToken(Line, Index);
+				FString ParamSize = NextToken(Line, Index);
+
+				// find first non-whitespace char
+				while (Index < Line.Len() && IsWhiteSpace(Line[Index]))
 				{
-					// ignore
+					Index++;
 				}
-				else
+				FString Name = Index == Line.Len() ? FString(TEXT("")) : FString(Line.RightChop(Index).Left(MaxSymbolLength));
+
+				++Entry.Module->Stats.Available;
+				PsymSymbols.Add(FPsymSymbol(FParse::HexNumber64(*Address) + BaseAddress, FParse::HexNumber(*Size), StringAllocator.Store(*Name)));
+			}
+			else if (Command == TEXT("PUBLIC"))
+			{
+				FString Address = NextToken(Line, Index);
+				if (Address == "m")
 				{
-					FString Size = NextToken(Line, Index);
-					FString LineNumber = NextToken(Line, Index);
-					FString FileNumber = NextToken(Line, Index);
-					
-					PsymSourceLines.Add(FPsymLine(FParse::HexNumber64(*Command) + BaseAddress, FParse::HexNumber(*Size), FCString::Atoi(*LineNumber), FCString::Atoi(*FileNumber)));
+					Address = NextToken(Line, Index);
 				}
-			};
-
-		FFileHelper::LoadFileToStringWithLineVisitor(Path, LineVisitor);
+				FString ParamSize = NextToken(Line, Index);
+				FString Name = NextToken(Line, Index);
 
-		Algo::Sort(PsymSymbols, [](const FPsymSymbol& Left, const FPsymSymbol& Right)
+				++Entry.Module->Stats.Available;
+				PsymSymbols.Add(FPsymSymbol(FParse::HexNumber64(*Address) + BaseAddress, FParse::HexNumber(*ParamSize), StringAllocator.Store(*Name)));
+			}
+			else if (Command == TEXT("STACK"))
+			{
+				// ignore
+			}
+			else if (Command == TEXT("INFO"))
 			{
-				return Left.Address < Right.Address;
-			});
-		Algo::Sort(PsymSourceLines, [](const FPsymLine& Left, const FPsymLine& Right)
+				// ignore
+			}
+			else
 			{
-				return Left.Address < Right.Address;
-			});
+				FString Size = NextToken(Line, Index);
+				FString LineNumber = NextToken(Line, Index);
+				FString FileNumber = NextToken(Line, Index);
 
-		TStringBuilder<256> StatusMessage;
-		EModuleStatus Status;
+				PsymSourceLines.Add(FPsymLine(FParse::HexNumber64(*Command) + BaseAddress, FParse::HexNumber(*Size), FCString::Atoi(*LineNumber), FCString::Atoi(*FileNumber)));
+			}
+		};
+
+	FFileHelper::LoadFileToStringWithLineVisitor(Path, LineVisitor);
 
-		// Only check the first 16 bytes of the BuildId as the psym generator appears to add a trailing 0.
-		if (FMemory::Memcmp(Entry.ImageId.GetData(), BuildId, FMath::Min<int>(Entry.ImageId.Num(), BuildIdSize)) != 0)
+	Algo::Sort(PsymSymbols, [](const FPsymSymbol& Left, const FPsymSymbol& Right)
 		{
-			StatusMessage.Appendf(TEXT("Build ID for psym %s does not match that of trace. Is this the correct psym for %s?"), Module->Name, Path);
-			Status = EModuleStatus::VersionMismatch;
-		}
-		else
+			return Left.Address < Right.Address;
+		});
+	Algo::Sort(PsymSourceLines, [](const FPsymLine& Left, const FPsymLine& Right)
 		{
-			StatusMessage.Appendf(TEXT("Loaded symbols for %s from %s."), Module->Name, Path);
-			Status = EModuleStatus::Loaded;
-			++ModulesLoaded;
-		}
+			return Left.Address < Right.Address;
+		});
 
-		// Make the status visible to the world
-		Entry.Module->StatusMessage = Session.StoreString(StatusMessage.ToView());
-		Entry.Module->Status.store(Status);
+	TStringBuilder<256> StatusMessage;
+	EModuleStatus Status;
+
+	// Only check the first 16 bytes of the BuildId as the psym generator appears to add a trailing 0.
+	if (FMemory::Memcmp(Entry.ImageId.GetData(), BuildId, FMath::Min<int>(Entry.ImageId.Num(), BuildIdSize)) != 0)
+	{
+		StatusMessage.Appendf(TEXT("Build ID of psym does not match that of trace! Is this the correct psym (%s) for module %s?"), Path, Module->Name);
+		Status = EModuleStatus::VersionMismatch;
+		++ModulesFailed;
 	}
+	else
+	{
+		StatusMessage.Append(Path);
+		Status = EModuleStatus::Loaded;
+		++ModulesLoaded;
+	}
+
+	// Make the status visible to the world
+	Entry.Module->StatusMessage = Session.StoreString(StatusMessage.ToView());
+	Entry.Module->Status.store(Status);
+
+	uint64 BytesAllocated = StringAllocator.GetAllocatedSize();
+	uint64 BytesUsed = StringAllocator.GetUsedSize();
+	uint64 BytesWasted = BytesAllocated - BytesUsed;
+	UE_LOG(LogPsymResolver, VeryVerbose, TEXT("String allocator: %.02f KiB used (+ %.02f KiB wasted) in %d blocks"),
+		(double)BytesUsed / 1024.0,
+		(double)BytesWasted / 1024.0,
+		StringAllocator.GetNumAllocatedBlocks());
 }
 
 /////////////////////////////////////////////////////////////////////
-
 void FPsymResolver::ResolveSymbol(uint64 Address, FResolvedSymbol& Target)
 {
+	const ESymbolQueryResult PreviousResult = Target.Result.load();
+	if (PreviousResult == ESymbolQueryResult::OK)
+	{
+		return;
+	}
+
 	const FModuleEntry* Entry = GetModuleForAddress(Address);
-	if (Entry)
+	if (!Entry)
+	{
+		UE_LOG(LogPsymResolver, Warning, TEXT("No module mapped to address 0x%016llx."), Address);
+		UpdateResolvedSymbol(Target,
+			ESymbolQueryResult::NotLoaded,
+			GUnknownModuleTextPsym,
+			GUnknownModuleTextPsym,
+			GUnknownModuleTextPsym,
+			0);
+		SymbolFilter.Update(Target);
+		return;
+	}
+
+	FModule::SymbolStats& ModuleStats = Entry->Module->Stats;
+	++ModuleStats.Discovered;
+
+	if (Entry->Module->Status == EModuleStatus::Loaded)
 	{
-		if (Entry->Module->Status == EModuleStatus::Loaded)
+		int32 SymbolIndex = Algo::UpperBoundBy(PsymSymbols, Address, &FPsymSymbol::Address) - 1;
+		if (PsymSymbols.IsValidIndex(SymbolIndex))
 		{
-			int32 SymbolIndex = Algo::UpperBoundBy(PsymSymbols, Address, &FPsymSymbol::Address) - 1;
-			if (PsymSymbols.IsValidIndex(SymbolIndex))
+			int32 LineIndex = Algo::UpperBoundBy(PsymSourceLines, Address, &FPsymLine::Address) - 1;
+
+			int32 LineNumber = 0;
+			const TCHAR* FileName = TEXT("?");
+			if (PsymSourceLines.IsValidIndex(LineIndex))
 			{
-				int32 LineIndex = Algo::UpperBoundBy(PsymSourceLines, Address, &FPsymLine::Address) - 1;
+				LineNumber = PsymSourceLines[LineIndex].LineNumber;
+				int32 FileIndex = PsymSourceLines[LineIndex].FileIndex;
 
-				int32 LineNumber = 0;
-				const TCHAR* FileName = TEXT("?");
-				if (PsymSourceLines.IsValidIndex(LineIndex))
+				const TCHAR** FileNamePtr = PsymSourceFiles.Find(FileIndex);
+				if (FileNamePtr)
 				{
-					LineNumber = PsymSourceLines[LineIndex].LineNumber;
-					int32 FileIndex = PsymSourceLines[LineIndex].FileIndex;
-				
-					const TCHAR** FileNamePtr = PsymSourceFiles.Find(FileIndex);
-					if (FileNamePtr)
-					{
-						FileName = *FileNamePtr;
-					}
+					FileName = *FileNamePtr;
 				}
-
-				UpdateResolvedSymbol(Target,
-					ESymbolQueryResult::OK,
-					Entry->Name,
-					PsymSymbols[SymbolIndex].Name,
-					FileName,
-					static_cast<uint16>(LineNumber));
-
-				SymbolFilter.Update(Target);
-				return;
 			}
 
-			++Entry->Module->Stats.Failed;
-
 			UpdateResolvedSymbol(Target,
-				ESymbolQueryResult::NotFound,
-				TEXT("?"),
-				TEXT("?"),
-				TEXT("?"),
-				0);
+				ESymbolQueryResult::OK,
+				Entry->Name,
+				PsymSymbols[SymbolIndex].Name,
+				FileName,
+				static_cast<uint16>(LineNumber));
 			SymbolFilter.Update(Target);
+			++ModuleStats.Resolved;
 		}
 		else
 		{
-			++Entry->Module->Stats.Failed;
+			UpdateResolvedSymbol(Target,
+				ESymbolQueryResult::NotFound,
+				Entry->Name,
+				GUnknownModuleTextPsym,
+				GUnknownModuleTextPsym,
+				0);
+			SymbolFilter.Update(Target);
+			++ModuleStats.Failed;
 		}
 	}
+	else
+	{
+		UpdateResolvedSymbol(Target,
+			ESymbolQueryResult::NotLoaded,
+			Entry->Name,
+			GUnknownModuleTextPsym,
+			GUnknownModuleTextPsym,
+			0);
+		SymbolFilter.Update(Target);
+		++ModuleStats.Failed;
+	}
 }
+
 /////////////////////////////////////////////////////////////////////
 
 } // namespace TraceServices
diff --git a/Engine/Source/Developer/TraceServices/Private/Model/PsymResolver.h b/Engine/Source/Developer/TraceServices/Private/Model/PsymResolver.h
--- a/Engine/Source/Developer/TraceServices/Private/Model/PsymResolver.h
+++ b/Engine/Source/Developer/TraceServices/Private/Model/PsymResolver.h
@@ -23,11 +23,12 @@ class FPsymResolver : public FRunnable
 public:
 	typedef TArray<TTuple<uint64, FResolvedSymbol*>> SymbolArray;
 
-	void Start();
 	FPsymResolver(IAnalysisSession& InSession, IResolvedSymbolFilter& InSymbolFilter);
 	~FPsymResolver();
+
+	void Start();
 	void QueueModuleLoad(const uint8* ImageId, uint32 ImageIdSize, FModule* Module);
-	void QueueModuleReload(const FModule* Module, const TCHAR* Path, TFunction<void(SymbolArray&)> ResolveOnSuccess);
+	void QueueModuleReload(FModule* Module, const TCHAR* Path, TFunction<void(SymbolArray&)> ResolveOnSuccess);
 	void QueueSymbolResolve(uint64 Address, FResolvedSymbol* Symbol);
 	void OnAnalysisComplete();
 	void GetStats(IModuleProvider::FStats* OutStats) const;
@@ -68,6 +69,7 @@ private:
 
 	void ResolveSymbol(uint64 Address, FResolvedSymbol& Target);
 	void LoadModuleSymbols(const FModule* Module, const TCHAR* Path, const TArrayView<const uint8> ImageId);
+	const FModuleEntry* GetEntryForModule(const FModule* Module) const;
 	const FModuleEntry* GetModuleForAddress(uint64 Address) const;
 
 	mutable FCriticalSection ModulesCs;
@@ -79,8 +81,10 @@ private:
 	std::atomic<uint32> ModulesFailed;
 	std::atomic<uint32> ModulesLoaded;
 
-	mutable FCriticalSection SymbolSearchPathsLock;
-	TArray<FString> SymbolSearchPaths;
+	mutable FRWLock CustomSymbolSearchPathsLock;
+	TArray<FString> CustomSymbolSearchPaths; // search paths added by user from UI in the current session
+	TArray<FString> ConfigSymbolSearchPaths; // search paths specified in environment variables and config files
+
 	bool bRunWorkerThread;
 	bool bDrainThenStop;
 	UPTRINT Handle;
diff --git a/Engine/Source/Developer/TraceServices/Private/Model/SymslibResolver.cpp b/Engine/Source/Developer/TraceServices/Private/Model/SymslibResolver.cpp
--- a/Engine/Source/Developer/TraceServices/Private/Model/SymslibResolver.cpp
+++ b/Engine/Source/Developer/TraceServices/Private/Model/SymslibResolver.cpp
@@ -24,8 +24,12 @@
 #include "Misc/ScopeLock.h"
 #include "Misc/ScopeRWLock.h"
 #include "Misc/StringBuilder.h"
+
+// TraceServices
+#include "Common/SymbolStringAllocator.h"
 #include "TraceServices/Model/AnalysisSession.h"
 #include "TraceServices/Model/Diagnostics.h"
+
 #include <atomic>
 
 #if PLATFORM_WINDOWS
@@ -43,7 +47,8 @@
 
 DEFINE_LOG_CATEGORY(LogSymslib);
 
-namespace TraceServices {
+namespace TraceServices
+{
 
 static const TCHAR* GUnknownModuleTextSymsLib = TEXT("Unknown");
 
@@ -468,49 +473,6 @@ namespace
 
 ////////////////////////////////////////////////////////////////////////////////////////////////////
 
-class FSymbolStringAllocator
-{
-public:
-	FSymbolStringAllocator(ILinearAllocator& InAllocator, uint32 InBlockSize)
-		: Allocator(InAllocator)
-		, BlockSize(InBlockSize)
-	{}
-
-	const TCHAR* Store(const TCHAR* InString)
-	{
-		return Store(FStringView(InString));
-	}
-
-	const TCHAR* Store(const FStringView InString)
-	{
-		const uint32 StringSize = InString.Len() + 1;
-		check(StringSize <= BlockSize);
-		if (StringSize > BlockRemaining)
-		{
-			Block = (TCHAR*) Allocator.Allocate(BlockSize * sizeof(TCHAR));
-			BlockRemaining = BlockSize;
-			++BlockUsed;
-		}
-		const uint32 CopiedSize = InString.CopyString(Block, BlockRemaining - 1, 0);
-		check(StringSize == CopiedSize + 1);
-		Block[StringSize - 1] = TEXT('\0');
-		BlockRemaining -= StringSize;
-		const TCHAR* OutString = Block;
-		Block += StringSize;
-		return OutString;
-	}
-
-private:
-	friend class FSymslibResolver;
-	ILinearAllocator& Allocator;
-	TCHAR* Block = nullptr;
-	uint32 BlockSize;
-	uint32 BlockRemaining = 0;
-	uint32 BlockUsed = 0;
-};
-
-////////////////////////////////////////////////////////////////////////////////////////////////////
-
 FSymslibResolver::FSymslibResolver(IAnalysisSession& InSession, IResolvedSymbolFilter& InSymbolFilter)
 	: Modules(InSession.GetLinearAllocator(), 128)
 	, CancelTasks(false)
@@ -520,36 +482,12 @@ FSymslibResolver::FSymslibResolver(IAnalysisSession& InSession, IResolvedSymbolF
 	// Setup search paths. Paths are searched in the following order:
 	// 1. Any new path entered by the user this session
 	// 2. Path of the executable (if available)
-	// 3. Paths from UE_INSIGHTS_SYMBOL_PATH
-	// 4. Paths from _NT_SYMBOL_PATH
+	// 3. Paths from UE_INSIGHTS_SYMBOL_PATH environment variable
+	// 4. Paths from _NT_SYMBOL_PATH environment variable
 	// 5. Paths from the user configuration file
 
-	auto SplitEnvPaths = [](FStringView EnvVariable, TArray<FString>& OutList)
-	{
-		FString SymbolPathPart, SymbolPathRemainder(EnvVariable);
-		while (SymbolPathRemainder.Split(TEXT(";"), &SymbolPathPart, &SymbolPathRemainder))
-		{
-			OutList.Emplace(SymbolPathPart);
-		}
-		if (!SymbolPathRemainder.IsEmpty())
-		{
-			OutList.Emplace(SymbolPathRemainder);
-		}
-	};
-
-	// Paths from environment
-	const FString InsightsSymbolPath =  FPlatformMisc::GetEnvironmentVariable(TEXT("UE_INSIGHTS_SYMBOL_PATH"));
-	UE_LOG(LogSymslib, Log, TEXT("UE_INSIGHTS_SYMBOL_PATH: %s"), InsightsSymbolPath.IsEmpty() ? TEXT("not set") : *InsightsSymbolPath);
-	SplitEnvPaths(InsightsSymbolPath, ConfigSymbolSearchPaths);
-
-#if PLATFORM_WINDOWS
-	const FString NTSymbolPath = FPlatformMisc::GetEnvironmentVariable(TEXT("_NT_SYMBOL_PATH"));
-	UE_LOG(LogSymslib, Log, TEXT("_NT_SYMBOL_PATH: %s"), NTSymbolPath.IsEmpty() ? TEXT("not set") : *NTSymbolPath);
-	SplitEnvPaths(NTSymbolPath, ConfigSymbolSearchPaths);
-#endif
-
-	// Paths from configuration file
-	FSymbolHelper::GetPathsFromConfiguration(LogSymslib, ConfigSymbolSearchPaths);
+	// Get the pre-configured symbol search paths (environment variables + user configuration file).
+	FSymbolHelper::GetSymbolSearchPaths(LogSymslib, ConfigSymbolSearchPaths);
 
 	UpdateSessionInfo();
 }
@@ -593,6 +531,12 @@ void FSymslibResolver::QueueModuleLoad(const uint8* ImageId, uint32 ImageIdSize,
 
 	++ModulesDiscovered;
 
+	// Reset stats for module.
+	Module->Stats.Discovered.store(0u);
+	Module->Stats.Resolved.store(0u);
+	Module->Stats.Failed.store(0u);
+	Module->Stats.Available.store(0u);
+
 	// Set the Pending state before scheduling the background task (to allow calling code to wait, if needed).
 	Module->Status.store(EModuleStatus::Pending);
 
@@ -619,6 +563,7 @@ void FSymslibResolver::QueueModuleReload(FModule* Module, const TCHAR* Path, TFu
 	{
 		return;
 	}
+	check(Entry->Module == Module);
 
 	// No use in trying reload already loaded modules
 	if (Module->Status == EModuleStatus::Loaded)
@@ -626,6 +571,12 @@ void FSymslibResolver::QueueModuleReload(FModule* Module, const TCHAR* Path, TFu
 		return;
 	}
 
+	// Reset stats for reloaded module.
+	Module->Stats.Discovered.store(0u);
+	Module->Stats.Resolved.store(0u);
+	Module->Stats.Failed.store(0u);
+	Module->Stats.Available.store(0u);
+
 	// Set the Pending state before scheduling the background task (to allow calling code to wait, if needed).
 	EModuleStatus PreviousStatus = Module->Status.exchange(EModuleStatus::Pending);
 	if (PreviousStatus >= EModuleStatus::FailedStatusStart)
@@ -676,13 +627,28 @@ void FSymslibResolver::QueueModuleReload(FModule* Module, const TCHAR* Path, TFu
 			// Ask the caller for a list of symbols that should be resolved, now that module is properly loaded and
 			// resolve them immediately.
 			{
+				FSymbolStringAllocator StringAllocator(Session.GetLinearAllocator(), (64 << 10) / sizeof(TCHAR)); // using 64 KiB page size
+
 				SymbolArray SymbolsToResolve;
-				FSymbolStringAllocator StringAllocator(Session.GetLinearAllocator(), (2 << 12) / sizeof(TCHAR));
 				ResolveOnSuccess(SymbolsToResolve);
 				for (TTuple<uint64, FResolvedSymbol*> Pair : SymbolsToResolve)
 				{
+					if (CancelTasks.load(std::memory_order_relaxed))
+					{
+						break;
+					}
 					ResolveSymbolTracked(Pair.Get<0>(), *Pair.Get<1>(), StringAllocator);
 				}
+
+				uint64 BytesAllocated = StringAllocator.GetAllocatedSize();
+				uint64 BytesUsed = StringAllocator.GetUsedSize();
+				uint64 BytesWasted = BytesAllocated - BytesUsed;
+				UE_LOG(LogSymslib, VeryVerbose, TEXT("String allocator: %.02f KiB used (+ %.02f KiB wasted) in %d blocks"),
+					(double)BytesUsed / 1024.0,
+					(double)BytesWasted / 1024.0,
+					StringAllocator.GetNumAllocatedBlocks());
+				SymbolBytesAllocated.fetch_add(BytesAllocated);
+				SymbolBytesWasted.fetch_add(BytesWasted);
 			}
 
 			// Finally if the cleanup task has already been dispatched, no other resolve request will come in and we
@@ -771,7 +737,7 @@ bool FSymslibResolver::HasFinishedResolving() const
 	}
 	return false;
 }
-	
+
 ////////////////////////////////////////////////////////////////////////////////////////////////////
 
 void FSymslibResolver::EnumerateSymbolSearchPaths(TFunctionRef<void(FStringView Path)> Callback) const
@@ -833,7 +799,8 @@ void FSymslibResolver::ResolveSymbols(TArrayView<FQueuedAddress>& QueuedWork)
 {
 	// Create a local string allocator. We don't use the session string store due to contention when going wide. Since
 	// the ModuleProvider already de-duplicates symbols we do not need this feature from the string store.
-	FSymbolStringAllocator StringAllocator(Session.GetLinearAllocator(), (2 << 12) / sizeof(TCHAR) );
+	FSymbolStringAllocator StringAllocator(Session.GetLinearAllocator(), (64 << 10) / sizeof(TCHAR)); // using 64 KiB page size
+
 	for (const FQueuedAddress& ToResolve : QueuedWork)
 	{
 		if (CancelTasks.load(std::memory_order_relaxed))
@@ -842,12 +809,14 @@ void FSymslibResolver::ResolveSymbols(TArrayView<FQueuedAddress>& QueuedWork)
 		}
 		ResolveSymbolTracked(ToResolve.Address, *ToResolve.Target, StringAllocator);
 	}
-	uint64 BytesAllocated = StringAllocator.BlockUsed * StringAllocator.BlockSize * sizeof(TCHAR);
-	uint64 BytesWasted = StringAllocator.BlockRemaining * sizeof(TCHAR);
+
+	uint64 BytesAllocated = StringAllocator.GetAllocatedSize();
+	uint64 BytesUsed = StringAllocator.GetUsedSize();
+	uint64 BytesWasted = BytesAllocated - BytesUsed;
 	UE_LOG(LogSymslib, VeryVerbose, TEXT("String allocator: %.02f KiB used (+ %.02f KiB wasted) in %d blocks"),
-		(double)(BytesAllocated - BytesWasted) / 1024.0,
-		(double)(BytesWasted) / 1024.0,
-		StringAllocator.BlockUsed);
+		(double)BytesUsed / 1024.0,
+		(double)BytesWasted / 1024.0,
+		StringAllocator.GetNumAllocatedBlocks());
 	SymbolBytesAllocated.fetch_add(BytesAllocated);
 	SymbolBytesWasted.fetch_add(BytesWasted);
 }
@@ -897,6 +866,12 @@ void FSymslibResolver::LoadModuleTracked(FModuleEntry* Entry, FStringView Overri
 {
 	FModule* Module = Entry->Module;
 
+	// Reset stats
+	Module->Stats.Discovered.store(0u);
+	Module->Stats.Resolved.store(0u);
+	Module->Stats.Failed.store(0u);
+	Module->Stats.Available.store(0u);
+
 	Module->Status.store(EModuleStatus::Pending);
 
 	// Build search paths
@@ -969,6 +944,11 @@ void FSymslibResolver::LoadModuleTracked(FModuleEntry* Entry, FStringView Overri
 
 EModuleStatus FSymslibResolver::LoadModule(FModuleEntry* Entry, FStringView SearchPathView, FStringBuilderBase& OutStatusMessage) const
 {
+	FModule::SymbolStats& ModuleStats = Entry->Module->Stats;
+
+	// how many symbols are loaded
+	ModuleStats.Available = 0;
+
 	const FString SearchPath(SearchPathView);
 
 	// temporary memory used for loading
@@ -1021,16 +1001,13 @@ EModuleStatus FSymslibResolver::LoadModule(FModuleEntry* Entry, FStringView Sear
 	int32 WorkerThreadCount = FMath::Max(1, FTaskGraphInterface::Get().GetNumWorkerThreads());
 	Instance->Arenas.SetNum(WorkerThreadCount);
 
-	// how many symbols are loaded
-	std::atomic<uint32> SymbolCount = 0;
-
 	// parse debug info in multiple threads
 	{
 		uint32 LaneSlot = FPlatformTLS::AllocTlsSlot();
 
 		std::atomic<uint32> LaneCount = 0;
 		syms_group_begin_multilane(Group, WorkerThreadCount);
-		ParallelFor(static_cast<int32>(UnitCount), [Instance, Group, LaneSlot, &LaneCount, &SymbolCount](int32 Index)
+		ParallelFor(static_cast<int32>(UnitCount), [Instance, Group, LaneSlot, &LaneCount, &ModuleStats](int32 Index)
 		{
 			SYMS_Arena* Arena;
 			uint32 LaneValue = uint32(reinterpret_cast<intptr_t>(FPlatformTLS::GetTlsValue(LaneSlot)));
@@ -1096,7 +1073,7 @@ EModuleStatus FSymslibResolver::LoadModule(FModuleEntry* Entry, FStringView Sear
 
 			syms_release_scratch(Scratch);
 
-			SymbolCount += static_cast<uint32>(ProcCount);
+			ModuleStats.Available += static_cast<uint32>(ProcCount);
 		});
 		syms_group_end_multilane(Group);
 
@@ -1126,7 +1103,7 @@ EModuleStatus FSymslibResolver::LoadModule(FModuleEntry* Entry, FStringView Sear
 		Instance->StrippedMap = syms_spatial_map_1d_copy(Arena, syms_group_link_name_spatial_map(Group));
 		Instance->StrippedSymbols = StrippedSymbols;
 
-		SymbolCount += (uint32)StrippedInfo.count;
+		ModuleStats.Available += (uint32)StrippedInfo.count;
 	}
 
 	Instance->UnitMap = syms_spatial_map_1d_copy(Arena, syms_group_unit_map(Group));
@@ -1137,7 +1114,7 @@ EModuleStatus FSymslibResolver::LoadModule(FModuleEntry* Entry, FStringView Sear
 
 	// In order to avoid listing all files, use only the last. On dwarf/elf this will be the main binary, on Windows
 	// this will be the pdb (the first is the exe/dll).
-	OutStatusMessage.Appendf(TEXT("%s loaded with %u symbols."), *Files.Last().FilePath, SymbolCount.load());
+	OutStatusMessage.Append(*Files.Last().FilePath);
 	return EModuleStatus::Loaded;
 }
 
diff --git a/Engine/Source/Developer/TraceServices/Private/Model/SymslibResolver.h b/Engine/Source/Developer/TraceServices/Private/Model/SymslibResolver.h
--- a/Engine/Source/Developer/TraceServices/Private/Model/SymslibResolver.h
+++ b/Engine/Source/Developer/TraceServices/Private/Model/SymslibResolver.h
@@ -17,15 +17,18 @@
 #include "Common/PagedArray.h"
 #include "HAL/CriticalSection.h"
 #include "HAL/Platform.h"
+#include "symslib.h"
+
+// TraceServices
 #include "TraceServices/Model/AnalysisSession.h"
 #include "TraceServices/Model/Modules.h"
-#include <atomic>
 
-#include "symslib.h"
+#include <atomic>
 
 DECLARE_LOG_CATEGORY_EXTERN(LogSymslib, Log, All);
 
-namespace TraceServices {
+namespace TraceServices
+{
 
 class FSymslibResolver
 {
@@ -34,6 +37,7 @@ public:
 
 	FSymslibResolver(IAnalysisSession& InSession, IResolvedSymbolFilter& InSymbolFilter);
 	~FSymslibResolver();
+
 	void QueueModuleLoad(const uint8* ImageId, uint32 ImageIdSize, FModule* Module);
 	void QueueModuleReload(FModule* Module, const TCHAR* Path, TFunction<void(SymbolArray&)> ResolveOnSuccess);
 	void QueueSymbolResolve(uint64 Address, FResolvedSymbol* Symbol);
diff --git a/Engine/Source/Developer/TraceServices/Public/TraceServices/Model/Modules.h b/Engine/Source/Developer/TraceServices/Public/TraceServices/Model/Modules.h
--- a/Engine/Source/Developer/TraceServices/Public/TraceServices/Model/Modules.h
+++ b/Engine/Source/Developer/TraceServices/Public/TraceServices/Model/Modules.h
@@ -142,10 +142,11 @@ struct FModule
 	// Statistics about the module
 	struct SymbolStats
 	{
-		std::atomic<uint32>		Discovered;
-		std::atomic<uint32>		Cached;
-		std::atomic<uint32>		Resolved;
-		std::atomic<uint32>		Failed;
+		std::atomic<uint32>		Discovered = 0; // number of symbols discovered in the trace (used symbols)
+		std::atomic<uint32>		Cached = 0;     // number of symbols loaded from cache
+		std::atomic<uint32>		Resolved = 0;   // number of symbols successfully resolved from the loaded module
+		std::atomic<uint32>		Failed = 0;     // number of symbols failed to resolve
+		std::atomic<uint32>		Available = 0;  // number of symbols available in the loaded module's debug information
 	} Stats;
 
 	FModule(const TCHAR* InName, const TCHAR* InFullName, uint64 InBase, uint32 InSize, EModuleStatus InStatus)
@@ -155,7 +156,8 @@ struct FModule
 		, Size(InSize)
 		, Status(InStatus)
 		, StatusMessage(nullptr)
-	{}
+	{
+	}
 };
 
 ////////////////////////////////////////////////////////////////////////////////
@@ -206,8 +208,7 @@ public:
 	virtual FGraphEventRef LoadSymbolsForModuleUsingPath(uint64 Base, const TCHAR* Path) = 0;
 
 	/**
-	 * Gets the search paths used to find debug symbols. Paths are traversed in the reverse order
-	 * order of importance.
+	 * Gets the search paths used to find debug symbols.
 	 * @param Callback Called for each search path.
 	 */
 	virtual void EnumerateSymbolSearchPaths(TFunctionRef<void(FStringView Path)> Callback) const = 0;
