[Insights]
* GPUProfiler: Moved the OldGpuProfilerTraceAnalysis to a separate analyzer class in order not to interfere with the new GPU Profiler implementation.
* GPUProfiler: Disabled merge by name for GPU timers as it is not safe to merge when metadata spec can be attached late.
* GPUProfiler: Added detection and warning for GPU timers de-duplicated because of different metadata spec (breadcrumb spec).
* GPUProfiler: Added debug code to enable logging of processed events (by 3 debug levels).
* GPUProfiler: Fixed processing for BeginBreadcrumb event to not create anymore empty metadata timers.
* CPUProfiler: Disabled merge by name for scoped timers with metadata. Only timers declared using EventSpec event can be merged by name.
* CPUProfiler: Added detection and warning when CPU timers are merged and then de-duplicated because of different metadata spec.
* CPUProfiler: Improved session locking, by acquiring the read/edit lock only when actually using the providers.
* CPUProfiler: Added debug code to detect timer name changes (toggled off by default by CPUPROFILER_DETECT_TIMER_NAME_CHANGE).
* CPUProfiler: Added debug code to enable logging of processed events (by 3 debug levels).
* TimingProfiler: Added debug code to detect timer changes (toggled off by default by TIMING_PROFILER_DETECT_TIMER_CHANGES).
* TimingProfiler: Added debug logging for code paths trying to access invalid metadata timer id or spec id (disabled by default).
* TimingProfiler: Added detection and automatic merge for multiple identical metadata specs.

diff --git a/Engine/Source/Developer/TraceServices/Private/Analyzers/CpuProfilerTraceAnalysis.cpp b/Engine/Source/Developer/TraceServices/Private/Analyzers/CpuProfilerTraceAnalysis.cpp
--- a/Engine/Source/Developer/TraceServices/Private/Analyzers/CpuProfilerTraceAnalysis.cpp
+++ b/Engine/Source/Developer/TraceServices/Private/Analyzers/CpuProfilerTraceAnalysis.cpp
@@ -2,21 +2,27 @@
 
 #include "CpuProfilerTraceAnalysis.h"
 
-#include "AnalysisServicePrivate.h"
 #include "CborReader.h"
 #include "CborWriter.h"
-#include "Common/Utils.h"
 #include "HAL/LowLevelMemTracker.h"
-#include "Model/ThreadsPrivate.h"
-#include "Model/MonotonicTimeline.h"
 #include "Serialization/MemoryReader.h"
 #include "Serialization/MemoryWriter.h"
+
+// TraceServices
+#include "AnalysisServicePrivate.h"
+#include "Common/Utils.h"
+#include "Model/MonotonicTimeline.h"
+#include "Model/ThreadsPrivate.h"
 #include "TraceServices/Model/TimingProfiler.h"
 #include "TraceServices/Utils.h"
 
 #define CPUPROFILER_DEBUG_LOGF(Format, ...) //{ if (ThreadState.ThreadId == 2) FPlatformMisc::LowLevelOutputDebugStringf(Format, __VA_ARGS__); }
 #define CPUPROFILER_DEBUG_BEGIN_EVENT(Time, Event) { ++TotalScopeCount; }
 #define CPUPROFILER_DEBUG_END_EVENT(Time)
+#define CPUPROFILER_LOG_API_L1(Format, ...) UE_LOG(LogTraceServices, Log, Format, ##__VA_ARGS__)
+#define CPUPROFILER_LOG_API_L2(Format, ...) //UE_LOG(LogTraceServices, Verbose, Format, ##__VA_ARGS__)
+#define CPUPROFILER_LOG_API_L3(Format, ...) //UE_LOG(LogTraceServices, VeryVerbose, Format, ##__VA_ARGS__)
+#define CPUPROFILER_DETECT_TIMER_NAME_CHANGE 0
 
 namespace TraceServices
 {
@@ -47,15 +53,16 @@ void FCpuProfilerAnalyzer::OnAnalysisBegin(const FOnAnalysisContext& Context)
 {
 	auto& Builder = Context.InterfaceBuilder;
 
-	Builder.RouteEvent(RouteId_EventSpec, "CpuProfiler", "EventSpec");
-	Builder.RouteEvent(RouteId_EndThread, "CpuProfiler", "EndThread");
+	Builder.RouteEvent(RouteId_EventSpec,    "CpuProfiler", "EventSpec");
+	Builder.RouteEvent(RouteId_EndThread,    "CpuProfiler", "EndThread");
 	Builder.RouteEvent(RouteId_EventBatchV3, "CpuProfiler", "EventBatchV3"); // added in UE 5.6
 	Builder.RouteEvent(RouteId_EventBatchV2, "CpuProfiler", "EventBatchV2"); // backward compatibility, added in UE 5.1, removed in 5.6
-	Builder.RouteEvent(RouteId_EventBatch, "CpuProfiler", "EventBatch"); // backward compatibility; removed in UE 5.1
-	Builder.RouteEvent(RouteId_EndCapture, "CpuProfiler", "EndCapture"); // backward compatibility; removed in UE 5.1
+	Builder.RouteEvent(RouteId_EventBatch,   "CpuProfiler", "EventBatch"); // backward compatibility; removed in UE 5.1
+	Builder.RouteEvent(RouteId_EndCapture,   "CpuProfiler", "EndCapture"); // backward compatibility; removed in UE 5.1
 	Builder.RouteEvent(RouteId_MetadataSpec, "CpuProfiler", "MetadataSpec");
-	Builder.RouteEvent(RouteId_Metadata, "CpuProfiler", "Metadata");
-	Builder.RouteLoggerEvents(RouteId_CpuScope, "Cpu", true);
+	Builder.RouteEvent(RouteId_Metadata,     "CpuProfiler", "Metadata");
+
+	Builder.RouteLoggerEvents(RouteId_CpuScope, "Cpu", true); // scoped trace events
 }
 
 ////////////////////////////////////////////////////////////////////////////////////////////////////
@@ -70,6 +77,7 @@ void FCpuProfilerAnalyzer::OnAnalysisEnd()
 
 		if (ThreadState.LastCycle != ~0ull) // if EndThread is not received
 		{
+			FAnalysisSessionEditScope _(Session);
 			DispatchRemainingPendingEvents(ThreadState);
 			EndOpenEvents(ThreadState, std::numeric_limits<double>::infinity());
 		}
@@ -94,6 +102,11 @@ void FCpuProfilerAnalyzer::OnAnalysisEnd()
 		UE_LOG(LogTraceServices, Warning, TEXT("[CpuProfiler] Possible incorrect use of FCpuProfilerTrace::OutputEventType()!"));
 	}
 
+	if (NumTimerWarnings > 0)
+	{
+		UE_LOG(LogTraceServices, Warning, TEXT("[CpuProfiler] Number of timer/metadata warnings: %u"), NumTimerWarnings);
+	}
+
 	UE_LOG(LogTraceServices, Log, TEXT("[CpuProfiler] Analysis completed (%d threads, %d timers, %d spec ids, %llu scopes, %llu bytes, %.1f bytes/scope)."),
 		ThreadStatesMap.Num(),
 		ScopeNameToTimerIdMap.Num(),
@@ -122,64 +135,21 @@ bool FCpuProfilerAnalyzer::OnEvent(uint16 RouteId, EStyle Style, const FOnEventC
 {
 	LLM_SCOPE_BYNAME(TEXT("Insights/FCpuProfilerAnalyzer"));
 
-	FAnalysisSessionEditScope _(Session);
-
 	const auto& EventData = Context.EventData;
+
 	switch (RouteId)
 	{
+
 	case RouteId_EventSpec:
 	{
-		uint32 SpecId = EventData.GetValue<uint32>("Id");
-
-		const TCHAR* TimerName = nullptr;
-		FString Name;
-		if (EventData.GetString("Name", Name))
-		{
-			TimerName = *Name;
-		}
-		else
-		{
-			uint8 CharSize = EventData.GetValue<uint8>("CharSize");
-			if (CharSize == sizeof(ANSICHAR))
-			{
-				const ANSICHAR* AnsiName = reinterpret_cast<const ANSICHAR*>(EventData.GetAttachment());
-				Name = StringCast<TCHAR>(AnsiName).Get();
-				TimerName = *Name;
-			}
-			else if (CharSize == 0 || CharSize == sizeof(TCHAR)) // 0 for backwards compatibility
-			{
-				TimerName = reinterpret_cast<const TCHAR*>(EventData.GetAttachment());
-			}
-			else
-			{
-				Name = FString::Printf(TEXT("<invalid %u>"), SpecId);
-				TimerName = *Name;
-			}
-		}
-
-		if (TimerName[0] == 0)
-		{
-			Name = FString::Printf(TEXT("<noname %u>"), SpecId);
-			TimerName = *Name;
-		}
-
-		const TCHAR* FileName = nullptr;
-		FString File;
-		uint32 Line = 0;
-		if (EventData.GetString("File", File) && !File.IsEmpty())
-		{
-			FileName = *File;
-			Line = EventData.GetValue<uint32>("Line");
-		}
-
-		DefineTimer(SpecId, Session.StoreString(TimerName), FileName, Line, true);
+		OnEventSpec(Context);
 		break;
 	}
 
 	case RouteId_EndThread:
 	{
 		const uint32 ThreadId = FTraceAnalyzerUtils::GetThreadIdField(Context);
-		FThreadState& ThreadState = GetThreadState(ThreadId);
+		FThreadState& ThreadState = GetOrAddThreadState(ThreadId);
 
 		if (ThreadState.LastCycle == ~0ull)
 		{
@@ -187,13 +157,17 @@ bool FCpuProfilerAnalyzer::OnEvent(uint16 RouteId, EStyle Style, const FOnEventC
 			break;
 		}
 
-		DispatchRemainingPendingEvents(ThreadState);
+		{
+			FAnalysisSessionEditScope _(Session);
+			DispatchRemainingPendingEvents(ThreadState);
+		}
 
 		const uint64 Cycle = EventData.GetValue<uint64>("Cycle", ThreadState.LastCycle); // added in UE 5.4
 		if (Cycle != 0)
 		{
 			ensure(Cycle >= ThreadState.LastCycle);
 			double Timestamp = Context.EventTime.AsSeconds(Cycle);
+			FAnalysisSessionEditScope _(Session);
 			Session.UpdateDurationSeconds(Timestamp);
 			EndOpenEvents(ThreadState, Timestamp);
 		}
@@ -209,7 +183,7 @@ bool FCpuProfilerAnalyzer::OnEvent(uint16 RouteId, EStyle Style, const FOnEventC
 	case RouteId_EventBatchV2: // backward compatibility
 	{
 		const uint32 ThreadId = Context.ThreadInfo.GetId();
-		FThreadState& ThreadState = GetThreadState(ThreadId);
+		FThreadState& ThreadState = GetOrAddThreadState(ThreadId);
 
 		if (ThreadState.LastCycle == ~0ull)
 		{
@@ -225,6 +199,8 @@ bool FCpuProfilerAnalyzer::OnEvent(uint16 RouteId, EStyle Style, const FOnEventC
 		if (ThreadState.LastCycle != 0)
 		{
 			double Timestamp = Context.EventTime.AsSeconds(ThreadState.LastCycle);
+
+			FAnalysisSessionEditScope _(Session);
 			Session.UpdateDurationSeconds(Timestamp);
 		}
 
@@ -236,7 +212,7 @@ bool FCpuProfilerAnalyzer::OnEvent(uint16 RouteId, EStyle Style, const FOnEventC
 	case RouteId_EndCapture: // backward compatibility
 	{
 		const uint32 ThreadId = FTraceAnalyzerUtils::GetThreadIdField(Context);
-		FThreadState& ThreadState = GetThreadState(ThreadId);
+		FThreadState& ThreadState = GetOrAddThreadState(ThreadId);
 
 		if (ThreadState.LastCycle == ~0ull)
 		{
@@ -252,6 +228,7 @@ bool FCpuProfilerAnalyzer::OnEvent(uint16 RouteId, EStyle Style, const FOnEventC
 
 		if (RouteId == RouteId_EndCapture)
 		{
+			FAnalysisSessionEditScope _(Session);
 			DispatchRemainingPendingEvents(ThreadState);
 			if (ThreadState.LastCycle != 0)
 			{
@@ -261,13 +238,11 @@ bool FCpuProfilerAnalyzer::OnEvent(uint16 RouteId, EStyle Style, const FOnEventC
 			}
 			ThreadState.LastCycle = ~0ull;
 		}
-		else
+		else if (ThreadState.LastCycle != 0)
 		{
-			if (ThreadState.LastCycle != 0)
-			{
-				double Timestamp = Context.EventTime.AsSeconds(ThreadState.LastCycle);
-				Session.UpdateDurationSeconds(Timestamp);
-			}
+			double Timestamp = Context.EventTime.AsSeconds(ThreadState.LastCycle);
+			FAnalysisSessionEditScope _(Session);
+			Session.UpdateDurationSeconds(Timestamp);
 		}
 
 		TotalEventSize += BufferSize;
@@ -286,134 +261,38 @@ bool FCpuProfilerAnalyzer::OnEvent(uint16 RouteId, EStyle Style, const FOnEventC
 		break;
 
 	case RouteId_MetadataSpec:
-	{
-		uint32 SpecId = Context.EventData.GetValue<uint32>("Id");
+		OnMetadataSpec(Context);
+		break;
 
-		FString Name, NameFormat;
-		Context.EventData.GetString("Name", Name);
-		Context.EventData.GetString("NameFormat", NameFormat);
-		TArrayView<const uint8> FieldNames = EventData.GetArrayView<uint8>("FieldNames");
+	case RouteId_Metadata:
+	{
+		uint32 MetadataId = Context.EventData.GetValue<uint32>("Id");
+		uint32 SpecId = Context.EventData.GetValue<uint32>("SpecId");
+		TArray<uint8> Metadata(EventData.GetArrayView<uint8>("Metadata"));
 
-		if (Name.Compare(NameFormat) == 0)
-		{
-			Name.Empty();
-		}
+		uint32 TimerId = GetOrAddTimer(SpecId);
 
-		int Index = 0;
-		while (Index < NameFormat.Len() && Index < Name.Len())
+		if (ensure(Metadata.Num() > 0))
 		{
-			if (Name[Index] != NameFormat[Index])
+			// We don't know if this Metadata trace event or a timing event with metadata arrive first, so handle both cases.
+			const uint32* FoundMetadataTimerId = MetadataIdToTimerIdMap.Find(MetadataId);
+			if (FoundMetadataTimerId == nullptr)
 			{
-				break;
+				FAnalysisSessionEditScope _(Session);
+				TimerId = EditableTimingProfilerProvider.AddMetadata(TimerId, MoveTemp(Metadata));
+				MetadataIdToTimerIdMap.Add(MetadataId, TimerId);
 			}
-
-			++Index;
-		}
-
-		if (Index > 1)
-		{
-			NameFormat.MidInline(Index);
-		}
-
-		if (Name.IsEmpty())
-		{
-			const FString Separators = TEXT("% (=");
-			Index = 0;
-			bool bIsInFormatSpecifier = false;
-			while (Index < NameFormat.Len())
-			{
-				int32 SpecIndex = -1;
-				if (Separators.FindChar(NameFormat[Index], SpecIndex))
-				{
-					Name = NameFormat.Left(Index);
-					NameFormat.MidInline(Index);
-					break;
-				}
-
-				++Index;
-			}
-
-			if (Name.IsEmpty())
-			{
-				if (!NameFormat.IsEmpty())
-				{
-					Name = NameFormat;
-					NameFormat.Empty();
-				}
-				else
-				{
-					Name = Session.StoreString(TEXT("Unknown"));
-				}
-			}
-		}
-
-		FMetadataSpec Spec;
-		if (!NameFormat.IsEmpty())
-		{
-			Spec.Format = Session.StoreString(NameFormat);
-		}
-
-		if (FieldNames.Num() > 0)
-		{
-			FMemoryReaderView MemoryReader(FieldNames);
-			FCborReader CborReader(&MemoryReader, ECborEndianness::StandardCompliant);
-			FCborContext CborContext;
-
-			while (CborReader.ReadNext(CborContext))
+			else
 			{
-				if (CborContext.MajorType() == ECborCode::TextString)
-				{
-					FString Field = CborContext.AsString();
-					Spec.FieldNames.Add(Session.StoreString(Field));
-				}
+				// Replace the placeholder metadata added if we received a timing event with this metadata first.
+				FAnalysisSessionEditScope _(Session);
+				EditableTimingProfilerProvider.SetMetadata(*FoundMetadataTimerId, MoveTemp(Metadata), TimerId);
 			}
 		}
-
-		uint32 TimerId = 0;
-		uint32* FoundTimerIdPtr = SpecIdToTimerIdMap.Find(SpecId);
-		if (FoundTimerIdPtr == nullptr)
-		{
-			TimerId = DefineTimer(SpecId, Session.StoreString(*Name), nullptr, 0, false);
-		}
-		else
-		{
-			TimerId = *FoundTimerIdPtr;
-			EditableTimingProfilerProvider.SetTimerName(TimerId, Name);
-		}
-
-		if (Spec.FieldNames.Num() > 0 || Spec.Format != nullptr)
-		{
-			uint32 MetadataSpecId = EditableTimingProfilerProvider.AddMetadataSpec(MoveTemp(Spec));
-			EditableTimingProfilerProvider.SetMetadataSpec(TimerId, MetadataSpecId);
-		}
-
 		break;
 	}
-	case RouteId_Metadata:
-	{
-		uint32 MetadataId = Context.EventData.GetValue<uint32>("Id");
-		uint32 SpecId = Context.EventData.GetValue<uint32>("SpecId");
-		TArray<uint8> Metadata(EventData.GetArrayView<uint8>("Metadata"));
 
-		// We don't know if the Metadata or the BeginMetadataEvent events arrive first so handle both cases.
-		uint32* TimerIdPtr = MetadataIdToTimerIdMap.Find(MetadataId);
-		if (TimerIdPtr == nullptr)
-		{
-			uint32 TimerId = GetTimerId(SpecId);
-			TimerId = EditableTimingProfilerProvider.AddMetadata(TimerId, MoveTemp(Metadata));
-			MetadataIdToTimerIdMap.Add(MetadataId, TimerId);
-		}
-		else
-		{
-			uint32 MetadataTimerId = *TimerIdPtr;
-			uint32 TimerId = GetTimerId(SpecId);
-			// Replace the placeholder metadata added if we received an event with this metadata first.
-			EditableTimingProfilerProvider.SetMetadata(MetadataTimerId, MoveTemp(Metadata), TimerId);
-		}
-
-		break;
-	}
-	}
+	} // switch (RouteId)
 
 	return true;
 }
@@ -422,6 +301,8 @@ bool FCpuProfilerAnalyzer::OnEvent(uint16 RouteId, EStyle Style, const FOnEventC
 
 void FCpuProfilerAnalyzer::ProcessBuffer(const FEventTime& EventTime, FThreadState& ThreadState, const uint8* BufferPtr, uint32 BufferSize)
 {
+	FAnalysisSessionEditScope _(Session);
+
 	uint64 LastCycle = ThreadState.LastCycle;
 
 	CPUPROFILER_DEBUG_LOGF(TEXT("[%u] ProcessBuffer %llu (%.9f)\n"), ThreadState.ThreadId, LastCycle, EventTime.AsSeconds(LastCycle));
@@ -460,7 +341,7 @@ void FCpuProfilerAnalyzer::ProcessBuffer(const FEventTime& EventTime, FThreadSta
 		if (DecodedCycle & 1ull)
 		{
 			uint32 SpecId = IntCastChecked<uint32>(FTraceAnalyzerUtils::Decode7bit(BufferPtr));
-			uint32 TimerId = GetTimerId(SpecId);
+			uint32 TimerId = GetOrAddTimer(SpecId);
 
 			FEventScopeState& ScopeState = ThreadState.ScopeStack.AddDefaulted_GetRef();
 			ScopeState.StartCycle = ActualCycle;
@@ -509,6 +390,8 @@ void FCpuProfilerAnalyzer::ProcessBuffer(const FEventTime& EventTime, FThreadSta
 
 void FCpuProfilerAnalyzer::ProcessBufferV2(const FEventTime& EventTime, FThreadState& ThreadState, const uint8* BufferPtr, uint32 BufferSize, int32 Version)
 {
+	FAnalysisSessionEditScope _(Session);
+
 	uint64 LastCycle = ThreadState.LastCycle;
 
 	CPUPROFILER_DEBUG_LOGF(TEXT("[%u] ProcessBuffer %llu (%.9f)\n"), ThreadState.ThreadId, LastCycle, EventTime.AsSeconds(LastCycle));
@@ -558,7 +441,7 @@ void FCpuProfilerAnalyzer::ProcessBufferV2(const FEventTime& EventTime, FThreadS
 				{
 					if (CoroutineTimerId == ~0u)
 					{
-						CoroutineTimerId = DefineNewTimerChecked(CoroutineSpecId, TEXT("Coroutine"));
+						CoroutineTimerId = AddCpuTimer(CoroutineSpecId, TEXT("Coroutine"));
 					}
 
 					TArray<uint8> CborData;
@@ -589,7 +472,7 @@ void FCpuProfilerAnalyzer::ProcessBufferV2(const FEventTime& EventTime, FThreadS
 				{
 					if (CoroutineUnknownTimerId == ~0u)
 					{
-						CoroutineUnknownTimerId = DefineNewTimerChecked(CoroutineUnknownSpecId, TEXT("<unknown>"));
+						CoroutineUnknownTimerId = AddCpuTimer(CoroutineUnknownSpecId, TEXT("<unknown>"));
 					}
 
 					//TODO: Restore the saved stack of CPU scoped timers for this CoroutineId.
@@ -675,7 +558,7 @@ void FCpuProfilerAnalyzer::ProcessBufferV2(const FEventTime& EventTime, FThreadS
 							constexpr uint32 MetadataUnknownSpecId = (1u << 31u) - 3u;
 							if (MetadataUnknownTimerId == ~0u)
 							{
-								MetadataUnknownTimerId = DefineNewTimerChecked(MetadataUnknownSpecId, TEXT("<unknown>"));
+								MetadataUnknownTimerId = AddCpuTimer(MetadataUnknownSpecId, TEXT("<unknown>"));
 							}
 
 							// Add an empty placeholder metadata so we obtain a MetadataId to use as the TimerId. Will be replaced with the actual metadata if the metadata event arrives later.
@@ -691,12 +574,12 @@ void FCpuProfilerAnalyzer::ProcessBufferV2(const FEventTime& EventTime, FThreadS
 					{
 						// Get the actual Spec Id.
 						SpecId = SpecId >> 1;
-						TimerId = GetTimerId(SpecId);
+						TimerId = GetOrAddTimer(SpecId);
 					}
 				}
 				else
 				{
-					TimerId = GetTimerId(SpecId);
+					TimerId = GetOrAddTimer(SpecId);
 				}
 
 				FEventScopeState& ScopeState = ThreadState.ScopeStack.AddDefaulted_GetRef();
@@ -851,7 +734,7 @@ void FCpuProfilerAnalyzer::OnCpuScopeEnter(const FOnEventContext& Context)
 	}
 
 	uint32 ThreadId = Context.ThreadInfo.GetId();
-	FThreadState& ThreadState = GetThreadState(ThreadId);
+	FThreadState& ThreadState = GetOrAddThreadState(ThreadId);
 
 	if (ThreadState.bShouldIgnorePendingEvents)
 	{
@@ -862,21 +745,25 @@ void FCpuProfilerAnalyzer::OnCpuScopeEnter(const FOnEventContext& Context)
 	SpecId = ~SpecId; // to keep out of the way of normal spec IDs.
 
 	uint32 TimerId;
-	uint32* TimerIdIter = SpecIdToTimerIdMap.Find(SpecId);
-	if (TimerIdIter)
+	if (const uint32* FoundTimerIdBySpecId = SpecIdToTimerIdMap.Find(SpecId))
 	{
-		TimerId = *TimerIdIter;
+		TimerId = *FoundTimerIdBySpecId;
 	}
 	else
 	{
+		CPUPROFILER_LOG_API_L2(TEXT("[CpuProfiler] OnCpuScopeEnter with a new timer"));
 		FString ScopeName;
 		ScopeName += Context.EventData.GetTypeInfo().GetName();
-		TimerId = DefineTimer(SpecId, Session.StoreString(*ScopeName), nullptr, 0, false);
+		TimerId = DefineUniqueTimer(SpecId, *ScopeName, nullptr, 0);
 	}
 
 	TArray<uint8> CborData;
 	Context.EventData.SerializeToCbor(CborData);
-	TimerId = EditableTimingProfilerProvider.AddMetadata(TimerId, MoveTemp(CborData));
+	if (ensure(CborData.Num() > 0))
+	{
+		FAnalysisSessionEditScope _(Session);
+		TimerId = EditableTimingProfilerProvider.AddMetadata(TimerId, MoveTemp(CborData));
+	}
 
 	uint64 Cycle = Context.EventTime.AsCycle64();
 	double Time = Context.EventTime.AsSeconds();
@@ -898,7 +785,7 @@ void FCpuProfilerAnalyzer::OnCpuScopeLeave(const FOnEventContext& Context)
 	}
 
 	uint32 ThreadId = Context.ThreadInfo.GetId();
-	FThreadState& ThreadState = GetThreadState(ThreadId);
+	FThreadState& ThreadState = GetOrAddThreadState(ThreadId);
 
 	if (ThreadState.bShouldIgnorePendingEvents)
 	{
@@ -917,122 +804,375 @@ void FCpuProfilerAnalyzer::OnCpuScopeLeave(const FOnEventContext& Context)
 
 ////////////////////////////////////////////////////////////////////////////////////////////////////
 
-uint32 FCpuProfilerAnalyzer::DefineTimer(uint32 SpecId, const TCHAR* Name, const TCHAR* File, uint32 Line, bool bMergeByName)
+void FCpuProfilerAnalyzer::OnEventSpec(const FOnEventContext& Context)
 {
-	// Expected: Name is already a pointer in a string store.
+	const auto& EventData = Context.EventData;
 
-	// The CPU scoped events (timers) can be merged by name.
-	// Ex.: If bMergeByName == true and there are multiple timers defined in code with same name,
-	//      those will appear in Insights as a single timer.
+	uint32 SpecId = EventData.GetValue<uint32>("Id");
 
-	// Check if a timer with same name was already defined.
-	FTimerInfo* FindTimerIdByName = ScopeNameToTimerIdMap.Find((uint64)Name);
-	if (FindTimerIdByName)
+	const TCHAR* TimerName = nullptr;
+	FString Name;
+	if (EventData.GetString("Name", Name))
 	{
-		// Yes, a timer with same name was already defined.
+		TimerName = *Name;
+	}
+	else
+	{
+		uint8 CharSize = EventData.GetValue<uint8>("CharSize");
+		if (CharSize == sizeof(ANSICHAR))
+		{
+			const ANSICHAR* AnsiName = reinterpret_cast<const ANSICHAR*>(EventData.GetAttachment());
+			Name = StringCast<TCHAR>(AnsiName).Get();
+			TimerName = *Name;
+		}
+		else if (CharSize == 0 || CharSize == sizeof(TCHAR)) // 0 for backwards compatibility
+		{
+			TimerName = reinterpret_cast<const TCHAR*>(EventData.GetAttachment());
+		}
+		else
+		{
+			Name = FString::Printf(TEXT("<invalid %u>"), SpecId);
+			TimerName = *Name;
+		}
+	}
 
-		++(FindTimerIdByName->Count);
+	if (TimerName[0] == 0)
+	{
+		Name = FString::Printf(TEXT("<noname %u>"), SpecId);
+		TimerName = *Name;
+	}
 
-		uint32 TimerId;
+	FString File;
+	uint32 Line = 0;
+	if (EventData.GetString("File", File) && !File.IsEmpty())
+	{
+		Line = EventData.GetValue<uint32>("Line");
+	}
+	const TCHAR* FileName = !File.IsEmpty() ? *File : nullptr;
 
-		// Check if SpecId is already mapped to timer.
-		const uint32* FindTimerId = SpecIdToTimerIdMap.Find(SpecId);
-		if (FindTimerId && bMergeByName)
+	CPUPROFILER_LOG_API_L2(TEXT("[CpuProfiler] EventSpec Id=%u Name=\"%s\" File=\"%s\" Line=%u"),
+		SpecId, TimerName, FileName ? FileName : TEXT("N/A"), Line);
+
+	const TCHAR* StoredTimerName = Session.StoreString(TimerName);
+	DefineMergedTimer(SpecId, StoredTimerName, FileName, Line);
+}
+
+////////////////////////////////////////////////////////////////////////////////////////////////////
+
+void FCpuProfilerAnalyzer::OnMetadataSpec(const FOnEventContext& Context)
+{
+	const auto& EventData = Context.EventData;
+
+	uint32 SpecId = EventData.GetValue<uint32>("Id");
+
+	FString Name, NameFormat;
+	EventData.GetString("Name", Name);
+	EventData.GetString("NameFormat", NameFormat);
+	TArrayView<const uint8> FieldNames = EventData.GetArrayView<uint8>("FieldNames");
+
+	CPUPROFILER_LOG_API_L2(TEXT("[CpuProfiler] MetadataSpec Id=%u Name=\"%s\" NameFormat=\"%s\" Fields[]={%u bytes}"),
+		SpecId, *Name, *NameFormat, FieldNames.Num());
+
+	if (Name.Compare(NameFormat) == 0)
+	{
+		Name.Empty();
+	}
+
+	int Index = 0;
+	while (Index < NameFormat.Len() && Index < Name.Len())
+	{
+		if (Name[Index] != NameFormat[Index])
 		{
-			// Yes, SpecId was already mapped to a timer (ex. as an <unknown> timer).
-			TimerId = *FindTimerId;
+			break;
+		}
 
-			// Update name for mapped timer.
-			EditableTimingProfilerProvider.SetTimerNameAndLocation(TimerId, Name, File, Line);
+		++Index;
+	}
 
-			// In this case, we do not remap the SpecId to the previously defined timer with same name.
-			// This is because the two timers are already used in timelines.
-			// So we will continue to use separate timers, even if those have same name.
+	if (Index > 1)
+	{
+		NameFormat.MidInline(Index);
+	}
+
+	if (Name.IsEmpty())
+	{
+		const FString Separators = TEXT("% (=");
+		Index = 0;
+		bool bIsInFormatSpecifier = false;
+		while (Index < NameFormat.Len())
+		{
+			int32 SpecIndex = -1;
+			if (Separators.FindChar(NameFormat[Index], SpecIndex))
+			{
+				Name = NameFormat.Left(Index);
+				NameFormat.MidInline(Index);
+				break;
+			}
+
+			++Index;
 		}
-		else
+
+		if (Name.IsEmpty())
 		{
-			if (bMergeByName)
+			if (!NameFormat.IsEmpty())
 			{
-				// Use the previously defined timer with same name.
-				TimerId = FindTimerIdByName->Id;
+				Name = NameFormat;
+				NameFormat.Empty();
 			}
 			else
 			{
-				// Define a new CPU timer.
-				TimerId = EditableTimingProfilerProvider.AddCpuTimer(Name, File, Line);
+				Name = Session.StoreString(TEXT("Unknown"));
 			}
+		}
+	}
 
-			// Map this SpecId to the previously defined timer with same name.
-			SpecIdToTimerIdMap.Add(SpecId, TimerId);
+	FMetadataSpec Spec;
+	if (!NameFormat.IsEmpty())
+	{
+		Spec.Format = Session.StoreString(NameFormat);
+	}
+
+	if (FieldNames.Num() > 0)
+	{
+		FMemoryReaderView MemoryReader(FieldNames);
+		FCborReader CborReader(&MemoryReader, ECborEndianness::StandardCompliant);
+		FCborContext CborContext;
+
+		while (CborReader.ReadNext(CborContext))
+		{
+			if (CborContext.MajorType() == ECborCode::TextString)
+			{
+				FString Field = CborContext.AsString();
+				Spec.FieldNames.Add(Session.StoreString(Field));
+			}
 		}
+	}
 
-		return TimerId;
+	uint32 TimerId;
+	if (const uint32* FoundTimerIdBySpecId = SpecIdToTimerIdMap.Find(SpecId))
+	{
+		TimerId = *FoundTimerIdBySpecId;
+
+		FAnalysisSessionEditScope _(Session);
+		SetTimerName(SpecId, TimerId, *Name);
 	}
 	else
 	{
-		// No, a timer with same name was not defined (or we do not want to merge by name).
+		CPUPROFILER_LOG_API_L2(TEXT("[CpuProfiler] MetadataSpec with a new timer"));
+		TimerId = DefineUniqueTimer(SpecId, *Name, nullptr, 0);
+	}
+
+	if (Spec.FieldNames.Num() > 0 || Spec.Format != nullptr)
+	{
+		FAnalysisSessionEditScope _(Session);
+		uint32 MetadataSpecId = EditableTimingProfilerProvider.AddMetadataSpec(MoveTemp(Spec));
+
+		const ITimingProfilerProvider* TimingProfilerProvider = EditableTimingProfilerProvider.GetReadProvider();
+		if (TimingProfilerProvider)
+		{
+			FTimingProfilerTimer PreviousTimer;
+			TimingProfilerProvider->ReadTimers(
+				[TimerId, &PreviousTimer]
+				(const ITimingProfilerTimerReader& TimerReader)
+				{
+					const FTimingProfilerTimer* PreviousTimerPtr = TimerReader.GetTimer(TimerId);
+					if (PreviousTimerPtr)
+					{
+						PreviousTimer = *PreviousTimerPtr;
+					}
+				});
+			if (PreviousTimer.HasValidMetadataSpecId())
+			{
+				TimerId = AddCpuTimer(SpecId, *Name, PreviousTimer.File, PreviousTimer.Line);
+
+				if (++NumTimerWarnings < NumMaxWarnings)
+				{
+					const FMetadataSpec* PreviousMetadataSpec = TimingProfilerProvider->GetMetadataSpec(PreviousTimer.MetadataSpecId);
+					check(PreviousMetadataSpec);
+					UE_LOG(LogTraceServices, Warning, TEXT("[CpuProfiler] MetadataSpec -- de-duplicated timer for SpecId %u -- old: %u (%s) metadata=\"%s\" (%d fields) --> new: %u (%s) metadata=\"%s\" (%d fields) !!!"),
+						SpecId,
+						PreviousTimer.Id, PreviousTimer.Name, PreviousMetadataSpec->Format, PreviousMetadataSpec->FieldNames.Num(),
+						TimerId, *Name, Spec.Format, Spec.FieldNames.Num());
+				}
+			}
+		}
+
+		EditableTimingProfilerProvider.SetMetadataSpec(TimerId, MetadataSpecId);
+	}
+}
+
+////////////////////////////////////////////////////////////////////////////////////////////////////
+
+uint32 FCpuProfilerAnalyzer::GetOrAddTimer(uint32 SpecId)
+{
+	if (const uint32* FoundTimerIdBySpecId = SpecIdToTimerIdMap.Find(SpecId))
+	{
+		return *FoundTimerIdBySpecId;
+	}
 
-		uint32 TimerId;
+	// Add a timer with an "unknown" name.
+	// The "unknown" timers are not merged by name, because the actual name
+	// might be updated when an EventSpec event is received (for this SpecId).
+	FAnalysisSessionEditScope _(Session);
+	return AddCpuTimer(SpecId, *FString::Printf(TEXT("<unknown %u>"), SpecId));
+}
+
+////////////////////////////////////////////////////////////////////////////////////////////////////
+
+uint32 FCpuProfilerAnalyzer::AddCpuTimer(uint32 SpecId, const TCHAR* TimerName, const TCHAR* File, uint32 Line)
+{
+	Session.WriteAccessCheck();
+
+	// Add a new CPU timer.
+	uint32 TimerId = EditableTimingProfilerProvider.AddCpuTimer(TimerName, File, Line);
+
+	// Map the SpecId to the timer.
+	SpecIdToTimerIdMap.Add(SpecId, TimerId);
+
+	CPUPROFILER_LOG_API_L3(TEXT("[CpuProfiler] --> AddCpuTimer(SpecId=%u, Name=\"%s\") --> TimerId=%u"), SpecId, TimerName, TimerId);
+
+	return TimerId;
+}
+
+////////////////////////////////////////////////////////////////////////////////////////////////////
+
+uint32 FCpuProfilerAnalyzer::DefineMergedTimer(uint32 SpecId, const TCHAR* StoredTimerName, const TCHAR* File, uint32 Line)
+{
+	// Expected: StoredTimerName is already a pointer in a string store.
+
+	// The CPU scoped events (timers) can be merged by name.
+	// If there are multiple timers defined in code with same name,
+	// those will appear in Insights as a single timer.
+
+	uint32 TimerId;
+
+	// Check if a timer with same name was already defined.
+	FTimerInfo* FoundTimerIdByName = ScopeNameToTimerIdMap.Find((uint64)StoredTimerName);
+	if (FoundTimerIdByName)
+	{
+		// Yes, a timer with same name was already defined.
+		++(FoundTimerIdByName->Count);
 
 		// Check if SpecId is already mapped to timer.
-		const uint32* FindTimerId = SpecIdToTimerIdMap.Find(SpecId);
-		if (FindTimerId)
+		if (const uint32* FoundTimerIdBySpecId = SpecIdToTimerIdMap.Find(SpecId))
 		{
 			// Yes, SpecId was already mapped to a timer (ex. as an <unknown> timer).
-			TimerId = *FindTimerId;
+			TimerId = *FoundTimerIdBySpecId;
 
 			// Update name for mapped timer.
-			EditableTimingProfilerProvider.SetTimerNameAndLocation(TimerId, Name, File, Line);
+			FAnalysisSessionEditScope _(Session);
+			SetTimerName(SpecId, TimerId, StoredTimerName);
+			EditableTimingProfilerProvider.SetTimerLocation(TimerId, File, Line);
+
+			// In this case, we do not remap the SpecId to the previously defined timer with same name.
+			// This is because the two timers are already used in timelines.
+			// So we will continue to use separate timers, even if those have same name.
 		}
 		else
 		{
-			// Define a new CPU timer.
-			TimerId = EditableTimingProfilerProvider.AddCpuTimer(Name, File, Line);
+			// Use the previously defined timer with same name.
+			TimerId = FoundTimerIdByName->Id;
 
-			// Map the SpecId to the timer.
+			// Map this SpecId to the previously defined timer with same name.
 			SpecIdToTimerIdMap.Add(SpecId, TimerId);
 		}
 
-		// Map the name to the timer.
-		ScopeNameToTimerIdMap.Add((uint64)Name, { TimerId, 1 });
 		return TimerId;
 	}
+	else // No, this is the first timer with this name.
+	{
+		// Create a new timer.
+		TimerId = DefineUniqueTimer(SpecId, StoredTimerName, File, Line);
+
+		// Map the name to the timer.
+		ScopeNameToTimerIdMap.Add((uint64)StoredTimerName, { TimerId, 1 });
+	}
+
+	return TimerId;
+}
+
+////////////////////////////////////////////////////////////////////////////////////////////////////
+
+uint32 FCpuProfilerAnalyzer::DefineUniqueTimer(uint32 SpecId, const TCHAR* TimerName, const TCHAR* File, uint32 Line)
+{
+	uint32 TimerId;
+
+	// Check if SpecId is already mapped to timer.
+	if (const uint32* FoundTimerIdBySpecId = SpecIdToTimerIdMap.Find(SpecId))
+	{
+		// Yes, SpecId was already mapped to a timer (ex. as an <unknown> timer).
+		TimerId = *FoundTimerIdBySpecId;
+
+		// Update name for the mapped timer.
+		FAnalysisSessionEditScope _(Session);
+		SetTimerName(SpecId, TimerId, TimerName);
+		EditableTimingProfilerProvider.SetTimerLocation(TimerId, File, Line);
+	}
+	else
+	{
+		// Define a new CPU timer.
+		FAnalysisSessionEditScope _(Session);
+		TimerId = AddCpuTimer(SpecId, TimerName, File, Line);
+	}
+
+	return TimerId;
 }
 
 ////////////////////////////////////////////////////////////////////////////////////////////////////
 
-uint32 FCpuProfilerAnalyzer::DefineNewTimerChecked(uint32 SpecId, const TCHAR* TimerName, const TCHAR* File, uint32 Line)
+const TCHAR* FCpuProfilerAnalyzer::GetTimerName(uint32 TimerId) const
 {
-	TimerName = Session.StoreString(TimerName);
-	uint32 NewTimerId = EditableTimingProfilerProvider.AddCpuTimer(TimerName, File, Line);
-	SpecIdToTimerIdMap.Add(SpecId, NewTimerId);
-	return NewTimerId;
+	const TCHAR* TimerName = nullptr;
+	const ITimingProfilerProvider* TimingProfilerProvider = EditableTimingProfilerProvider.GetReadProvider();
+	if (TimingProfilerProvider)
+	{
+		TimingProfilerProvider->ReadTimers(
+			[TimerId, &TimerName]
+			(const ITimingProfilerTimerReader& TimerReader)
+			{
+				const FTimingProfilerTimer* Timer = TimerReader.GetTimer(TimerId);
+				if (Timer)
+				{
+					TimerName = Timer->Name;
+				}
+			});
+	}
+	return TimerName;
 }
 
 ////////////////////////////////////////////////////////////////////////////////////////////////////
 
-uint32 FCpuProfilerAnalyzer::GetTimerId(uint32 SpecId)
+void FCpuProfilerAnalyzer::SetTimerName(uint32 SpecId, uint32 TimerId, const TCHAR* TimerName)
 {
-	if (uint32* FindIt = SpecIdToTimerIdMap.Find(SpecId))
+#if CPUPROFILER_DETECT_TIMER_NAME_CHANGE
+	const TCHAR* PrevTimerName = GetTimerName(TimerId);
+#endif
+
+	Session.WriteAccessCheck();
+	EditableTimingProfilerProvider.SetTimerName(TimerId, TimerName);
+
+#if CPUPROFILER_DETECT_TIMER_NAME_CHANGE
+	const TCHAR* NewTimerName = GetTimerName(TimerId);
+	if (PrevTimerName != nullptr && PrevTimerName != NewTimerName && PrevTimerName[0] != TEXT('<'))
 	{
-		return *FindIt;
+		UE_LOG(LogTraceServices, Warning, TEXT("[CpuProfiler] --> SetTimerName(SpecId=%u, TimerId=%u, Name=\"%s\") : name was \"%s\" !!!"), SpecId, TimerId, TimerName, PrevTimerName);
 	}
 	else
 	{
-		// Adds a timer with an "unknown" name.
-		// The "unknown" timers are not merged by name, because the actual name
-		// might be updated when an EventSpec event is received (for this SpecId).
-		return DefineNewTimerChecked(SpecId, *FString::Printf(TEXT("<unknown %u>"), SpecId));
+		CPUPROFILER_LOG_API_L2(TEXT("[CpuProfiler] --> SetTimerName(SpecId=%u, TimerId=%u, Name=\"%s\""), SpecId, TimerId, TimerName);
 	}
+#endif
 }
 
 ////////////////////////////////////////////////////////////////////////////////////////////////////
 
-FCpuProfilerAnalyzer::FThreadState& FCpuProfilerAnalyzer::GetThreadState(uint32 ThreadId)
+FCpuProfilerAnalyzer::FThreadState& FCpuProfilerAnalyzer::GetOrAddThreadState(uint32 ThreadId)
 {
 	FThreadState* ThreadState = ThreadStatesMap.FindRef(ThreadId);
 	if (!ThreadState)
 	{
+		FAnalysisSessionEditScope _(Session);
+
 		ThreadState = new FThreadState();
 		ThreadState->ThreadId = ThreadId;
 		ThreadState->Timeline = &EditableTimingProfilerProvider.GetCpuThreadEditableTimeline(ThreadId);
@@ -1049,6 +1189,10 @@ FCpuProfilerAnalyzer::FThreadState& FCpuProfilerAnalyzer::GetThreadState(uint32
 
 } // namespace TraceServices
 
+#undef CPUPROFILER_DETECT_TIMER_NAME_CHANGE
+#undef CPUPROFILER_LOG_API_L3
+#undef CPUPROFILER_LOG_API_L2
+#undef CPUPROFILER_LOG_API_L1
 #undef CPUPROFILER_DEBUG_LOGF
 #undef CPUPROFILER_DEBUG_BEGIN_EVENT
 #undef CPUPROFILER_DEBUG_END_EVENT
diff --git a/Engine/Source/Developer/TraceServices/Private/Analyzers/CpuProfilerTraceAnalysis.h b/Engine/Source/Developer/TraceServices/Private/Analyzers/CpuProfilerTraceAnalysis.h
--- a/Engine/Source/Developer/TraceServices/Private/Analyzers/CpuProfilerTraceAnalysis.h
+++ b/Engine/Source/Developer/TraceServices/Private/Analyzers/CpuProfilerTraceAnalysis.h
@@ -2,11 +2,16 @@
 
 #pragma once
 
-#include "Trace/Analyzer.h"
+#include "CoreTypes.h"
 #include "Containers/Map.h"
 #include "Containers/UnrealString.h"
-#include "TraceServices/Model/TimingProfiler.h"
+
+// TraceAnalysis
+#include "Trace/Analyzer.h"
+
+// TraceServices
 #include "Model/MonotonicTimeline.h"
+#include "TraceServices/Model/TimingProfiler.h"
 
 namespace TraceServices
 {
@@ -77,10 +82,17 @@ private:
 	void EndOpenEvents(FThreadState& ThreadState, double Timestamp);
 	void OnCpuScopeEnter(const FOnEventContext& Context);
 	void OnCpuScopeLeave(const FOnEventContext& Context);
-	uint32 DefineTimer(uint32 SpecId, const TCHAR* TimerName, const TCHAR* File, uint32 Line, bool bMergeByName); // returns the TimerId
-	uint32 DefineNewTimerChecked(uint32 SpecId, const TCHAR* TimerName, const TCHAR* File = nullptr, uint32 Line = 0); // returns the TimerId
-	uint32 GetTimerId(uint32 SpecId);
-	FThreadState& GetThreadState(uint32 ThreadId);
+	void OnEventSpec(const FOnEventContext& Context);
+	void OnMetadataSpec(const FOnEventContext& Context);
+
+	uint32 GetOrAddTimer(uint32 SpecId);
+	uint32 AddCpuTimer(uint32 SpecId, const TCHAR* TimerName, const TCHAR* File = nullptr, uint32 Line = 0); // returns the TimerId
+	uint32 DefineMergedTimer(uint32 SpecId, const TCHAR* TimerName, const TCHAR* File, uint32 Line); // returns the TimerId
+	uint32 DefineUniqueTimer(uint32 SpecId, const TCHAR* TimerName, const TCHAR* File, uint32 Line); // returns the TimerId
+	const TCHAR* GetTimerName(uint32 TimerId) const;
+	void SetTimerName(uint32 SpecId, uint32 TimerId, const TCHAR* TimerName);
+
+	FThreadState& GetOrAddThreadState(uint32 ThreadId);
 
 private:
 	IAnalysisSession& Session;
@@ -98,6 +110,9 @@ private:
 	uint64 TotalEventSize = 0;
 	uint64 TotalScopeCount = 0;
 	double BytesPerScope = 0.0;
+
+	uint32 NumTimerWarnings = 0;
+	static constexpr uint32 NumMaxWarnings = 100;
 };
 
 } // namespace TraceServices
diff --git a/Engine/Source/Developer/TraceServices/Private/Analyzers/GpuProfilerTraceAnalysis.cpp b/Engine/Source/Developer/TraceServices/Private/Analyzers/GpuProfilerTraceAnalysis.cpp
--- a/Engine/Source/Developer/TraceServices/Private/Analyzers/GpuProfilerTraceAnalysis.cpp
+++ b/Engine/Source/Developer/TraceServices/Private/Analyzers/GpuProfilerTraceAnalysis.cpp
@@ -2,20 +2,27 @@
 
 #include "GpuProfilerTraceAnalysis.h"
 
-#include "AnalysisServicePrivate.h"
-#include "CborWriter.h"
 #include "CborReader.h"
-#include "Common/Utils.h"
+#include "CborWriter.h"
 #include "HAL/LowLevelMemTracker.h"
-#include "Model/CountersPrivate.h"
 #include "Serialization/MemoryReader.h"
 #include "Serialization/MemoryWriter.h"
 
-#define UE_INSIGHTS_MERGE_UNACCOUNTED_TIMERS 0
+// TraceServices
+#include "AnalysisServicePrivate.h"
+#include "Common/Utils.h"
+#include "Model/CountersPrivate.h"
+#include "Model/TimingProfilerPrivate.h"
+
+#define GPUPROFILER_LOG_API_L1(Format, ...) UE_LOG(LogTraceServices, Log, Format, ##__VA_ARGS__)
+#define GPUPROFILER_LOG_API_L2(Format, ...) //UE_LOG(LogTraceServices, Verbose, Format, ##__VA_ARGS__)
+#define GPUPROFILER_LOG_API_L3(Format, ...) //UE_LOG(LogTraceServices, VeryVerbose, Format, ##__VA_ARGS__)
 
 namespace TraceServices
 {
 
+////////////////////////////////////////////////////////////////////////////////////////////////////
+
 FGpuProfilerAnalyzer::FGpuProfilerAnalyzer(FAnalysisSession& InSession, FTimingProfilerProvider& InTimingProfilerProvider, IEditableCounterProvider& InEditableCounterProvider)
 	: Session(InSession)
 	, TimingProfilerProvider(InTimingProfilerProvider)
@@ -23,49 +30,30 @@ FGpuProfilerAnalyzer::FGpuProfilerAnalyzer(FAnalysisSession& InSession, FTimingP
 {
 }
 
+////////////////////////////////////////////////////////////////////////////////////////////////////
+
 void FGpuProfilerAnalyzer::OnAnalysisBegin(const FOnAnalysisContext& Context)
 {
 	auto& Builder = Context.InterfaceBuilder;
 
-	// The New GPU Profiler
-	Builder.RouteEvent(RouteId_Init, "GpuProfiler", "Init");
-	Builder.RouteEvent(RouteId_QueueSpec, "GpuProfiler", "QueueSpec");
-	Builder.RouteEvent(RouteId_EventFrameBoundary, "GpuProfiler", "EventFrameBoundary");
-	Builder.RouteEvent(RouteId_EventBreadcrumbSpec, "GpuProfiler", "EventBreadcrumbSpec");
+	Builder.RouteEvent(RouteId_Init,                 "GpuProfiler", "Init");
+	Builder.RouteEvent(RouteId_QueueSpec,            "GpuProfiler", "QueueSpec");
+	Builder.RouteEvent(RouteId_EventFrameBoundary,   "GpuProfiler", "EventFrameBoundary");
+	Builder.RouteEvent(RouteId_EventBreadcrumbSpec,  "GpuProfiler", "EventBreadcrumbSpec");
 	Builder.RouteEvent(RouteId_EventBeginBreadcrumb, "GpuProfiler", "EventBeginBreadcrumb");
-	Builder.RouteEvent(RouteId_EventEndBreadcrumb, "GpuProfiler", "EventEndBreadcrumb");
-	Builder.RouteEvent(RouteId_EventBeginWork, "GpuProfiler", "EventBeginWork");
-	Builder.RouteEvent(RouteId_EventEndWork, "GpuProfiler", "EventEndWork");
-	Builder.RouteEvent(RouteId_EventWait, "GpuProfiler", "EventWait");
-	Builder.RouteEvent(RouteId_EventStats, "GpuProfiler", "EventStats");
-	Builder.RouteEvent(RouteId_SignalFence, "GpuProfiler", "SignalFence");
-	Builder.RouteEvent(RouteId_WaitFence, "GpuProfiler", "WaitFence");
-
-	// The Old GPU Profiler (deprecated in UE 5.6)
-	// Analysis code is maintained for backward compatibility with old traces.
-	Builder.RouteEvent(RouteId_EventSpec, "GpuProfiler", "EventSpec");
-	Builder.RouteEvent(RouteId_Frame, "GpuProfiler", "Frame");
-	Builder.RouteEvent(RouteId_Frame2, "GpuProfiler", "Frame2");
+	Builder.RouteEvent(RouteId_EventEndBreadcrumb,   "GpuProfiler", "EventEndBreadcrumb");
+	Builder.RouteEvent(RouteId_EventBeginWork,       "GpuProfiler", "EventBeginWork");
+	Builder.RouteEvent(RouteId_EventEndWork,         "GpuProfiler", "EventEndWork");
+	Builder.RouteEvent(RouteId_EventWait,            "GpuProfiler", "EventWait");
+	Builder.RouteEvent(RouteId_EventStats,           "GpuProfiler", "EventStats");
+	Builder.RouteEvent(RouteId_SignalFence,          "GpuProfiler", "SignalFence");
+	Builder.RouteEvent(RouteId_WaitFence,            "GpuProfiler", "WaitFence");
 }
 
+////////////////////////////////////////////////////////////////////////////////////////////////////
+
 void FGpuProfilerAnalyzer::OnAnalysisEnd()
 {
-	//////////////////////////////////////////////////
-	// Old GPU Profiler
-
-	if (NumFramesWithErrors > 0)
-	{
-		UE_LOG(LogTraceServices, Error, TEXT("[GpuProfiler] Frames with errors: %u"), NumFramesWithErrors);
-	}
-
-	if (NumFrames > 0 || EventTypeMap.Num() > 0)
-	{
-		UE_LOG(LogTraceServices, Log, TEXT("[GpuProfiler] Analysis completed (%u frames, %d timers)."), NumFrames, EventTypeMap.Num());
-	}
-
-	//////////////////////////////////////////////////
-	// New GPU Profiler
-
 	if (ErrorData.NumInterleavedEvents > 0)
 	{
 		UE_LOG(LogTraceServices, Warning, TEXT("[GpuProfiler] Number of interleaved events: %u, max error %f"), ErrorData.NumInterleavedEvents, ErrorData.InterleavedEventsMaxDelta);
@@ -86,13 +74,20 @@ void FGpuProfilerAnalyzer::OnAnalysisEnd()
 		UE_LOG(LogTraceServices, Warning, TEXT("[GpuProfiler] Number of negative duration events: %u, max error %f"), ErrorData.NumNegativeDurationEvents, ErrorData.NegativeDurationEventsMaxDelta);
 	}
 
+	if (ErrorData.NumTimerWarnings > 0)
+	{
+		UE_LOG(LogTraceServices, Warning, TEXT("[GpuProfiler] Number of timer/metadata warnings: %u"), ErrorData.NumTimerWarnings);
+	}
+
 	if (Queues.Num() > 0 || TimerMap.Num() > 0)
 	{
 		UE_LOG(LogTraceServices, Log, TEXT("[GpuProfiler] Analysis completed (%u queues, %d timers, %d breadcrumb specs, %d breadcrumb names)."),
-			Queues.Num(), TimerMap.Num(), BreadcrumbSpecMap.Num(), BreadcrumbMap.Num());
+			Queues.Num(), TimerMap.Num(), SpecIdToTimerIdMap.Num(), BreadcrumbMap.Num());
 	}
 }
 
+////////////////////////////////////////////////////////////////////////////////////////////////////
+
 bool FGpuProfilerAnalyzer::OnEvent(uint16 RouteId, EStyle Style, const FOnEventContext& Context)
 {
 	LLM_SCOPE_BYNAME(TEXT("Insights/FGpuProfilerAnalyzer"));
@@ -102,15 +97,13 @@ bool FGpuProfilerAnalyzer::OnEvent(uint16 RouteId, EStyle Style, const FOnEventC
 	switch (RouteId)
 	{
 
-	//////////////////////////////////////////////////
-	// The New GPU Profiler
-
 	case RouteId_Init:
 	{
 		Version = (uint32)EventData.GetValue<uint8>("Version");
-		UE_LOG(LogTraceServices, Log, TEXT("[GpuProfiler] Init Version=%u"), Version);
+		GPUPROFILER_LOG_API_L1(TEXT("[GpuProfiler] Init Version=%u"), Version);
 		break;
 	}
+
 	case RouteId_QueueSpec:
 	{
 		const uint32 QueueId = EventData.GetValue<uint32>("QueueId");
@@ -122,24 +115,26 @@ bool FGpuProfilerAnalyzer::OnEvent(uint16 RouteId, EStyle Style, const FOnEventC
 		FString Name;
 		EventData.GetString("TypeString", Name);
 
-		UE_LOG(LogTraceServices, Log, TEXT("[GpuProfiler] QueueSpec QueueId=%u, GPU=%u, Index=%u, Type=%d \"%s\""), QueueId, (uint32)GPU, (uint32)Index, (uint32)Type, *Name);
+		GPUPROFILER_LOG_API_L2(TEXT("[GpuProfiler] QueueSpec QueueId=%u, GPU=%u, Index=%u, Type=%d \"%s\""), QueueId, (uint32)GPU, (uint32)Index, (uint32)Type, *Name);
 
-		FAnalysisSessionEditScope _(Session);
 		const TCHAR* PersistentName = Session.StoreString(Name);
 
-		TimingProfilerProvider.AddGpuQueue(QueueId, GPU, Index, Type, PersistentName);
+		{
+			FAnalysisSessionEditScope _(Session);
+			TimingProfilerProvider.AddGpuQueue(QueueId, GPU, Index, Type, PersistentName);
+		}
 
 		FQueue& Queue = GetOrAddQueue(QueueId);
 		InitCountersDesc(Queue, GPU, Index, *Name);
-
 		break;
 	}
+
 	case RouteId_EventFrameBoundary:
 	{
 		const uint32 QueueId = EventData.GetValue<uint32>("QueueId");
 		const uint32 FrameNumber = EventData.GetValue<uint32>("FrameNumber");
 
-		UE_LOG(LogTraceServices, VeryVerbose, TEXT("[GpuProfiler] EventFrameBoundary QueueId=%u, FrameNumber=%u"), QueueId, FrameNumber);
+		GPUPROFILER_LOG_API_L3(TEXT("[GpuProfiler] EventFrameBoundary QueueId=%u, FrameNumber=%u"), QueueId, FrameNumber);
 
 		FQueue& Queue = GetOrAddQueue(QueueId);
 		Queue.FrameNumber = FrameNumber;
@@ -147,7 +142,6 @@ bool FGpuProfilerAnalyzer::OnEvent(uint16 RouteId, EStyle Style, const FOnEventC
 		if (Queue.LastTime > 0)
 		{
 			FAnalysisSessionEditScope _(Session);
-
 			Queue.NumDrawsCounter->SetValue(Queue.LastTime, (int64)Queue.NumDraws);
 			Queue.NumPrimitivesCounter->SetValue(Queue.LastTime, (int64)Queue.NumPrimitives);
 		}
@@ -155,172 +149,78 @@ bool FGpuProfilerAnalyzer::OnEvent(uint16 RouteId, EStyle Style, const FOnEventC
 		Queue.NumDraws = 0;
 		Queue.NumPrimitives = 0;
 		Queue.LastTime = 0;
-
 		break;
 	}
+
 	case RouteId_EventBreadcrumbSpec:
 	{
-		FString Name, NameFormat;
-		uint32 SpecId = EventData.GetValue<uint32>("SpecId");
-		EventData.GetString("StaticName", Name);
-		EventData.GetString("NameFormat", NameFormat);
-		TArrayView<const uint8> FieldNames = EventData.GetArrayView<uint8>("FieldNames");
-
-		if (Name.Compare(NameFormat) == 0)
-		{
-			Name.Empty();
-		}
-
-		int Index = 0;
-		while (Index < NameFormat.Len() && Index < Name.Len())
-		{
-			if (Name[Index] != NameFormat[Index])
-			{
-				break;
-			}
-
-			++Index;
-		}
-
-		if (Index > 1)
-		{
-			NameFormat.MidInline(Index);
-		}
-
-		if (Name.IsEmpty())
-		{
-			const FString Separators = TEXT("% (=");
-			Index = 0;
-			bool bIsInFormatSpecifier = false;
-			while (Index < NameFormat.Len())
-			{
-				int32 SpecIndex = -1;
-				if (Separators.FindChar(NameFormat[Index], SpecIndex))
-				{
-					Name = NameFormat.Left(Index);
-					NameFormat.MidInline(Index);
-					break;
-				}
-
-				++Index;
-			}
-
-			if (Name.IsEmpty())
-			{
-				if (!NameFormat.IsEmpty())
-				{
-					Name = NameFormat;
-					NameFormat.Empty();
-				}
-				else
-				{
-					Name = Session.StoreString(TEXT("Unknown"));
-				}
-			}
-		}
-
-		FMetadataSpec Spec;
-		if (!NameFormat.IsEmpty())
-		{
-			Spec.Format = Session.StoreString(NameFormat);
-		}
-
-		if (FieldNames.Num() > 0)
-		{
-			FMemoryReaderView MemoryReader(FieldNames);
-			FCborReader CborReader(&MemoryReader, ECborEndianness::StandardCompliant);
-			FCborContext CborContext;
-
-			while (CborReader.ReadNext(CborContext))
-			{
-				if (CborContext.MajorType() == ECborCode::TextString)
-				{
-					FString Field = CborContext.AsString();
-					Spec.FieldNames.Add(Session.StoreString(Field));
-				}
-			}
-		}
-
-		FAnalysisSessionEditScope _(Session);
-		uint32 TimerId = 0;
-		uint32* FoundSpecIdPtr = BreadcrumbSpecMap.Find(SpecId);
-		if (FoundSpecIdPtr == nullptr)
-		{
-			TimerId = GetOrAddTimerId(Name);
-			BreadcrumbSpecMap.Add(SpecId, TimerId);
-		}
-		else
-		{
-			TimingProfilerProvider.SetTimerName(TimerId, Name);
-		}
-
-		if (Spec.FieldNames.Num() > 0 || Spec.Format != nullptr)
-		{
-			uint32 MetadataSpecId = TimingProfilerProvider.AddMetadataSpec(MoveTemp(Spec));
-			TimingProfilerProvider.SetMetadataSpec(TimerId, MetadataSpecId);
-		}
-
+		OnEventBreadcrumbSpec(Context);
 		break;
 	}
+
 	case RouteId_EventBeginBreadcrumb:
 	{
 		const uint32 SpecId = EventData.GetValue<uint32>("SpecId");
 		const uint32 QueueId = EventData.GetValue<uint32>("QueueId");
 		const uint64 GPUTimestampTOP = EventData.GetValue<uint64>("GPUTimestampTOP");
-
 		// Ignore events for which the timestamp could not be determined.
 		if (GPUTimestampTOP == 0)
 		{
 			break;
 		}
-
-		TArray<uint8> Metadata(EventData.GetArrayView<uint8>("Metadata"));
 		const double Time = Context.EventTime.AsSeconds(GPUTimestampTOP);
+		TArray<uint8> Metadata(EventData.GetArrayView<uint8>("Metadata"));
 
-		UE_LOG(LogTraceServices, VeryVerbose, TEXT("[GpuProfiler] EventBeginBreadcrumb QueueId=%u, Time=%f, SpecId=\"%u\""), QueueId, Time, SpecId);
+		GPUPROFILER_LOG_API_L3(TEXT("[GpuProfiler] EventBeginBreadcrumb QueueId=%u, Time=%f, SpecId=\"%u\""), QueueId, Time, SpecId);
 
-		uint32 TimerId = 0;
-		uint32* TimerIdPtr;
-		TimerIdPtr = BreadcrumbSpecMap.Find(SpecId);
+		uint32 TimerId = GetOrAddTimer(SpecId, TEXT("<unknown>"));
 
+		if (Metadata.Num() > 0) // excludes empty metadata
 		{
 			FAnalysisSessionEditScope _(Session);
-			if (TimerIdPtr)
-			{
-				TimerId = *TimerIdPtr;
-			}
-			else
-			{
-				TimerId = GetOrAddTimerId(TEXT("Unknown"));
-				BreadcrumbSpecMap.Add(SpecId, TimerId);
-			}
-
 			TimerId = TimingProfilerProvider.AddMetadata(TimerId, MoveTemp(Metadata));
 		}
+#if 0 // for debugging only
+		else
+		{
+			FAnalysisSessionEditScope _(Session);
+			TimingProfilerProvider.ReadTimers(
+				[TimerId]
+				(const ITimingProfilerTimerReader& TimerReader)
+				{
+					const FTimingProfilerTimer* Timer = TimerReader.GetTimer(TimerId);
+					check(Timer);
+					if (Timer->HasValidMetadataSpecId())
+					{
+						UE_LOG(LogTraceServices, Warning, TEXT("[GpuProfiler] Empty metadata for a timing event while its timer %u (%s) has valid metadata spec id (%u) !!!"), TimerId, Timer->Name, Timer->MetadataSpecId);
+					}
+				});
+		}
+#endif
 
 		FQueue& Queue = GetOrAddQueue(QueueId);
 		BeginEvent(Queue, 0, Time, TimerId);
 		break;
 	}
+
 	case RouteId_EventEndBreadcrumb:
 	{
 		const uint32 QueueId = EventData.GetValue<uint32>("QueueId");
 		const uint64 GPUTimestampBOP = EventData.GetValue<uint64>("GPUTimestampBOP");
-
 		// Ignore events for which the timestamp could not be determined.
 		if (GPUTimestampBOP == 0)
 		{
 			break;
 		}
-
 		const double Time = Context.EventTime.AsSeconds(GPUTimestampBOP);
 
-		UE_LOG(LogTraceServices, VeryVerbose, TEXT("[GpuProfiler] EventEndBreadcrumb QueueId=%u, Time=%f"), QueueId, Time);
+		GPUPROFILER_LOG_API_L3(TEXT("[GpuProfiler] EventEndBreadcrumb QueueId=%u, Time=%f"), QueueId, Time);
 
 		FQueue& Queue = GetOrAddQueue(QueueId);
 		EndEvent(Queue, 0, Time, InvalidTimerId);
 		break;
 	}
+
 	case RouteId_EventBeginWork:
 	{
 		const uint32 QueueId = EventData.GetValue<uint32>("QueueId");
@@ -328,12 +228,12 @@ bool FGpuProfilerAnalyzer::OnEvent(uint16 RouteId, EStyle Style, const FOnEventC
 		double Time = Context.EventTime.AsSeconds(GPUTimestampTOP);
 		//const uint64 CPUTimestamp = EventData.GetValue<uint64>("CPUTimestamp"); // not used
 
-		UE_LOG(LogTraceServices, VeryVerbose, TEXT("[GpuProfiler] EventBeginWork QueueId=%u, GPUTimestampTOP=%f"), QueueId, Time);
+		GPUPROFILER_LOG_API_L3(TEXT("[GpuProfiler] EventBeginWork QueueId=%u, GPUTimestampTOP=%f"), QueueId, Time);
 
 		if (GpuWorkTimerId == InvalidTimerId)
 		{
 			const FString GpuWorkTimerName(TEXT("GpuWork"));
-			GpuWorkTimerId = GetOrAddTimerId(GpuWorkTimerName);
+			GpuWorkTimerId = GetOrAddTimer(~0u, GpuWorkTimerName);
 		}
 
 		FQueue& Queue = GetOrAddQueue(QueueId);
@@ -351,21 +251,24 @@ bool FGpuProfilerAnalyzer::OnEvent(uint16 RouteId, EStyle Style, const FOnEventC
 			FAnalysisSessionEditScope _(Session);
 			MetadataTimerId = TimingProfilerProvider.AddMetadata(GpuWorkTimerId, MoveTemp(CborData));
 		}
+
 		BeginEvent(Queue, 1, Time, MetadataTimerId);
 		break;
 	}
+
 	case RouteId_EventEndWork:
 	{
 		uint32 QueueId = EventData.GetValue<uint32>("QueueId");
 		const uint64 GPUTimestampBOP = EventData.GetValue<uint64>("GPUTimestampBOP");
 		const double Time = Context.EventTime.AsSeconds(GPUTimestampBOP);
 
-		UE_LOG(LogTraceServices, VeryVerbose, TEXT("[GpuProfiler] EventEndWork QueueId=%u, GPUTimestampBOP=%f"), QueueId, Time);
+		GPUPROFILER_LOG_API_L3(TEXT("[GpuProfiler] EventEndWork QueueId=%u, GPUTimestampBOP=%f"), QueueId, Time);
 
 		FQueue& Queue = GetOrAddQueue(QueueId);
 		EndEvent(Queue, 1, Time, GpuWorkTimerId);
 		break;
 	}
+
 	case RouteId_EventWait:
 	{
 		uint32 QueueId = EventData.GetValue<uint32>("QueueId");
@@ -375,12 +278,12 @@ bool FGpuProfilerAnalyzer::OnEvent(uint16 RouteId, EStyle Style, const FOnEventC
 		const uint64 EndTimestamp = EventData.GetValue<uint64>("EndTime");
 		const double EndTime = Context.EventTime.AsSeconds(EndTimestamp);
 
-		UE_LOG(LogTraceServices, VeryVerbose, TEXT("[GpuProfiler] EventWait QueueId=%u, StartTime=%f, EndTime=%f"), QueueId, StartTime, EndTime);
+		GPUPROFILER_LOG_API_L3(TEXT("[GpuProfiler] EventWait QueueId=%u, StartTime=%f, EndTime=%f"), QueueId, StartTime, EndTime);
 
 		if (GpuWaitTimerId == InvalidTimerId)
 		{
 			const FString GpuWaitTimerName(TEXT("GpuWait"));
-			GpuWaitTimerId = GetOrAddTimerId(GpuWaitTimerName);
+			GpuWaitTimerId = GetOrAddTimer(~0u, GpuWaitTimerName);
 		}
 
 		FQueue& Queue = GetOrAddQueue(QueueId);
@@ -403,6 +306,7 @@ bool FGpuProfilerAnalyzer::OnEvent(uint16 RouteId, EStyle Style, const FOnEventC
 		EndEvent(Queue, 1, EndTime, GpuWaitTimerId);
 		break;
 	}
+
 	case RouteId_EventStats:
 	{
 		const uint32 QueueId = EventData.GetValue<uint32>("QueueId");
@@ -413,9 +317,9 @@ bool FGpuProfilerAnalyzer::OnEvent(uint16 RouteId, EStyle Style, const FOnEventC
 
 		Queue.NumDraws += NumDraws;
 		Queue.NumPrimitives += NumPrimitives;
-
 		break;
 	}
+
 	case RouteId_SignalFence:
 	{
 		uint32 QueueId = EventData.GetValue<uint32>("QueueId");
@@ -429,9 +333,9 @@ bool FGpuProfilerAnalyzer::OnEvent(uint16 RouteId, EStyle Style, const FOnEventC
 
 		FAnalysisSessionEditScope _(Session);
 		TimingProfilerProvider.AddGpuSignalFence(QueueId, Fence);
-
 		break;
 	}
+
 	case RouteId_WaitFence:
 	{
 		uint32 QueueId = EventData.GetValue<uint32>("QueueId");
@@ -446,183 +350,15 @@ bool FGpuProfilerAnalyzer::OnEvent(uint16 RouteId, EStyle Style, const FOnEventC
 
 		FAnalysisSessionEditScope _(Session);
 		TimingProfilerProvider.AddGpuWaitFence(QueueId, Fence);
-
 		break;
 	}
 
-	//////////////////////////////////////////////////
-	// The Old GPU Profiler
-
-	case RouteId_EventSpec:
-	{
-		uint32 EventType = EventData.GetValue<uint32>("EventType");
-		const auto& Name = EventData.GetArray<UTF16CHAR>("Name");
-
-		auto NameTChar = StringCast<TCHAR>(Name.GetData(), Name.Num());
-		uint32* TimerIndexPtr = EventTypeMap.Find(EventType);
-		if (!TimerIndexPtr)
-		{
-			FAnalysisSessionEditScope _(Session);
-			uint32 TimerId;
-#if UE_INSIGHTS_MERGE_UNACCOUNTED_TIMERS
-			FStringView Unaccounted = TEXTVIEW("Unaccounted -");
-			if (NameTChar.Length() > Unaccounted.Len() &&
-				FCString::Strncmp(Unaccounted.GetData(), NameTChar.Get(), Unaccounted.Len()) == 0)
-			{
-				TimerId = TimingProfilerProvider.AddGpuTimer(TEXTVIEW("Unaccounted"));
-			}
-			else
-#endif // UE_INSIGHTS_MERGE_UNACCOUNTED_TIMERS
-			{
-				TimerId = TimingProfilerProvider.AddGpuTimer(FStringView(NameTChar.Get(), NameTChar.Length()));
-			}
-			EventTypeMap.Add(EventType, TimerId);
-		}
-		else
-		{
-			FAnalysisSessionEditScope _(Session);
-			TimingProfilerProvider.SetTimerName(*TimerIndexPtr, FStringView(NameTChar.Get(), NameTChar.Length()));
-		}
-		break;
-	}
-	case RouteId_Frame:
-	case RouteId_Frame2:
-	{
-		FAnalysisSessionEditScope _(Session);
-
-		TraceServices::FTimingProfilerProvider::TimelineInternal& ThisTimeline = (RouteId == RouteId_Frame) ?
-			TimingProfilerProvider.EditGpuTimeline() :
-			TimingProfilerProvider.EditGpu2Timeline();
-		double& ThisMinTime = (RouteId == RouteId_Frame) ? MinTime : MinTime2;
-
-		const auto& Data = EventData.GetArray<uint8>("Data");
-		const uint8* BufferPtr = Data.GetData();
-		const uint8* BufferEnd = BufferPtr + Data.Num();
-
-		uint64 CalibrationBias = EventData.GetValue<uint64>("CalibrationBias");
-		uint64 LastTimestamp = EventData.GetValue<uint64>("TimestampBase");
-		uint32 RenderingFrameNumber = EventData.GetValue<uint32>("RenderingFrameNumber");
-
-		++NumFrames;
-
-		double LastTime = 0.0;
-		uint32 CurrentDepth = 0;
-		bool bHasErrors = false;
-
-		while (BufferPtr < BufferEnd)
-		{
-			uint64 DecodedTimestamp = FTraceAnalyzerUtils::Decode7bit(BufferPtr);
-			uint64 ActualTimestamp = (DecodedTimestamp >> 1) + LastTimestamp;
-			LastTimestamp = ActualTimestamp;
-			LastTime = double(ActualTimestamp + CalibrationBias) * 0.000001;
-			LastTime += Context.EventTime.AsSeconds(0);
-
-			if (LastTime < 0.0)
-			{
-				if (DecodedTimestamp & 1ull)
-				{
-					BufferPtr += sizeof(uint32);
-				}
-				bHasErrors = true;
-				continue;
-			}
-
-			// If it advances with more than 1h, it is probably a wrong timestamp.
-			if (LastTime > ThisMinTime + 3600.0 && ThisMinTime != DBL_MIN)
-			{
-				LastTime = ThisMinTime;
-				bHasErrors = true;
-			}
-
-			// The monolithic timeline assumes that timestamps are ever increasing, but
-			// with gpu/cpu calibration and drift there can be a tiny bit of overlap between
-			// frames. So we just clamp.
-			if (ThisMinTime > LastTime)
-			{
-				LastTime = ThisMinTime;
-			}
-			ThisMinTime = LastTime;
-
-			if (DecodedTimestamp & 1ull)
-			{
-				uint32 EventType = *reinterpret_cast<const uint32*>(BufferPtr);
-				BufferPtr += sizeof(uint32);
-				if (EventTypeMap.Contains(EventType))
-				{
-					FTimingProfilerEvent Event;
-					Event.TimerIndex = EventTypeMap[EventType];
-					ThisTimeline.AppendBeginEvent(LastTime, Event);
-				}
-				else
-				{
-					FTimingProfilerEvent Event;
-					Event.TimerIndex = TimingProfilerProvider.AddGpuTimer(TEXTVIEW("<unknown>"));
-					EventTypeMap.Add(EventType, Event.TimerIndex);
-					ThisTimeline.AppendBeginEvent(LastTime, Event);
-				}
-				++CurrentDepth;
-			}
-			else
-			{
-				if (CurrentDepth > 0)
-				{
-					--CurrentDepth;
-				}
-				ThisTimeline.AppendEndEvent(LastTime);
-			}
-		}
-		check(BufferPtr == BufferEnd);
-		check(CurrentDepth == 0);
-		if (bHasErrors && ++NumFramesWithErrors <= 100)
-		{
-			UE_LOG(LogTraceServices, Error, TEXT("[GpuProfiler] The rendering frame %u has invalid timestamps!"), RenderingFrameNumber);
-		}
-		Session.UpdateDurationSeconds(LastTime);
-		break;
-	}
-
-	//////////////////////////////////////////////////
-
 	} // switch (RouteId)
 
 	return true;
 }
 
-uint32 FGpuProfilerAnalyzer::GetOrAddTimerId(const FString& Breadcrumb)
-{
-	const uint32* FoundTimerId = BreadcrumbMap.Find(Breadcrumb);
-	if (FoundTimerId)
-	{
-		return *FoundTimerId;
-	}
-	else
-	{
-		FAnalysisSessionEditScope _(Session);
-		uint32 TimerId = TimingProfilerProvider.AddGpuTimer(Breadcrumb);
-		const TCHAR* TimerName = nullptr;
-		TimingProfilerProvider.ReadTimers([TimerId, &TimerName](const ITimingProfilerTimerReader& TimerReader)
-			{
-				const FTimingProfilerTimer* Timer = TimerReader.GetTimer(TimerId);
-				if (Timer)
-				{
-					TimerName = Timer->Name;
-				}
-			});
-		BreadcrumbMap.Add(Breadcrumb, TimerId);
-		TimerMap.Add(TimerId, TimerName);
-		return TimerId;
-	}
-}
-
-const TCHAR* FGpuProfilerAnalyzer::GetTimerName(uint32 TimerId) const
-{
-	if (TimerId == InvalidTimerId)
-	{
-		return TEXT("<invalid>");
-	}
-	const TCHAR* const* TimerNamePtr = TimerMap.Find(TimerId);
-	return TimerNamePtr ? *TimerNamePtr : TEXT("<unknown>");
-}
+////////////////////////////////////////////////////////////////////////////////////////////////////
 
 FGpuProfilerAnalyzer::FQueue& FGpuProfilerAnalyzer::GetOrAddQueue(uint32 QueueId)
 {
@@ -638,6 +374,8 @@ FGpuProfilerAnalyzer::FQueue& FGpuProfilerAnalyzer::GetOrAddQueue(uint32 QueueId
 	return NewQueue;
 }
 
+////////////////////////////////////////////////////////////////////////////////////////////////////
+
 void FGpuProfilerAnalyzer::BeginEvent(FQueue& Queue, int32 StackIndex, double BeginEventTime, uint32 BeginEventTimerId)
 {
 	Queue.Stack[StackIndex].Push({ BeginEventTime, BeginEventTimerId });
@@ -672,6 +410,8 @@ void FGpuProfilerAnalyzer::BeginEvent(FQueue& Queue, int32 StackIndex, double Be
 	}
 }
 
+////////////////////////////////////////////////////////////////////////////////////////////////////
+
 void FGpuProfilerAnalyzer::EndEvent(FQueue& Queue, int32 StackIndex, double EndEventTime, uint32 EndEventTimerId)
 {
 	Queue.LastTime = EndEventTime;
@@ -761,6 +501,225 @@ void FGpuProfilerAnalyzer::EndEvent(FQueue& Queue, int32 StackIndex, double EndE
 	}
 }
 
+////////////////////////////////////////////////////////////////////////////////////////////////////
+
+void FGpuProfilerAnalyzer::OnEventBreadcrumbSpec(const FOnEventContext& Context)
+{
+	const auto& EventData = Context.EventData;
+
+	uint32 SpecId = EventData.GetValue<uint32>("SpecId");
+
+	FString Name, NameFormat;
+	EventData.GetString("StaticName", Name);
+	EventData.GetString("NameFormat", NameFormat);
+	TArrayView<const uint8> FieldNames = EventData.GetArrayView<uint8>("FieldNames");
+
+	GPUPROFILER_LOG_API_L2(TEXT("[GpuProfiler] EventBreadcrumbSpec Id=%u Name=\"%s\" NameFormat=\"%s\" Fields[]={%u bytes}"),
+		SpecId, *Name, *NameFormat, FieldNames.Num());
+
+	if (Name.Compare(NameFormat) == 0)
+	{
+		Name.Empty();
+	}
+
+	int Index = 0;
+	while (Index < NameFormat.Len() && Index < Name.Len())
+	{
+		if (Name[Index] != NameFormat[Index])
+		{
+			break;
+		}
+
+		++Index;
+	}
+
+	if (Index > 1)
+	{
+		NameFormat.MidInline(Index);
+	}
+
+	if (Name.IsEmpty())
+	{
+		const FString Separators = TEXT("% (=");
+		Index = 0;
+		bool bIsInFormatSpecifier = false;
+		while (Index < NameFormat.Len())
+		{
+			int32 SpecIndex = -1;
+			if (Separators.FindChar(NameFormat[Index], SpecIndex))
+			{
+				Name = NameFormat.Left(Index);
+				NameFormat.MidInline(Index);
+				break;
+			}
+
+			++Index;
+		}
+
+		if (Name.IsEmpty())
+		{
+			if (!NameFormat.IsEmpty())
+			{
+				Name = NameFormat;
+				NameFormat.Empty();
+			}
+			else
+			{
+				Name = Session.StoreString(TEXT("Unknown"));
+			}
+		}
+	}
+
+	FMetadataSpec Spec;
+	if (!NameFormat.IsEmpty())
+	{
+		Spec.Format = Session.StoreString(NameFormat);
+	}
+
+	if (FieldNames.Num() > 0)
+	{
+		FMemoryReaderView MemoryReader(FieldNames);
+		FCborReader CborReader(&MemoryReader, ECborEndianness::StandardCompliant);
+		FCborContext CborContext;
+
+		while (CborReader.ReadNext(CborContext))
+		{
+			if (CborContext.MajorType() == ECborCode::TextString)
+			{
+				FString Field = CborContext.AsString();
+				Spec.FieldNames.Add(Session.StoreString(Field));
+			}
+		}
+	}
+
+	uint32 TimerId;
+	if (const uint32* FoundTimerIdBySpecId = SpecIdToTimerIdMap.Find(SpecId))
+	{
+		TimerId = *FoundTimerIdBySpecId;
+		SetTimerName(SpecId, TimerId, Name);
+	}
+	else
+	{
+		GPUPROFILER_LOG_API_L2(TEXT("[GpuProfiler] EventBreadcrumbSpec with a new timer (\"%s\")"), *Name);
+		TimerId = AddGpuTimer(SpecId, Name);
+	}
+
+	if (Spec.FieldNames.Num() > 0 || Spec.Format != nullptr)
+	{
+		FAnalysisSessionEditScope _(Session);
+		uint32 MetadataSpecId = TimingProfilerProvider.AddMetadataSpec(MoveTemp(Spec));
+
+		FTimingProfilerTimer PreviousTimer;
+		TimingProfilerProvider.ReadTimers(
+			[TimerId, &PreviousTimer]
+			(const ITimingProfilerTimerReader& TimerReader)
+			{
+				const FTimingProfilerTimer* PreviousTimerPtr = TimerReader.GetTimer(TimerId);
+				if (PreviousTimerPtr)
+				{
+					PreviousTimer = *PreviousTimerPtr;
+				}
+			});
+		if (PreviousTimer.HasValidMetadataSpecId())
+		{
+			TimerId = AddGpuTimer(SpecId, Name);
+
+			if (++ErrorData.NumTimerWarnings < ErrorData.NumMaxWarnings)
+			{
+				const FMetadataSpec* PreviousMetadataSpec = TimingProfilerProvider.GetMetadataSpec(PreviousTimer.MetadataSpecId);
+				check(PreviousMetadataSpec);
+				UE_LOG(LogTraceServices, Warning, TEXT("[GpuProfiler] EventBreadcrumbSpec -- de-duplicated timer for SpecId %u -- old: %u (%s) metadata=\"%s\" (%d fields) --> new: %u (%s) metadata=\"%s\" (%d fields) !!!"),
+					SpecId,
+					PreviousTimer.Id, PreviousTimer.Name, PreviousMetadataSpec->Format, PreviousMetadataSpec->FieldNames.Num(),
+					TimerId, *Name, Spec.Format, Spec.FieldNames.Num());
+			}
+		}
+
+		TimingProfilerProvider.SetMetadataSpec(TimerId, MetadataSpecId);
+	}
+}
+
+////////////////////////////////////////////////////////////////////////////////////////////////////
+
+uint32 FGpuProfilerAnalyzer::GetOrAddTimer(uint32 SpecId, const FString& Breadcrumb)
+{
+	if (SpecId != ~0u)
+	{
+		if (const uint32* FoundTimerIdBySpecId = SpecIdToTimerIdMap.Find(SpecId))
+		{
+			return *FoundTimerIdBySpecId;
+		}
+	}
+
+	// Merge by name for the GPU timers is disabled as it is not safe to merge when metdata can be attached late.
+	//if (const uint32* FoundTimerId = BreadcrumbMap.Find(Breadcrumb))
+	//{
+	//	return *FoundTimerId;
+	//}
+
+	return AddGpuTimer(SpecId, Breadcrumb);
+}
+////////////////////////////////////////////////////////////////////////////////////////////////////
+
+uint32 FGpuProfilerAnalyzer::AddGpuTimer(uint32 SpecId, const FString& Breadcrumb)
+{
+	uint32 TimerId;
+	const TCHAR* TimerName = nullptr;
+
+	{
+		FAnalysisSessionEditScope _(Session);
+		TimerId = TimingProfilerProvider.AddGpuTimer(Breadcrumb);
+		TimingProfilerProvider.ReadTimers(
+			[TimerId, &TimerName]
+			(const ITimingProfilerTimerReader& TimerReader)
+			{
+				const FTimingProfilerTimer* Timer = TimerReader.GetTimer(TimerId);
+				if (Timer)
+				{
+					TimerName = Timer->Name;
+				}
+			});
+	}
+
+	TimerMap.Add(TimerId, TimerName);
+	BreadcrumbMap.Add(Breadcrumb, TimerId);
+
+	if (SpecId != ~0u)
+	{
+		SpecIdToTimerIdMap.Add(SpecId, TimerId);
+
+		GPUPROFILER_LOG_API_L3(TEXT("[GpuProfiler] --> AddGpuTimer(SpecId=%u, Name=\"%s\") --> TimerId=%u"), SpecId, TimerName, TimerId);
+	}
+	else
+	{
+		GPUPROFILER_LOG_API_L3(TEXT("[GpuProfiler] --> AddGpuTimer(Name=\"%s\") --> TimerId=%u"), TimerName, TimerId);
+	}
+
+	return TimerId;
+}
+
+////////////////////////////////////////////////////////////////////////////////////////////////////
+
+const TCHAR* FGpuProfilerAnalyzer::GetTimerName(uint32 TimerId) const
+{
+	if (TimerId == InvalidTimerId)
+	{
+		return TEXT("<invalid>");
+	}
+	const TCHAR* const* TimerNamePtr = TimerMap.Find(TimerId);
+	return TimerNamePtr ? *TimerNamePtr : TEXT("<unknown>");
+}
+
+////////////////////////////////////////////////////////////////////////////////////////////////////
+
+void FGpuProfilerAnalyzer::SetTimerName(uint32 SpecId, uint32 TimerId, FStringView TimerName)
+{
+	FAnalysisSessionEditScope _(Session);
+	TimingProfilerProvider.SetTimerName(TimerId, TimerName);
+}
+
+////////////////////////////////////////////////////////////////////////////////////////////////////
+
 void FGpuProfilerAnalyzer::InitCounters(FQueue& FoundQueue)
 {
 	FoundQueue.NumDrawsCounter = EditableCounterProvider.CreateEditableCounter();
@@ -770,6 +729,8 @@ void FGpuProfilerAnalyzer::InitCounters(FQueue& FoundQueue)
 	FoundQueue.NumPrimitivesCounter->SetIsFloatingPoint(false);
 }
 
+////////////////////////////////////////////////////////////////////////////////////////////////////
+
 void FGpuProfilerAnalyzer::InitCountersDesc(FQueue& FoundQueue, uint8 Gpu, uint8 Index, const TCHAR* Name)
 {
 	const FString DrawsCounterName = FString::Printf(TEXT("NumDraws_GPU%u_%s_%u"), Gpu, Name, Index);
@@ -782,6 +743,10 @@ void FGpuProfilerAnalyzer::InitCountersDesc(FQueue& FoundQueue, uint8 Gpu, uint8
 	FoundQueue.NumPrimitivesCounter->SetDescription(TEXT("The number of primitives on the specified queue."));
 }
 
+////////////////////////////////////////////////////////////////////////////////////////////////////
+
 } // namespace TraceServices
 
-#undef UE_INSIGHTS_MERGE_UNACCOUNTED_TIMERS
+#undef GPUPROFILER_LOG_API_L3
+#undef GPUPROFILER_LOG_API_L2
+#undef GPUPROFILER_LOG_API_L1
diff --git a/Engine/Source/Developer/TraceServices/Private/Analyzers/GpuProfilerTraceAnalysis.h b/Engine/Source/Developer/TraceServices/Private/Analyzers/GpuProfilerTraceAnalysis.h
--- a/Engine/Source/Developer/TraceServices/Private/Analyzers/GpuProfilerTraceAnalysis.h
+++ b/Engine/Source/Developer/TraceServices/Private/Analyzers/GpuProfilerTraceAnalysis.h
@@ -2,22 +2,28 @@
 
 #pragma once
 
+#include "CoreTypes.h"
+#include "Containers/Map.h"
+#include "Containers/StringFwd.h"
+#include "Containers/UnrealString.h"
+
+// TraceAnalysis
 #include "Trace/Analyzer.h"
-#include "Model/TimingProfilerPrivate.h"
 
 namespace TraceServices
 {
-	class FAnalysisSession;
-	class IEditableCounter;
-	class IEditableCounterProvider;
 
-class FGpuProfilerAnalyzer
-	: public UE::Trace::IAnalyzer
+class FAnalysisSession;
+class FTimingProfilerProvider;
+class IEditableCounter;
+class IEditableCounterProvider;
+
+// The New GPU Profiler
+class FGpuProfilerAnalyzer : public UE::Trace::IAnalyzer
 {
 private:
 	enum : uint16
 	{
-		// The New GPU Profiler
 		RouteId_Init,
 		RouteId_QueueSpec,
 		RouteId_EventFrameBoundary,
@@ -30,11 +36,6 @@ private:
 		RouteId_EventStats,
 		RouteId_SignalFence,
 		RouteId_WaitFence,
-
-		// The Old GPU Profiler
-		RouteId_EventSpec,
-		RouteId_Frame, // GPU Index 0
-		RouteId_Frame2, // GPU Index 1
 	};
 
 	struct FOpenEvent
@@ -61,6 +62,7 @@ private:
 		static constexpr uint32 NumMaxErrors = 100;
 
 		uint32 NumWarnings = 0;
+		uint32 NumTimerWarnings = 0;
 		uint32 NumErrors = 0;
 
 		uint32 NumInterleavedEvents = 0;
@@ -80,11 +82,15 @@ public:
 	virtual bool OnEvent(uint16 RouteId, EStyle Style, const FOnEventContext& Context) override;
 
 private:
-	uint32 GetOrAddTimerId(const FString& Breadcrumb);
-	const TCHAR* GetTimerName(uint32 TimerId) const;
 	FQueue& GetOrAddQueue(uint32 QueueId);
 	void BeginEvent(FQueue& Queue, int32 StackIndex, double Time, uint32 TimerId);
 	void EndEvent(FQueue& Queue, int32 StackIndex, double Time, uint32 TimerId);
+	void OnEventBreadcrumbSpec(const FOnEventContext& Context);
+
+	uint32 GetOrAddTimer(uint32 SpecId, const FString& Breadcrumb);
+	uint32 AddGpuTimer(uint32 SpecId, const FString& Breadcrumb);
+	const TCHAR* GetTimerName(uint32 TimerId) const;
+	void SetTimerName(uint32 SpecId, uint32 TimerId, FStringView TimerName);
 
 	void InitCounters(FQueue& FoundQueue);
 	void InitCountersDesc(FQueue& FoundQueue, uint8 Gpu, uint8 Index, const TCHAR* Name);
@@ -94,27 +100,15 @@ private:
 	FTimingProfilerProvider& TimingProfilerProvider;
 	IEditableCounterProvider& EditableCounterProvider;
 
-	//////////////////////////////////////////////////
-	// The New GPU Profiler
-
 	uint32 Version = 0;
 	static constexpr uint32 InvalidTimerId = 0;
 	uint32 GpuWorkTimerId = InvalidTimerId;
 	uint32 GpuWaitTimerId = InvalidTimerId;
-	TMap<uint32, uint32> BreadcrumbSpecMap; // breadcrumb spec id --> GPU timer id
+	TMap<uint32, uint32> SpecIdToTimerIdMap; // breadcrumb spec id --> GPU timer id
 	TMap<FString, uint32> BreadcrumbMap; // breadcrumb name --> GPU timer id
 	TMap<uint32, const TCHAR*> TimerMap; // GPU timer id --> persistent timer name
 	TMap<uint32, FQueue> Queues; // QueueId --> FQueue
 	FErrorData ErrorData;
-
-	//////////////////////////////////////////////////
-	// The Old GPU Profiler
-
-	TMap<uint64, uint32> EventTypeMap;
-	double MinTime = DBL_MIN;
-	double MinTime2 = DBL_MIN;
-	uint32 NumFrames = 0;
-	uint32 NumFramesWithErrors = 0;
 };
 
 } // namespace TraceServices
diff --git a/Engine/Source/Developer/TraceServices/Private/Analyzers/OldGpuProfilerTraceAnalysis.cpp b/Engine/Source/Developer/TraceServices/Private/Analyzers/OldGpuProfilerTraceAnalysis.cpp
new file mode 100644
--- /dev/null
+++ b/Engine/Source/Developer/TraceServices/Private/Analyzers/OldGpuProfilerTraceAnalysis.cpp
@@ -0,0 +1,218 @@
+// Copyright Epic Games, Inc. All Rights Reserved.
+
+#include "OldGpuProfilerTraceAnalysis.h"
+
+#include "CborWriter.h"
+#include "CborReader.h"
+#include "HAL/LowLevelMemTracker.h"
+#include "Serialization/MemoryReader.h"
+#include "Serialization/MemoryWriter.h"
+
+// TraceServices
+#include "AnalysisServicePrivate.h"
+#include "Common/Utils.h"
+#include "Model/TimingProfilerPrivate.h"
+
+#define UE_INSIGHTS_MERGE_UNACCOUNTED_TIMERS 0
+
+namespace TraceServices
+{
+
+////////////////////////////////////////////////////////////////////////////////////////////////////
+
+FOldGpuProfilerAnalyzer::FOldGpuProfilerAnalyzer(FAnalysisSession& InSession, FTimingProfilerProvider& InTimingProfilerProvider)
+	: Session(InSession)
+	, TimingProfilerProvider(InTimingProfilerProvider)
+{
+}
+
+////////////////////////////////////////////////////////////////////////////////////////////////////
+
+void FOldGpuProfilerAnalyzer::OnAnalysisBegin(const FOnAnalysisContext& Context)
+{
+	auto& Builder = Context.InterfaceBuilder;
+
+	// The Old GPU Profiler (deprecated in UE 5.6)
+	// Analysis code is maintained for backward compatibility with old traces.
+	Builder.RouteEvent(RouteId_EventSpec, "GpuProfiler", "EventSpec");
+	Builder.RouteEvent(RouteId_Frame,     "GpuProfiler", "Frame");
+	Builder.RouteEvent(RouteId_Frame2,    "GpuProfiler", "Frame2");
+}
+
+////////////////////////////////////////////////////////////////////////////////////////////////////
+
+void FOldGpuProfilerAnalyzer::OnAnalysisEnd()
+{
+	if (NumFramesWithErrors > 0)
+	{
+		UE_LOG(LogTraceServices, Error, TEXT("[GpuProfiler] Frames with errors: %u"), NumFramesWithErrors);
+	}
+	
+	if (NumFrames > 0 || EventTypeMap.Num() > 0)
+	{
+		UE_LOG(LogTraceServices, Log, TEXT("[GpuProfiler] Analysis completed (%u frames, %d timers)."), NumFrames, EventTypeMap.Num());
+	}
+}
+
+////////////////////////////////////////////////////////////////////////////////////////////////////
+
+bool FOldGpuProfilerAnalyzer::OnEvent(uint16 RouteId, EStyle Style, const FOnEventContext& Context)
+{
+	LLM_SCOPE_BYNAME(TEXT("Insights/FOldGpuProfilerAnalyzer"));
+
+	switch (RouteId)
+	{
+	case RouteId_EventSpec:
+		OnEventSpec(Context);
+		break;
+	case RouteId_Frame:
+		OnFrame(Context, 0);
+		break;
+	case RouteId_Frame2:
+		OnFrame(Context, 1);
+		break;
+	}
+
+	return true;
+}
+
+////////////////////////////////////////////////////////////////////////////////////////////////////
+
+void FOldGpuProfilerAnalyzer::OnEventSpec(const UE::Trace::IAnalyzer::FOnEventContext& Context)
+{
+	const auto& EventData = Context.EventData;
+
+	uint32 EventType = EventData.GetValue<uint32>("EventType");
+	const auto& Name = EventData.GetArray<UTF16CHAR>("Name");
+
+	auto NameTChar = StringCast<TCHAR>(Name.GetData(), Name.Num());
+	uint32* TimerIndexPtr = EventTypeMap.Find(EventType);
+	if (!TimerIndexPtr)
+	{
+		uint32 TimerId;
+#if UE_INSIGHTS_MERGE_UNACCOUNTED_TIMERS
+		FStringView Unaccounted = TEXTVIEW("Unaccounted -");
+		if (NameTChar.Length() > Unaccounted.Len() &&
+			FCString::Strncmp(Unaccounted.GetData(), NameTChar.Get(), Unaccounted.Len()) == 0)
+		{
+			FAnalysisSessionEditScope _(Session);
+			TimerId = TimingProfilerProvider.AddGpuTimer(TEXTVIEW("Unaccounted"));
+		}
+		else
+#endif // UE_INSIGHTS_MERGE_UNACCOUNTED_TIMERS
+		{
+			FAnalysisSessionEditScope _(Session);
+			TimerId = TimingProfilerProvider.AddGpuTimer(FStringView(NameTChar.Get(), NameTChar.Length()));
+		}
+		EventTypeMap.Add(EventType, TimerId);
+	}
+	else
+	{
+		FAnalysisSessionEditScope _(Session);
+		TimingProfilerProvider.SetTimerName(*TimerIndexPtr, FStringView(NameTChar.Get(), NameTChar.Length()));
+	}
+}
+
+////////////////////////////////////////////////////////////////////////////////////////////////////
+
+void FOldGpuProfilerAnalyzer::OnFrame(const UE::Trace::IAnalyzer::FOnEventContext& Context, uint32 GpuIndex)
+{
+	const auto& EventData = Context.EventData;
+
+	FAnalysisSessionEditScope _(Session);
+
+	TraceServices::FTimingProfilerProvider::TimelineInternal& ThisTimeline = (GpuIndex == 0) ?
+		TimingProfilerProvider.EditGpuTimeline() :
+		TimingProfilerProvider.EditGpu2Timeline();
+	double& ThisMinTime = (GpuIndex == 0) ? MinTime : MinTime2;
+
+	const auto& Data = EventData.GetArray<uint8>("Data");
+	const uint8* BufferPtr = Data.GetData();
+	const uint8* BufferEnd = BufferPtr + Data.Num();
+
+	uint64 CalibrationBias = EventData.GetValue<uint64>("CalibrationBias");
+	uint64 LastTimestamp = EventData.GetValue<uint64>("TimestampBase");
+	uint32 RenderingFrameNumber = EventData.GetValue<uint32>("RenderingFrameNumber");
+
+	++NumFrames;
+
+	double LastTime = 0.0;
+	uint32 CurrentDepth = 0;
+	bool bHasErrors = false;
+
+	while (BufferPtr < BufferEnd)
+	{
+		uint64 DecodedTimestamp = FTraceAnalyzerUtils::Decode7bit(BufferPtr);
+		uint64 ActualTimestamp = (DecodedTimestamp >> 1) + LastTimestamp;
+		LastTimestamp = ActualTimestamp;
+		LastTime = double(ActualTimestamp + CalibrationBias) * 0.000001;
+		LastTime += Context.EventTime.AsSeconds(0);
+
+		if (LastTime < 0.0)
+		{
+			if (DecodedTimestamp & 1ull)
+			{
+				BufferPtr += sizeof(uint32);
+			}
+			bHasErrors = true;
+			continue;
+		}
+
+		// If it advances with more than 1h, it is probably a wrong timestamp.
+		if (LastTime > ThisMinTime + 3600.0 && ThisMinTime != DBL_MIN)
+		{
+			LastTime = ThisMinTime;
+			bHasErrors = true;
+		}
+
+		// The monolithic timeline assumes that timestamps are ever increasing, but
+		// with GPU/CPU calibration and drift there can be a tiny bit of overlap between
+		// frames. So we just clamp.
+		if (ThisMinTime > LastTime)
+		{
+			LastTime = ThisMinTime;
+		}
+		ThisMinTime = LastTime;
+
+		if (DecodedTimestamp & 1ull)
+		{
+			uint32 EventType = *reinterpret_cast<const uint32*>(BufferPtr);
+			BufferPtr += sizeof(uint32);
+			if (EventTypeMap.Contains(EventType))
+			{
+				FTimingProfilerEvent Event;
+				Event.TimerIndex = EventTypeMap[EventType];
+				ThisTimeline.AppendBeginEvent(LastTime, Event);
+			}
+			else
+			{
+				FTimingProfilerEvent Event;
+				Event.TimerIndex = TimingProfilerProvider.AddGpuTimer(TEXTVIEW("<unknown>"));
+				EventTypeMap.Add(EventType, Event.TimerIndex);
+				ThisTimeline.AppendBeginEvent(LastTime, Event);
+			}
+			++CurrentDepth;
+		}
+		else
+		{
+			if (CurrentDepth > 0)
+			{
+				--CurrentDepth;
+			}
+			ThisTimeline.AppendEndEvent(LastTime);
+		}
+	}
+	check(BufferPtr == BufferEnd);
+	check(CurrentDepth == 0);
+	if (bHasErrors && ++NumFramesWithErrors <= 100)
+	{
+		UE_LOG(LogTraceServices, Error, TEXT("[GpuProfiler] The rendering frame %u has invalid timestamps!"), RenderingFrameNumber);
+	}
+	Session.UpdateDurationSeconds(LastTime);
+}
+
+////////////////////////////////////////////////////////////////////////////////////////////////////
+
+} // namespace TraceServices
+
+#undef UE_INSIGHTS_MERGE_UNACCOUNTED_TIMERS
diff --git a/Engine/Source/Developer/TraceServices/Private/Analyzers/OldGpuProfilerTraceAnalysis.h b/Engine/Source/Developer/TraceServices/Private/Analyzers/OldGpuProfilerTraceAnalysis.h
new file mode 100644
--- /dev/null
+++ b/Engine/Source/Developer/TraceServices/Private/Analyzers/OldGpuProfilerTraceAnalysis.h
@@ -0,0 +1,48 @@
+// Copyright Epic Games, Inc. All Rights Reserved.
+
+#pragma once
+
+#include "CoreTypes.h"
+#include "Containers/Map.h"
+#include "Trace/Analyzer.h"
+
+namespace TraceServices
+{
+
+class FAnalysisSession;
+class FTimingProfilerProvider;
+
+// The Old GPU Profiler (deprecated in UE 5.6)
+// Analysis code is maintained for backward compatibility with old traces.
+class FOldGpuProfilerAnalyzer : public UE::Trace::IAnalyzer
+{
+private:
+	enum : uint16
+	{
+		RouteId_EventSpec,
+		RouteId_Frame, // GPU Index 0
+		RouteId_Frame2, // GPU Index 1
+	};
+
+public:
+	FOldGpuProfilerAnalyzer(FAnalysisSession& Session, FTimingProfilerProvider& TimingProfilerProvider);
+	virtual void OnAnalysisBegin(const FOnAnalysisContext& Context) override;
+	virtual void OnAnalysisEnd() override;
+	virtual bool OnEvent(uint16 RouteId, EStyle Style, const FOnEventContext& Context) override;
+
+private:
+	void OnEventSpec(const FOnEventContext& Context);
+	void OnFrame(const FOnEventContext& Context, uint32 GpuIndex = 0);
+
+private:
+	FAnalysisSession& Session;
+	FTimingProfilerProvider& TimingProfilerProvider;
+
+	TMap<uint64, uint32> EventTypeMap;
+	double MinTime = DBL_MIN;
+	double MinTime2 = DBL_MIN;
+	uint32 NumFrames = 0;
+	uint32 NumFramesWithErrors = 0;
+};
+
+} // namespace TraceServices
diff --git a/Engine/Source/Developer/TraceServices/Private/Model/TimingProfiler.cpp b/Engine/Source/Developer/TraceServices/Private/Model/TimingProfiler.cpp
--- a/Engine/Source/Developer/TraceServices/Private/Model/TimingProfiler.cpp
+++ b/Engine/Source/Developer/TraceServices/Private/Model/TimingProfiler.cpp
@@ -2,12 +2,20 @@
 
 #include "TraceServices/Model/TimingProfiler.h"
 
-#include "TraceServices/Model/Frames.h"
-#include "Model/TimingProfilerPrivate.h"
+#include "Templates/TypeHash.h"
+
+// TraceServices
 #include "AnalysisServicePrivate.h"
 #include "Common/StringStore.h"
 #include "Common/TimelineStatistics.h"
-#include "Templates/TypeHash.h"
+#include "Common/Utils.h"
+#include "Model/TimingProfilerPrivate.h"
+#include "TraceServices/Model/Frames.h"
+
+#define TIMING_PROFILER_DETECT_TIMER_CHANGES 0
+#define TIMING_PROFILER_LOG_INVALID_METADATA_1 //{ UE_LOG(LogTraceServices, Warning, TEXT("[TimingProfiler][1] MetadataTimerId=%u >= 0"), MetadataTimerId); UE_DEBUG_BREAK(); }
+#define TIMING_PROFILER_LOG_INVALID_METADATA_2 //{ UE_LOG(LogTraceServices, Warning, TEXT("[TimingProfiler][2] MetadataTimerId=%u >= Num=%u"), MetadataTimerId, uint32(Metadatas.Num())); UE_DEBUG_BREAK(); }
+#define TIMING_PROFILER_LOG_INVALID_METADATA_SPEC //{ UE_LOG(LogTraceServices, Warning, TEXT("[TimingProfiler][3] MetadataSpecId=%u >= Num=%u"), MetadataSpecId, uint32(MetadataSpecs.Num())); UE_DEBUG_BREAK(); }
 
 namespace TraceServices
 {
@@ -152,7 +160,16 @@ void FTimingProfilerProvider::SetTimerName(uint32 TimerId, FStringView Name)
 	if (ensure(TimerId < uint32(Timers.Num())))
 	{
 		FTimingProfilerTimer& Timer = Timers[TimerId];
-		Timer.Name = Session.StoreString(Name);
+
+		const TCHAR* NewName = Session.StoreString(Name);
+#if TIMING_PROFILER_DETECT_TIMER_CHANGES
+		if (!ensure(Timer.Name == nullptr || Timer.Name == NewName || Timer.Name[0] == TEXT('<')))
+		{
+			UE_LOG(LogTraceServices, Warning, TEXT("[TimingProfiler] Name changed for timer %u: from \"%s\" to \"%s\" !!!"),
+				TimerId, Timer.Name, NewName);
+		}
+#endif
+		Timer.Name = NewName;
 	}
 }
 
@@ -163,8 +180,27 @@ void FTimingProfilerProvider::SetTimerNameAndLocation(uint32 TimerId, FStringVie
 	if (ensure(TimerId < uint32(Timers.Num())))
 	{
 		FTimingProfilerTimer& Timer = Timers[TimerId];
-		Timer.Name = Session.StoreString(Name);
-		Timer.File = File ? Session.StoreString(File) : nullptr;
+
+		const TCHAR* NewName = Session.StoreString(Name);
+#if TIMING_PROFILER_DETECT_TIMER_CHANGES
+		if (!ensure(Timer.Name == nullptr || Timer.Name == NewName || Timer.Name[0] == TEXT('<')))
+		{
+			UE_LOG(LogTraceServices, Warning, TEXT("[TimingProfiler] Name changed for timer %u: from \"%s\" to \"%s\" !!!"),
+				TimerId, Timer.Name, NewName);
+		}
+#endif
+		Timer.Name = NewName;
+
+		const TCHAR* NewFile = File ? Session.StoreString(File) : nullptr;
+#if TIMING_PROFILER_DETECT_TIMER_CHANGES
+		if (!ensure(Timer.File == nullptr || Timer.File == NewFile))
+		{
+			UE_LOG(LogTraceServices, Warning, TEXT("[TimingProfiler] Location changed for timer %u (%s): from \"%s\" (%u) to \"%s\" (%u) !!!"),
+				TimerId, Timer.Name, Timer.File, Timer.Line, NewFile, Line);
+		}
+#endif
+		Timer.File = NewFile;
+
 		Timer.Line = Line;
 	}
 }
@@ -176,7 +212,17 @@ void FTimingProfilerProvider::SetTimerLocation(uint32 TimerId, const TCHAR* File
 	if (ensure(TimerId < uint32(Timers.Num())))
 	{
 		FTimingProfilerTimer& Timer = Timers[TimerId];
-		Timer.File = File ? Session.StoreString(File) : nullptr;
+
+		const TCHAR* NewFile = File ? Session.StoreString(File) : nullptr;
+#if TIMING_PROFILER_DETECT_TIMER_CHANGES
+		if (!ensure(Timer.File == nullptr || Timer.File == NewFile))
+		{
+			UE_LOG(LogTraceServices, Warning, TEXT("[TimingProfiler] Location changed for timer %u (%s): from \"%s\" (%u) to \"%s\" (%u) !!!"),
+				TimerId, Timer.Name, Timer.File, Timer.Line, NewFile, Line);
+		}
+#endif
+		Timer.File = NewFile;
+
 		Timer.Line = Line;
 	}
 }
@@ -188,6 +234,16 @@ void FTimingProfilerProvider::SetMetadataSpec(uint32 TimerId, uint32 MetadataSpe
 	if (ensure(TimerId < uint32(Timers.Num())))
 	{
 		FTimingProfilerTimer& Timer = Timers[TimerId];
+#if TIMING_PROFILER_DETECT_TIMER_CHANGES
+		if (!ensure(Timer.MetadataSpecId == FMetadataSpec::InvalidMetadataSpecId ||
+					Timer.MetadataSpecId == MetadataSpecId))
+		{
+			UE_LOG(LogTraceServices, Warning, TEXT("Metadata changed for timer %u (%s): from \"%s\" (%d fields) to \"%s\" (%d fields) !!!"),
+				TimerId, Timer.Name,
+				MetadataSpecs[Timer.MetadataSpecId].Format, MetadataSpecs[Timer.MetadataSpecId].FieldNames.Num(),
+				MetadataSpecs[MetadataSpecId].Format, MetadataSpecs[MetadataSpecId].FieldNames.Num());
+		}
+#endif
 		Timer.MetadataSpecId = MetadataSpecId;
 	}
 }
@@ -208,12 +264,14 @@ uint32 FTimingProfilerProvider::GetOriginalTimerIdFromMetadata(uint32 MetadataTi
 
 	if (int32(MetadataTimerId) >= 0)
 	{
+		TIMING_PROFILER_LOG_INVALID_METADATA_1;
 		return MetadataTimerId;
 	}
 
 	MetadataTimerId = ~MetadataTimerId;
 	if (MetadataTimerId >= uint32(Metadatas.Num()))
 	{
+		TIMING_PROFILER_LOG_INVALID_METADATA_2;
 		return 0;
 	}
 
@@ -227,12 +285,14 @@ TArrayView<const uint8> FTimingProfilerProvider::GetMetadata(uint32 MetadataTime
 
 	if (int32(MetadataTimerId) >= 0)
 	{
+		//TIMING_PROFILER_LOG_INVALID_METADATA_1; // expected in the normal use case
 		return TArrayView<const uint8>();
 	}
 
 	MetadataTimerId = ~MetadataTimerId;
 	if (MetadataTimerId >= uint32(Metadatas.Num()))
 	{
+		TIMING_PROFILER_LOG_INVALID_METADATA_2;
 		return TArrayView<const uint8>();
 	}
 
@@ -246,12 +306,14 @@ void FTimingProfilerProvider::SetMetadata(uint32 MetadataTimerId, TArray<uint8>&
 
 	if (int32(MetadataTimerId) >= 0)
 	{
+		TIMING_PROFILER_LOG_INVALID_METADATA_1;
 		return;
 	}
 
 	MetadataTimerId = ~MetadataTimerId;
 	if (MetadataTimerId >= uint32(Metadatas.Num()))
 	{
+		TIMING_PROFILER_LOG_INVALID_METADATA_2;
 		return;
 	}
 
@@ -265,12 +327,14 @@ void FTimingProfilerProvider::SetMetadata(uint32 MetadataTimerId, TArray<uint8>&
 
 	if (int32(MetadataTimerId) >= 0)
 	{
+		TIMING_PROFILER_LOG_INVALID_METADATA_1;
 		return;
 	}
 
 	MetadataTimerId = ~MetadataTimerId;
 	if (MetadataTimerId >= uint32(Metadatas.Num()))
 	{
+		TIMING_PROFILER_LOG_INVALID_METADATA_2;
 		return;
 	}
 
@@ -285,12 +349,14 @@ TArrayView<uint8> FTimingProfilerProvider::GetEditableMetadata(uint32 MetadataTi
 
 	if (int32(MetadataTimerId) >= 0)
 	{
+		TIMING_PROFILER_LOG_INVALID_METADATA_1;
 		return TArrayView<uint8>();
 	}
 
 	MetadataTimerId = ~MetadataTimerId;
 	if (MetadataTimerId >= uint32(Metadatas.Num()))
 	{
+		TIMING_PROFILER_LOG_INVALID_METADATA_2;
 		return TArrayView<uint8>();
 	}
 
@@ -302,8 +368,49 @@ uint32 FTimingProfilerProvider::AddMetadataSpec(FMetadataSpec&& Metadata)
 {
 	Session.WriteAccessCheck();
 
+	uint32 Hash = PointerHash(Metadata.Format);
+	for (const TCHAR* Field : Metadata.FieldNames)
+	{
+		Hash ^= PointerHash(Field);
+	}
+
+	TArray<uint32> FoundMetadataSpecIds;
+	MetadataSpecMap.MultiFind(Hash, FoundMetadataSpecIds);
+	if (FoundMetadataSpecIds.Num() > 0)
+	{
+		for (uint32 FoundMetadataSpecId : FoundMetadataSpecIds)
+		{
+			check(FoundMetadataSpecId < uint32(MetadataSpecs.Num()));
+			const FMetadataSpec& FoundMetadataSpec = MetadataSpecs[FoundMetadataSpecId];
+			if (Metadata.Format == FoundMetadataSpec.Format)
+				//|| FCString::Strcmp(Metadata.Format, FoundMetadataSpec.Format) == 0)
+			{
+				const int32 NumFields = Metadata.FieldNames.Num();
+				if (NumFields == FoundMetadataSpec.FieldNames.Num())
+				{
+					bool bSame = true;
+					for (int32 FieldIndex = 0; FieldIndex < NumFields; ++FieldIndex)
+					{
+						if (Metadata.FieldNames[FieldIndex] == FoundMetadataSpec.FieldNames[FieldIndex])
+							//|| FCString::Strcmp(Metadata.FieldNames[FieldIndex], FoundMetadataSpec.FieldNames[FieldIndex]) == 0)
+						{
+							continue;
+						}
+						bSame = false;
+						break;
+					}
+					if (bSame)
+					{
+						return FoundMetadataSpecId;
+					}
+				}
+			}
+		}
+	}
+
 	uint32 MetadataSpecId = MetadataSpecs.Num();
 	MetadataSpecs.Add(MoveTemp(Metadata));
+	MetadataSpecMap.Add(Hash, MetadataSpecId);
 
 	return MetadataSpecId;
 }
@@ -312,12 +419,13 @@ const FMetadataSpec* FTimingProfilerProvider::GetMetadataSpec(uint32 MetadataSpe
 {
 	Session.ReadAccessCheck();
 
-	if (MetadataSpecId < (uint32)MetadataSpecs.Num())
+	if (MetadataSpecId >= (uint32)MetadataSpecs.Num())
 	{
-		return &MetadataSpecs[MetadataSpecId];
+		TIMING_PROFILER_LOG_INVALID_METADATA_SPEC;
+		return nullptr;
 	}
 
-	return nullptr;
+	return &MetadataSpecs[MetadataSpecId];
 }
 
 IEditableTimeline<FTimingProfilerEvent>& FTimingProfilerProvider::GetCpuThreadEditableTimeline(uint32 ThreadId)
@@ -1248,3 +1356,8 @@ ITimingProfilerButterfly* FTimingProfilerProvider::CreateButterflyInternal(TArra
 }
 
 } // namespace TraceServices
+
+#undef TIMING_PROFILER_LOG_INVALID_METADATA_SPEC
+#undef TIMING_PROFILER_LOG_INVALID_METADATA_2
+#undef TIMING_PROFILER_LOG_INVALID_METADATA_1
+#undef TIMING_PROFILER_DETECT_TIMER_CHANGES
diff --git a/Engine/Source/Developer/TraceServices/Private/Model/TimingProfilerPrivate.h b/Engine/Source/Developer/TraceServices/Private/Model/TimingProfilerPrivate.h
--- a/Engine/Source/Developer/TraceServices/Private/Model/TimingProfilerPrivate.h
+++ b/Engine/Source/Developer/TraceServices/Private/Model/TimingProfilerPrivate.h
@@ -153,6 +153,7 @@ private:
 	IAnalysisSession& Session;
 	TArray<FMetadata> Metadatas;
 	TArray<FMetadataSpec> MetadataSpecs;
+	TMultiMap<uint32, uint32> MetadataSpecMap; // hash(FMetadataSpec) --> MetadataSpecId (index)
 	TArray<FTimingProfilerTimer> Timers;
 	TArray<TSharedRef<TimelineInternal>> Timelines;
 	TTableLayout<FTimingProfilerAggregatedStats> AggregatedStatsTableLayout;
diff --git a/Engine/Source/Developer/TraceServices/Private/Modules/TimingProfilerModule.cpp b/Engine/Source/Developer/TraceServices/Private/Modules/TimingProfilerModule.cpp
--- a/Engine/Source/Developer/TraceServices/Private/Modules/TimingProfilerModule.cpp
+++ b/Engine/Source/Developer/TraceServices/Private/Modules/TimingProfilerModule.cpp
@@ -5,6 +5,7 @@
 #include "AnalysisServicePrivate.h"
 #include "Analyzers/CpuProfilerTraceAnalysis.h"
 #include "Analyzers/GpuProfilerTraceAnalysis.h"
+#include "Analyzers/OldGpuProfilerTraceAnalysis.h"
 #include "Analyzers/VerseTraceAnalysis.h"
 #include "Model/CountersPrivate.h"
 #include "Model/ThreadsPrivate.h"
@@ -27,10 +28,20 @@ void FTimingProfilerModule::OnAnalysisBegin(IAnalysisSession& InSession)
 
 	IEditableThreadProvider& EditableThreadProvider = EditThreadProvider(Session);
 
+	// Add the TimingProfiler provider.
 	TSharedPtr<FTimingProfilerProvider> TimingProfilerProvider = MakeShared<FTimingProfilerProvider>(Session);
 	Session.AddProvider(GetTimingProfilerProviderName(), TimingProfilerProvider, TimingProfilerProvider);
+
+	// Add the CPU Profiler analyzer.
 	Session.AddAnalyzer(new FCpuProfilerAnalyzer(Session, *TimingProfilerProvider, EditableThreadProvider));
+
+	// Add the old GPU Profiler analyzer (for backward compatibility with old traces).
+	Session.AddAnalyzer(new FOldGpuProfilerAnalyzer(Session, *TimingProfilerProvider));
+
+	// Add the new GPU Profiler analyzer.
 	Session.AddAnalyzer(new FGpuProfilerAnalyzer(Session, *TimingProfilerProvider, EditCounterProvider(Session)));
+
+	// Add the new Verse Sampling analyzer.
 	Session.AddAnalyzer(new FVerseAnalyzer(Session, *TimingProfilerProvider));
 }
 
