[Insights] Modules: Fixed stats regarding number of symbols discovered / cached / resolved / failed.

#FailedValidations Preflight

diff --git a/Engine/Source/Developer/TraceInsights/Private/Insights/MemoryProfiler/Widgets/SMemAllocTableTreeView.cpp b/Engine/Source/Developer/TraceInsights/Private/Insights/MemoryProfiler/Widgets/SMemAllocTableTreeView.cpp
--- a/Engine/Source/Developer/TraceInsights/Private/Insights/MemoryProfiler/Widgets/SMemAllocTableTreeView.cpp
+++ b/Engine/Source/Developer/TraceInsights/Private/Insights/MemoryProfiler/Widgets/SMemAllocTableTreeView.cpp
@@ -883,15 +883,17 @@ FText SMemAllocTableTreeView::GetSymbolResolutionStatus() const
 		{
 			TraceServices::IModuleProvider::FStats Stats;
 			ModuleProvider->GetStats(&Stats);
-			//check(Stats.SymbolsDiscovered >= Stats.SymbolsResolved + Stats.SymbolsFailed);
-			const int32 SymbolsPending = Stats.SymbolsDiscovered - Stats.SymbolsResolved - Stats.SymbolsFailed;
-			if (SymbolsPending > 0)
+
+			const int32 SymbolsPending = Stats.SymbolsDiscovered - (Stats.SymbolsCached + Stats.SymbolsResolved + Stats.SymbolsFailed);
+			if (SymbolsPending != 0)
 			{
-				return FText::Format(LOCTEXT("SymbolsResolved1", "Resolving {0} / {1} symbols ({2} resolved, {3} failed)"), SymbolsPending, Stats.SymbolsDiscovered, Stats.SymbolsResolved, Stats.SymbolsFailed);
+				return FText::Format(LOCTEXT("SymbolsResolved1", "Resolving {0} / {1} symbols ({2} cached + {3} resolved + {4} failed)"),
+					SymbolsPending, Stats.SymbolsDiscovered, Stats.SymbolsCached, Stats.SymbolsResolved, Stats.SymbolsFailed);
 			}
 			else
 			{
-				return FText::Format(LOCTEXT("SymbolsResolved2", "{0} symbols ({1} resolved, {2} failed)"), Stats.SymbolsDiscovered, Stats.SymbolsResolved, Stats.SymbolsFailed);
+				return FText::Format(LOCTEXT("SymbolsResolved2", "{0} symbols ({1} cached + {2} resolved + {3} failed)"),
+					Stats.SymbolsDiscovered, Stats.SymbolsCached, Stats.SymbolsResolved, Stats.SymbolsFailed);
 			}
 		}
 	}
diff --git a/Engine/Source/Developer/TraceServices/Private/Model/DbgHelpResolver.cpp b/Engine/Source/Developer/TraceServices/Private/Model/DbgHelpResolver.cpp
--- a/Engine/Source/Developer/TraceServices/Private/Model/DbgHelpResolver.cpp
+++ b/Engine/Source/Developer/TraceServices/Private/Model/DbgHelpResolver.cpp
@@ -158,11 +158,13 @@ void FDbgHelpResolver::GetStats(IModuleProvider::FStats* OutStats) const
 	FMemory::Memzero(*OutStats);
 	for (uint32 ModuleIndex = 0; ModuleIndex < Modules.Num(); ++ModuleIndex)
 	{
-		FModule::SymbolStats& ModuleStats = Modules[ModuleIndex].Module->Stats;
+		const FModule::SymbolStats& ModuleStats = Modules[ModuleIndex].Module->Stats;
 		OutStats->SymbolsDiscovered += ModuleStats.Discovered.load();
+		OutStats->SymbolsCached += ModuleStats.Cached.load();
 		OutStats->SymbolsResolved += ModuleStats.Resolved.load();
 		OutStats->SymbolsFailed += ModuleStats.Failed.load();
 	}
+	OutStats->SymbolsFailed += NoModuleSymbols;
 	OutStats->ModulesDiscovered = ModulesDiscovered.load();
 	OutStats->ModulesFailed = ModulesFailed.load();
 	OutStats->ModulesLoaded = ModulesLoaded.load();
@@ -330,13 +332,14 @@ void FDbgHelpResolver::ResolveSymbol(uint64 Address, FResolvedSymbol& Target)
 {
 	if (Target.Result.load() == ESymbolQueryResult::OK)
 	{
+		++AlreadyResolvedSymbols;
 		return;
 	}
 
 	FModuleEntry* Entry = GetModuleForAddress(Address);
 	if (!Entry)
 	{
-		UE_LOG(LogDbgHelp, Warning, TEXT("No module mapped to address 0x%016llx."), Address);
+		UE_LOG(LogDbgHelp, Warning, TEXT("No module mapped to address 0x%llX."), Address);
 		UpdateResolvedSymbol(Target,
 			ESymbolQueryResult::NotFound,
 			GUnknownModuleTextDbgHelp,
@@ -344,6 +347,7 @@ void FDbgHelpResolver::ResolveSymbol(uint64 Address, FResolvedSymbol& Target)
 			GUnknownModuleTextDbgHelp,
 			0);
 		SymbolFilter.Update(Target);
+		++NoModuleSymbols;
 		return;
 	}
 	FModule* Module = Entry->Module;
diff --git a/Engine/Source/Developer/TraceServices/Private/Model/DbgHelpResolver.h b/Engine/Source/Developer/TraceServices/Private/Model/DbgHelpResolver.h
--- a/Engine/Source/Developer/TraceServices/Private/Model/DbgHelpResolver.h
+++ b/Engine/Source/Developer/TraceServices/Private/Model/DbgHelpResolver.h
@@ -84,6 +84,8 @@ private:
 	std::atomic<uint32> ModulesDiscovered;
 	std::atomic<uint32> ModulesFailed;
 	std::atomic<uint32> ModulesLoaded;
+	std::atomic<uint32> AlreadyResolvedSymbols;
+	std::atomic<uint32> NoModuleSymbols;
 
 	IAnalysisSession& Session;
 	IResolvedSymbolFilter& SymbolFilter;
diff --git a/Engine/Source/Developer/TraceServices/Private/Model/ModuleProvider.cpp b/Engine/Source/Developer/TraceServices/Private/Model/ModuleProvider.cpp
--- a/Engine/Source/Developer/TraceServices/Private/Model/ModuleProvider.cpp
+++ b/Engine/Source/Developer/TraceServices/Private/Model/ModuleProvider.cpp
@@ -78,6 +78,7 @@ public:
 	void OnModuleLoad(const FStringView& Module, uint64 Base, uint32 Size, const uint8* Checksum, uint32 ChecksumSize) override;
 	void OnModuleUnload(uint64 Base) override;
 	void OnAnalysisComplete() override;
+
 	void SaveSymbolsToCache(IAnalysisCache& Cache);
 	void LoadSymbolsFromCache(IAnalysisCache& Cache);
 
@@ -174,6 +175,7 @@ const FResolvedSymbol* TModuleProvider<SymbolResolverType>::GetSymbol(uint64 Add
 		// Add a pending entry to our cache.
 		ResolvedSymbol = &SymbolCache.EmplaceBack(ESymbolQueryResult::Pending, nullptr, nullptr, nullptr, (uint16)0, EResolvedSymbolFilterStatus::Unknown);
 		SymbolCacheLookup.Add(Address, ResolvedSymbol);
+
 		++SymbolsDiscovered;
 	}
 	check(ResolvedSymbol);
@@ -263,7 +265,7 @@ FGraphEventRef TModuleProvider<SymbolResolverType>::LoadSymbolsForModuleUsingPat
 					return InModule->Status.load();
 				};
 
-				UE_LOG(LogTraceServices, Display, TEXT("Queuing symbol loading using path %s."), *FullPath);
+				UE_LOG(LogTraceServices, Display, TEXT("[Modules] Queuing symbol loading using path %s."), *FullPath);
 
 				// Load the requested module
 				const EModuleStatus Result = ReloadModuleFn(Module, *FullPath);
@@ -300,7 +302,7 @@ FGraphEventRef TModuleProvider<SymbolResolverType>::LoadSymbolsForModuleUsingPat
 					}
 				}
 
-				UE_LOG(LogTraceServices, Display, TEXT("Completed loading symbols for path %s."), *FullPath);
+				UE_LOG(LogTraceServices, Display, TEXT("[Modules] Completed loading symbols for path %s."), *FullPath);
 			});
 
 			return LoadSymbolsTask;
@@ -322,14 +324,19 @@ void TModuleProvider<SymbolResolverType>::EnumerateSymbolSearchPaths(TFunctionRe
 template<typename SymbolResolverType>
 void TModuleProvider<SymbolResolverType>::GetStats(FStats* OutStats) const
 {
+	// Get stats from resolver.
 	Resolver->GetStats(OutStats);
-	OutStats->SymbolsDiscovered = SymbolsDiscovered.load();
-	// Add the cached symbols (the resolver doesn't know about them)
-	OutStats->SymbolsResolved += NumCachedSymbols;
+
+	// The resolver does not include Cached into Discovered stat (but only Resolved and Failed).
+	// Also, the resolver's Discovered and Cached stats (computed as sum of all modules) are not accurate
+	// due to overlaps and unloaded/reloaded modules.
+	// So we override these with the provider's stats.
+	OutStats->SymbolsDiscovered = SymbolsDiscovered.load() + NumCachedSymbols;
+	OutStats->SymbolsCached = NumCachedSymbols;
 }
 
 ////////////////////////////////////////////////////////////////////////////////////////////////////
-	
+
 template<typename SymbolResolverType>
 bool TModuleProvider<SymbolResolverType>::HasFinishedResolving() const
 {
@@ -343,6 +350,7 @@ void TModuleProvider<SymbolResolverType>::OnModuleLoad(const FStringView& Module
 {
 	if (ModuleName.Len() == 0)
 	{
+		UE_LOG(LogTraceServices, Warning, TEXT("[Modules] No name module is ignored (Base=0x%llX, Size=%u) !!!"), Base, Size);
 		return;
 	}
 
@@ -352,20 +360,43 @@ void TModuleProvider<SymbolResolverType>::OnModuleLoad(const FStringView& Module
 
 	FWriteScopeLock _(ModulesLock);
 
-	// Check if the module was already added.
 	for (uint32 ModuleIndex = 0; ModuleIndex < Modules.Num(); ++ModuleIndex)
 	{
-		const FModule& LoadedModule = Modules[ModuleIndex];
+		FModule& LoadedModule = Modules[ModuleIndex];
+
+		// Check if the module was already added.
 		if (LoadedModule.Name == Name && // only comparing pointers from session's string store
 			LoadedModule.Base == Base &&
 			LoadedModule.Size == Size)
 		{
+			if (LoadedModule.bIsUnloaded)
+			{
+				UE_LOG(LogTraceServices, Log, TEXT("[Modules] Re-loaded module %s (Base=0x%llX, Size=%u)."), LoadedModule.Name, LoadedModule.Base, LoadedModule.Size);
+				LoadedModule.bIsUnloaded = false;
+			}
+			else
+			{
+				// Duplicated module !!!
+				UE_LOG(LogTraceServices, Warning, TEXT("[Modules] Duplicated module %s (Base=0x%llX, Size=%u) !!!"), LoadedModule.Name, LoadedModule.Base, LoadedModule.Size);
+			}
 			return;
 		}
+
+		// Check if the new module overlaps an existing (active) module.
+		if (!LoadedModule.bIsUnloaded &&
+			Base + Size > LoadedModule.Base &&
+			Base < LoadedModule.Base + LoadedModule.Size)
+		{
+			// Overlapping modules !!!
+			UE_LOG(LogTraceServices, Warning, TEXT("[Modules] Module %s (Base=0x%llX, Size=%u) overlaps with module %s (Base=0x%llX, Size=%u) !!!"),
+				Name, Base, Size, LoadedModule.Name, LoadedModule.Base, LoadedModule.Size);
+		}
 	}
 
 	FModule& NewModule = Modules.EmplaceBack(Name, FullName, Base, Size, EModuleStatus::Discovered);
 
+	UE_LOG(LogTraceServices, Log, TEXT("[Modules] Loaded module %s (Base=0x%llX, Size=%u)."), NewModule.Name, NewModule.Base, NewModule.Size);
+
 	// The number of cached symbols for this module is equal to the number
 	// of matching addresses in the cache.
 	NewModule.Stats.Cached = GetNumCachedSymbolsFromModule(Base, Size);
@@ -378,7 +409,23 @@ void TModuleProvider<SymbolResolverType>::OnModuleLoad(const FStringView& Module
 template<typename SymbolResolverType>
 void TModuleProvider<SymbolResolverType>::OnModuleUnload(uint64 Base)
 {
-	//todo: Find entry, set bLoaded to false
+	FWriteScopeLock _(ModulesLock);
+
+	// Check if the module was already added.
+	for (uint32 ModuleIndex = 0; ModuleIndex < Modules.Num(); ++ModuleIndex)
+	{
+		FModule& LoadedModule = Modules[ModuleIndex];
+
+		if (LoadedModule.Base == Base && !LoadedModule.bIsUnloaded)
+		{
+			LoadedModule.bIsUnloaded = true;
+			UE_LOG(LogTraceServices, Log, TEXT("[Modules] Unloaded module %s (Base=0x%llX, Size=%u)."), LoadedModule.Name, LoadedModule.Base, LoadedModule.Size);
+			return;
+		}
+	}
+
+	// Unknown module !!!
+	UE_LOG(LogTraceServices, Warning, TEXT("[Modules] No module to unload with Base=0x%llX !!!"), Base);
 }
 
 ////////////////////////////////////////////////////////////////////////////////////////////////////
@@ -420,7 +467,7 @@ void TModuleProvider<SymbolResolverType>::SaveSymbolsToCache(IAnalysisCache& Cac
 		SavedSymbols.EmplaceBack(FSavedSymbol{*Address, ModuleOffset, NameOffset, FileOffset, Symbol.Line});
 		++NumSavedSymbols;
 	}
-	UE_LOG(LogTraceServices, Display, TEXT("Added %d symbols to the %d previously saved symbols."), NumSavedSymbols, NumPreviouslySavedSymbols);
+	UE_LOG(LogTraceServices, Display, TEXT("[Modules] Added %d symbols to the %d previously saved symbols."), NumSavedSymbols, NumPreviouslySavedSymbols);
 }
 
 ////////////////////////////////////////////////////////////////////////////////////////////////////
@@ -438,7 +485,7 @@ void TModuleProvider<SymbolResolverType>::LoadSymbolsFromCache(IAnalysisCache& C
 		const TCHAR* File = Strings.GetStringAtOffset(Symbol.FileOffset);
 		if (Module == nullptr || Name == nullptr || File == nullptr)
 		{
-			UE_LOG(LogTraceServices, Warning, TEXT("Found cached symbol (address %llx) which referenced unknown string."), Symbol.Address);
+			UE_LOG(LogTraceServices, Warning, TEXT("[Modules] Found cached symbol (address 0x%llX) which referenced unknown string !!!"), Symbol.Address);
 			continue;
 		}
 		FResolvedSymbol& Resolved = SymbolCache.EmplaceBack(ESymbolQueryResult::OK, Module, Name, File, static_cast<uint16>(Symbol.Line), EResolvedSymbolFilterStatus::Unknown);
@@ -453,7 +500,7 @@ void TModuleProvider<SymbolResolverType>::LoadSymbolsFromCache(IAnalysisCache& C
 		},
 		EParallelForFlags::BackgroundPriority);
 
-	UE_LOG(LogTraceServices, Display, TEXT("Loaded %d symbols from cache."), SymbolCacheLookup.Num());
+	UE_LOG(LogTraceServices, Display, TEXT("[Modules] Loaded %d symbols from cache."), SymbolCacheLookup.Num());
 }
 
 ////////////////////////////////////////////////////////////////////////////////////////////////////
diff --git a/Engine/Source/Developer/TraceServices/Private/Model/PsymResolver.cpp b/Engine/Source/Developer/TraceServices/Private/Model/PsymResolver.cpp
--- a/Engine/Source/Developer/TraceServices/Private/Model/PsymResolver.cpp
+++ b/Engine/Source/Developer/TraceServices/Private/Model/PsymResolver.cpp
@@ -160,10 +160,13 @@ void FPsymResolver::GetStats(IModuleProvider::FStats* OutStats) const
 	FMemory::Memzero(*OutStats);
 	for(const FModuleEntry& Entry : LoadedModules)
 	{
-		OutStats->SymbolsDiscovered += Entry.Module->Stats.Discovered.load();
-		OutStats->SymbolsResolved += Entry.Module->Stats.Resolved.load();
-		OutStats->SymbolsFailed += Entry.Module->Stats.Failed.load();
+		const FModule::SymbolStats& ModuleStats = Entry.Module->Stats;
+		OutStats->SymbolsDiscovered += ModuleStats.Discovered.load();
+		OutStats->SymbolsCached += ModuleStats.Cached.load();
+		OutStats->SymbolsResolved += ModuleStats.Resolved.load();
+		OutStats->SymbolsFailed += ModuleStats.Failed.load();
 	}
+	OutStats->SymbolsFailed += NoModuleSymbols;
 	OutStats->ModulesDiscovered = ModulesDiscovered.load();
 	OutStats->ModulesFailed = ModulesFailed.load();
 	OutStats->ModulesLoaded = ModulesLoaded.load();
@@ -471,13 +474,14 @@ void FPsymResolver::ResolveSymbol(uint64 Address, FResolvedSymbol& Target)
 	const ESymbolQueryResult PreviousResult = Target.Result.load();
 	if (PreviousResult == ESymbolQueryResult::OK)
 	{
+		++AlreadyResolvedSymbols;
 		return;
 	}
 
 	const FModuleEntry* Entry = GetModuleForAddress(Address);
 	if (!Entry)
 	{
-		UE_LOG(LogPsymResolver, Warning, TEXT("No module mapped to address 0x%016llx."), Address);
+		UE_LOG(LogPsymResolver, Warning, TEXT("No module mapped to address 0x%llX."), Address);
 		UpdateResolvedSymbol(Target,
 			ESymbolQueryResult::NotLoaded,
 			GUnknownModuleTextPsym,
@@ -485,6 +489,7 @@ void FPsymResolver::ResolveSymbol(uint64 Address, FResolvedSymbol& Target)
 			GUnknownModuleTextPsym,
 			0);
 		SymbolFilter.Update(Target);
+		++NoModuleSymbols;
 		return;
 	}
 
diff --git a/Engine/Source/Developer/TraceServices/Private/Model/PsymResolver.h b/Engine/Source/Developer/TraceServices/Private/Model/PsymResolver.h
--- a/Engine/Source/Developer/TraceServices/Private/Model/PsymResolver.h
+++ b/Engine/Source/Developer/TraceServices/Private/Model/PsymResolver.h
@@ -80,6 +80,8 @@ private:
 	std::atomic<uint32> ModulesDiscovered;
 	std::atomic<uint32> ModulesFailed;
 	std::atomic<uint32> ModulesLoaded;
+	std::atomic<uint32> AlreadyResolvedSymbols;
+	std::atomic<uint32> NoModuleSymbols;
 
 	mutable FRWLock CustomSymbolSearchPathsLock;
 	TArray<FString> CustomSymbolSearchPaths; // search paths added by user from UI in the current session
diff --git a/Engine/Source/Developer/TraceServices/Private/Model/SymslibResolver.cpp b/Engine/Source/Developer/TraceServices/Private/Model/SymslibResolver.cpp
--- a/Engine/Source/Developer/TraceServices/Private/Model/SymslibResolver.cpp
+++ b/Engine/Source/Developer/TraceServices/Private/Model/SymslibResolver.cpp
@@ -717,11 +717,13 @@ void FSymslibResolver::GetStats(IModuleProvider::FStats* OutStats) const
 	FMemory::Memzero(*OutStats);
 	for (uint32 ModuleIndex = 0; ModuleIndex < Modules.Num(); ++ModuleIndex)
 	{
-		FModule::SymbolStats& ModuleStats = Modules[ModuleIndex].Module->Stats;
+		const FModule::SymbolStats& ModuleStats = Modules[ModuleIndex].Module->Stats;
 		OutStats->SymbolsDiscovered += ModuleStats.Discovered.load();
+		OutStats->SymbolsCached += ModuleStats.Cached.load();
 		OutStats->SymbolsResolved += ModuleStats.Resolved.load();
 		OutStats->SymbolsFailed += ModuleStats.Failed.load();
 	}
+	OutStats->SymbolsFailed += NoModuleSymbols;
 	OutStats->ModulesDiscovered = ModulesDiscovered.load();
 	OutStats->ModulesFailed = ModulesFailed.load();
 	OutStats->ModulesLoaded = ModulesLoaded.load();
@@ -1125,13 +1127,14 @@ void FSymslibResolver::ResolveSymbolTracked(uint64 Address, FResolvedSymbol& Tar
 	const ESymbolQueryResult PreviousResult = Target.Result.load();
 	if (PreviousResult == ESymbolQueryResult::OK)
 	{
+		++AlreadyResolvedSymbols;
 		return;
 	}
 
 	FModuleEntry* Entry = GetModuleForAddress(Address);
 	if (!Entry)
 	{
-		UE_LOG(LogSymslib, Warning, TEXT("No module mapped to address 0x%016llx."), Address);
+		UE_LOG(LogSymslib, Warning, TEXT("No module mapped to address 0x%llX."), Address);
 		UpdateResolvedSymbol(Target,
 			ESymbolQueryResult::NotLoaded,
 			GUnknownModuleTextSymsLib,
@@ -1139,6 +1142,7 @@ void FSymslibResolver::ResolveSymbolTracked(uint64 Address, FResolvedSymbol& Tar
 			GUnknownModuleTextSymsLib,
 			0);
 		SymbolFilter.Update(Target);
+		++NoModuleSymbols;
 		return;
 	}
 
diff --git a/Engine/Source/Developer/TraceServices/Private/Model/SymslibResolver.h b/Engine/Source/Developer/TraceServices/Private/Model/SymslibResolver.h
--- a/Engine/Source/Developer/TraceServices/Private/Model/SymslibResolver.h
+++ b/Engine/Source/Developer/TraceServices/Private/Model/SymslibResolver.h
@@ -130,6 +130,8 @@ private:
 	std::atomic<uint32> ModulesDiscovered;
 	std::atomic<uint32> ModulesFailed;
 	std::atomic<uint32> ModulesLoaded;
+	std::atomic<uint32> AlreadyResolvedSymbols;
+	std::atomic<uint32> NoModuleSymbols;
 	std::atomic<uint64> SymbolBytesAllocated;
 	std::atomic<uint64> SymbolBytesWasted;
 
diff --git a/Engine/Source/Developer/TraceServices/Public/TraceServices/Model/Modules.h b/Engine/Source/Developer/TraceServices/Public/TraceServices/Model/Modules.h
--- a/Engine/Source/Developer/TraceServices/Public/TraceServices/Model/Modules.h
+++ b/Engine/Source/Developer/TraceServices/Public/TraceServices/Model/Modules.h
@@ -16,9 +16,9 @@ namespace TraceServices
 
 ////////////////////////////////////////////////////////////////////////////////
 /**
-  * Result of a query. Since symbol resolving can be deferred this signals if a
-  * symbol has been resolved, waiting to be resolved or wasn't found at all.
-  */
+ * Result of a query. Since symbol resolving can be deferred this signals if a
+ * symbol has been resolved, waiting to be resolved or wasn't found at all.
+ */
 enum class ESymbolQueryResult : uint8
 {
 	Pending,	// Symbol is pending resolution
@@ -57,10 +57,10 @@ enum class EResolvedSymbolFilterStatus : uint8
 
 ////////////////////////////////////////////////////////////////////////////////
 /**
-  * Represent a resolved symbol. The resolve status and string values may change
-  * over time, but string pointers returned from the methods are guaranteed to live
-  * during the entire analysis session.
-  */
+ * Represent a resolved symbol. The resolve status and string values may change
+ * over time, but string pointers returned from the methods are guaranteed to live
+ * during the entire analysis session.
+ */
 struct FResolvedSymbol
 {
 	const TCHAR* Module;
@@ -93,7 +93,7 @@ public:
 };
 
 ////////////////////////////////////////////////////////////////////////////////
-enum class EModuleStatus
+enum class EModuleStatus : uint8
 {
 	Discovered,			// Symbols are discovered for this module
 	Pending,			// Module is pending load
@@ -135,6 +135,8 @@ struct FModule
 	uint64						Base;
 	// Size in memory
 	uint32						Size;
+	// If the module was unloaded or not
+	std::atomic<bool>			bIsUnloaded;
 	// Status of loading debug information
 	std::atomic<EModuleStatus>	Status;
 	// If status is loaded, contains the path to the debug info. If failure state contains an error message.
@@ -154,6 +156,7 @@ struct FModule
 		, FullName(InFullName)
 		, Base(InBase)
 		, Size(InSize)
+		, bIsUnloaded(false)
 		, Status(InStatus)
 		, StatusMessage(nullptr)
 	{
@@ -172,6 +175,7 @@ public:
 		uint32 ModulesLoaded;
 		uint32 ModulesFailed;
 		uint32 SymbolsDiscovered;
+		uint32 SymbolsCached;
 		uint32 SymbolsResolved;
 		uint32 SymbolsFailed;
 	};
