[Insights] Minor refactoring for monotonic timelines.

#rnx

diff --git a/Engine/Source/Developer/TraceServices/Public/Model/AsyncEnumerateTask.h b/Engine/Source/Developer/TraceServices/Public/Model/AsyncEnumerateTask.h
--- a/Engine/Source/Developer/TraceServices/Public/Model/AsyncEnumerateTask.h
+++ b/Engine/Source/Developer/TraceServices/Public/Model/AsyncEnumerateTask.h
@@ -60,7 +60,8 @@ public:
 		}
 
 		// Check if we must also send the initial stack events to the caller.
-		if (Data.SortOrder == EEventSortOrder::ByEndTime || (Data.TaskIndex == 0 && Data.SortOrder == EEventSortOrder::ByStartTime))
+		if ((Data.SortOrder == EEventSortOrder::ByEndTime) ||
+			(Data.TaskIndex == 0 && Data.SortOrder == EEventSortOrder::ByStartTime))
 		{
 			InitCallback(0);
 
@@ -85,7 +86,7 @@ public:
 		// Iterate the pages owned by this task.
 		while (ScopeEntry && CurrentPageIndex < Data.EndPageIndex)
 		{
-			if (ScopeEntry->Time < 0.0)
+			if (ScopeEntry->Time < 0.0) // "begin" event
 			{
 				check(CurrentStackDepth < SettingsType::MaxDepth);
 				FEnumerationStackEntry& StackEntry = EventStack[CurrentStackDepth++];
@@ -99,7 +100,7 @@ public:
 
 				Event = EventsIterator.NextItem();
 			}
-			else
+			else // "end" event
 			{
 				check(CurrentStackDepth > 0);
 				--CurrentStackDepth;
@@ -151,7 +152,7 @@ public:
 			} while (ScopeEntryIterator.NextPage());
 		}
 
-		// If there are still events on the stack at this point, it means analysis has not yet reached their end, 
+		// If there are still events on the stack at this point, it means analysis has not yet reached their end,
 		// so we truncate them to the time of last received event.
 		while (CurrentStackDepth > 0)
 		{
@@ -171,7 +172,10 @@ public:
 private:
 	static bool ShouldStopSearchingForEndEvents(FAsyncEnumerateTaskData<InEventType, SettingsType>& Data, FEnumerationStackEntry* EventStack, const FEventScopeEntryPage* InitialPage, int32 CurrentStackDepth, const FDetailLevel<InEventType, SettingsType>& DetailLevel)
 	{
-		if (Data.SortOrder == EEventSortOrder::ByStartTime && Data.TaskIndex > 0 && CurrentStackDepth > 0 && CurrentStackDepth <= InitialPage->InitialStackCount)
+		if (Data.SortOrder == EEventSortOrder::ByStartTime &&
+			Data.TaskIndex > 0 &&
+			CurrentStackDepth > 0 &&
+			CurrentStackDepth <= InitialPage->InitialStackCount)
 		{
 			double InitialStartTime = DetailLevel.GetScopeEntryTime(InitialPage->InitialStack[CurrentStackDepth - 1].EnterScopeIndex);
 
@@ -217,15 +221,18 @@ public:
 			uint32 Depth;
 			EventType Event;
 		};
+
 		TArray<FOutputEvent> OutputEvents;
 
-		EnumerateAsyncAlgoritm<InEventType, SettingsType>::EnumerateEventsDownSampled(Data, 
-			[&CurrentDepth, &MinimumDepth](int32 InInitialDepth)
+		EnumerateAsyncAlgoritm<InEventType, SettingsType>::EnumerateEventsDownSampled(Data,
+			[&CurrentDepth, &MinimumDepth]
+			(int32 InInitialDepth)
 			{
 				MinimumDepth = InInitialDepth;
 				CurrentDepth = InInitialDepth;
 			},
-			[this, &EventStack, &OutputEvents, &CurrentDepth, &MinimumDepth](bool IsEnter, double Time, const EventType& Event)
+			[this, &EventStack, &OutputEvents, &CurrentDepth, &MinimumDepth]
+			(bool IsEnter, double Time, const EventType& Event)
 			{
 				if (IsEnter)
 				{
@@ -255,7 +262,7 @@ public:
 						{
 							if (OutputEvent.EndTime >= Data.StartTime && OutputEvent.StartTime <= Data.EndTime)
 							{
-								if (this->Data.EventRangeCallback(OutputEvent.StartTime, OutputEvent.EndTime, OutputEvent.Depth, OutputEvent.Event, Data.TaskIndex) == EEventEnumerate::Stop)
+								if (Data.EventRangeCallback(OutputEvent.StartTime, OutputEvent.EndTime, OutputEvent.Depth, OutputEvent.Event, Data.TaskIndex) == EEventEnumerate::Stop)
 								{
 									return EEventEnumerate::Stop;
 								}
@@ -273,20 +280,22 @@ public:
 
 	void EnumerateOrderedByEndTime()
 	{
-		uint32 CurrentDepth = 0;
 		struct FStackEntry
 		{
 			double StartTime;
 		};
 
 		FStackEntry EventStack[SettingsType::MaxDepth];
+		uint32 CurrentDepth = 0;
 
 		EnumerateAsyncAlgoritm<InEventType, SettingsType>::EnumerateEventsDownSampled(Data,
-			[&CurrentDepth](int32 InInitialDepth)
+			[&CurrentDepth]
+			(int32 InInitialDepth)
 			{
 				CurrentDepth = InInitialDepth;
 			},
-			[this, &EventStack, &CurrentDepth](bool IsEnter, double Time, const EventType& Event)
+			[this, &EventStack, &CurrentDepth]
+			(bool IsEnter, double Time, const EventType& Event)
 			{
 				if (IsEnter)
 				{
@@ -300,7 +309,7 @@ public:
 
 					if (Time >= Data.StartTime && StackEntry.StartTime <= Data.EndTime)
 					{
-						if (this->Data.EventRangeCallback(StackEntry.StartTime, Time, CurrentDepth, Event, Data.TaskIndex) == EEventEnumerate::Stop)
+						if (Data.EventRangeCallback(StackEntry.StartTime, Time, CurrentDepth, Event, Data.TaskIndex) == EEventEnumerate::Stop)
 						{
 							return EEventEnumerate::Stop;
 						}
diff --git a/Engine/Source/Developer/TraceServices/Public/Model/MonotonicTimeline.h b/Engine/Source/Developer/TraceServices/Public/Model/MonotonicTimeline.h
--- a/Engine/Source/Developer/TraceServices/Public/Model/MonotonicTimeline.h
+++ b/Engine/Source/Developer/TraceServices/Public/Model/MonotonicTimeline.h
@@ -861,7 +861,7 @@ public:
 
 			if (EventInfo.EndTime < 0.0)
 			{
-				//The end of the event has not been reached by analysis
+				// The end of the event has not been reached by analysis.
 				EventInfo.EndTime = DetailLevel.InsertionState.LastTime;
 			}
 
@@ -1035,8 +1035,7 @@ public:
 			ScopeEntry = ScopeEntryIterator.GetCurrentItem();
 		}
 
-		while (ScopeEntry &&
-			   Depth < CurrentStackDepth)
+		while (ScopeEntry && Depth < CurrentStackDepth)
 		{
 			if (ScopeEntry->Time < 0.0) // "begin" event
 			{
@@ -1097,11 +1096,12 @@ public:
 			--FirstScopePageIndex;
 		}
 
-		const FEventScopeEntryPage* ScopePage = DetailLevel.ScopeEntries.GetPage(FirstScopePageIndex);
-		if (ScopePage->BeginTime > Time)
+		else
 		{
 			return 0;
 		}
+		const FEventScopeEntryPage* ScopePage = DetailLevel.ScopeEntries.GetPage(FirstScopePageIndex);
+		check(ScopePage->BeginTime <= Time);
 
 		int32 CurrentStackDepth = ScopePage->InitialStackCount;
 
@@ -1190,20 +1190,20 @@ private:
 			// At the very first call, CurrentScopeEntryPage will be -1
 			if (DetailLevel.InsertionState.CurrentScopeEntryPageIndex != (uint64) -1)
 			{
-				FEventScopeEntryPage* CurrentScopeEntryPage = DetailLevel.ScopeEntries.GetPage(DetailLevel.InsertionState.CurrentScopeEntryPageIndex);
-				int32 PreviousPageInitialStackCount = CurrentScopeEntryPage->InitialStackCount;
 				check(DetailLevel.InsertionState.CurrentDepth <= SettingsType::MaxDepth);
-				int32 CurrentDepth = DetailLevel.InsertionState.CurrentDepth;
+
+				FEventScopeEntryPage* CurrentScopeEntryPage = DetailLevel.ScopeEntries.GetPage(DetailLevel.InsertionState.CurrentScopeEntryPageIndex);
+				const int32 MaxEntryIndex = FMath::Min(CurrentScopeEntryPage->InitialStackCount, DetailLevel.InsertionState.CurrentDepth);
+
 				// Update the open scopes that were also open at the beginning of the last page so the values
 				// represent stats up to and including the current page
-				int ii = 0;
-				for (; ii < PreviousPageInitialStackCount && ii < CurrentDepth; ++ii)
+				for (int32 EntryIndex = 0; EntryIndex < MaxEntryIndex; ++EntryIndex)
 				{
-					FEventStackEntry& InsertionStateStackEntry = DetailLevel.InsertionState.EventStack[ii];
-					FEventStackEntry& PreviousPageStackEntry = CurrentScopeEntryPage->InitialStack[ii];
+					FEventStackEntry& InsertionStateStackEntry = DetailLevel.InsertionState.EventStack[EntryIndex];
+					FEventStackEntry& PreviousPageStackEntry = CurrentScopeEntryPage->InitialStack[EntryIndex];
 
-					if (InsertionStateStackEntry.EnterScopeIndex == PreviousPageStackEntry.EnterScopeIndex
-						&& InsertionStateStackEntry.EventIndex == PreviousPageStackEntry.EventIndex)
+					if (InsertionStateStackEntry.EnterScopeIndex == PreviousPageStackEntry.EnterScopeIndex &&
+						InsertionStateStackEntry.EventIndex == PreviousPageStackEntry.EventIndex)
 					{
 						PreviousPageStackEntry.ExclTime = InsertionStateStackEntry.ExclTime;
 					}
@@ -1258,8 +1258,8 @@ private:
 			{
 				FEventStackEntry& PreviousPageStackEntry = CurrentScopeEntryPage->InitialStack[DetailLevel.InsertionState.CurrentDepth];
 
-				if (StackEntry.EnterScopeIndex == PreviousPageStackEntry.EnterScopeIndex
-					&& StackEntry.EventIndex == PreviousPageStackEntry.EventIndex)
+				if (StackEntry.EnterScopeIndex == PreviousPageStackEntry.EnterScopeIndex &&
+					StackEntry.EventIndex == PreviousPageStackEntry.EventIndex)
 				{
 					PreviousPageStackEntry.ExclTime = StackEntry.ExclTime;
 					PreviousPageStackEntry.EndTime = Time;
