[Insights]
* Logs: Fixed analysis to support log messages received before their spec. The message spec is now updated late for these messages.
* Logs: Added end of analysis message with stats (number of messages, specs and categories) and warnings (number of inserts and number of messages received with unknown spec).
* Logs: Minimized the session locking for log trace analysis.

diff --git a/Engine/Source/Developer/TraceServices/Private/Analyzers/BookmarksTraceAnalysis.cpp b/Engine/Source/Developer/TraceServices/Private/Analyzers/BookmarksTraceAnalysis.cpp
--- a/Engine/Source/Developer/TraceServices/Private/Analyzers/BookmarksTraceAnalysis.cpp
+++ b/Engine/Source/Developer/TraceServices/Private/Analyzers/BookmarksTraceAnalysis.cpp
@@ -68,7 +68,7 @@ bool FBookmarksAnalyzer::OnEvent(uint16 RouteId, EStyle Style, const FOnEventCon
 		{
 			const ANSICHAR* File = reinterpret_cast<const ANSICHAR*>(EventData.GetAttachment());
 			BookmarkFile = Session.StoreString(ANSI_TO_TCHAR(File));
-			BookmarkFormatString = Session.StoreString(reinterpret_cast<const TCHAR*>(EventData.GetAttachment() + strlen(File) + 1));
+			BookmarkFormatString = Session.StoreString(reinterpret_cast<const TCHAR*>(EventData.GetAttachment() + TCString<ANSICHAR>::Strlen(File) + 1));
 		}
 
 		{
@@ -104,7 +104,7 @@ bool FBookmarksAnalyzer::OnEvent(uint16 RouteId, EStyle Style, const FOnEventCon
 		{
 			// Todo: update this to use provider locking instead of session locking
 			// FProviderEditScopeLock EditableLogProviderLock (*EditableLogProvider);
-			EditableLogProvider->AppendMessage(BookmarkPoint, Timestamp, FormatArgsView.GetData());
+			EditableLogProvider->AppendMessage(BookmarkPoint, Timestamp, FormatArgsView);
 		}
 		break;
 	}
diff --git a/Engine/Source/Developer/TraceServices/Private/Analyzers/LogTraceAnalysis.cpp b/Engine/Source/Developer/TraceServices/Private/Analyzers/LogTraceAnalysis.cpp
--- a/Engine/Source/Developer/TraceServices/Private/Analyzers/LogTraceAnalysis.cpp
+++ b/Engine/Source/Developer/TraceServices/Private/Analyzers/LogTraceAnalysis.cpp
@@ -27,60 +27,99 @@ void FLogTraceAnalyzer::OnAnalysisBegin(const FOnAnalysisContext& Context)
 	Builder.RouteEvent(RouteId_LogMessage, "Logging", "LogMessage");
 }
 
+void FLogTraceAnalyzer::OnAnalysisEnd()
+{
+	FAnalysisSessionEditScope _(Session);
+
+	const uint64 InsertCount = LogProvider.GetInsertCount();
+	if (InsertCount > 0)
+	{
+		UE_LOG(LogTraceServices, Warning, TEXT("[Logs] Number of inserts (out of order messages): %llu"), InsertCount);
+	}
+
+	const uint64 UnknownSpecTotalCount = LogProvider.GetUnknownSpecTotalCount();
+	if (UnknownSpecTotalCount > 0)
+	{
+		UE_LOG(LogTraceServices, Warning, TEXT("[Logs] Number of log messages received with unknown spec (LogPoint): %llu (%llu still unresolved)"),
+			UnknownSpecTotalCount, LogProvider.GetUnresolvedSpecCount());
+	}
+
+	UE_LOG(LogTraceServices, Log, TEXT("[Logs] Analysis completed (%llu messages, %llu specs, %llu categories)."),
+		LogProvider.GetMessageCount(),
+		LogProvider.GetMessageSpecCount(),
+		LogProvider.GetCategoryCount());
+}
+
 bool FLogTraceAnalyzer::OnEvent(uint16 RouteId, EStyle Style, const FOnEventContext& Context)
 {
 	LLM_SCOPE_BYNAME(TEXT("Insights/FLogTraceAnalyzer"));
 
-	FAnalysisSessionEditScope _(Session);
-
 	const auto& EventData = Context.EventData;
 	switch (RouteId)
 	{
 	case RouteId_LogCategory:
 	{
 		uint64 CategoryPointer = EventData.GetValue<uint64>("CategoryPointer");
-		FLogCategoryInfo& Category = LogProvider.GetCategory(CategoryPointer);
-
 		FString Name = FTraceAnalyzerUtils::LegacyAttachmentString<TCHAR>("Name", Context);
-		Category.Name = Session.StoreString(*Name);
+		const TCHAR* CategoryName = Session.StoreString(*Name);
+		ELogVerbosity::Type CategoryDefaultVerbosity = static_cast<ELogVerbosity::Type>(EventData.GetValue<uint8>("DefaultVerbosity"));
 
-		Category.DefaultVerbosity = static_cast<ELogVerbosity::Type>(EventData.GetValue<uint8>("DefaultVerbosity"));
+		FAnalysisSessionEditScope _(Session);
+		FLogCategoryInfo& Category = LogProvider.GetCategory(CategoryPointer);
+		Category.Name = CategoryName;
+		Category.DefaultVerbosity = CategoryDefaultVerbosity;
+		//LogProvider.UpdateCategory(CategoryPointer);
 		break;
 	}
 	case RouteId_LogMessageSpec:
 	{
 		uint64 LogPoint = EventData.GetValue<uint64>("LogPoint");
-		FLogMessageSpec& Spec = LogProvider.GetMessageSpec(LogPoint);
+
 		uint64 CategoryPointer = EventData.GetValue<uint64>("CategoryPointer");
-		FLogCategoryInfo& Category = LogProvider.GetCategory(CategoryPointer);
-		Spec.Category = &Category;
-		Spec.Line = EventData.GetValue<int32>("Line");
-		Spec.Verbosity = static_cast<ELogVerbosity::Type>(EventData.GetValue<uint8>("Verbosity"));
 
-		FString FileName;
-		if (EventData.GetString("FileName", FileName))
-		{
-			Spec.File = Session.StoreString(*FileName);
+		int32 Line = EventData.GetValue<int32>("Line");
+		ELogVerbosity::Type Verbosity = static_cast<ELogVerbosity::Type>(EventData.GetValue<uint8>("Verbosity"));
 
-			FString FormatString;
-			EventData.GetString("FormatString", FormatString);
-			Spec.FormatString = Session.StoreString(*FormatString);
+		const TCHAR* File = nullptr;
+		const TCHAR* FormatString = nullptr;
+		FString Str;
+		if (EventData.GetString("FileName", Str))
+		{
+			File = Session.StoreString(*Str);
+			if (EventData.GetString("FormatString", Str))
+			{
+				FormatString = Session.StoreString(*Str);
+			}
 		}
-		else
+		else // backward compatibility
 		{
-			const ANSICHAR* File = reinterpret_cast<const ANSICHAR*>(EventData.GetAttachment());
-			Spec.File = Session.StoreString(ANSI_TO_TCHAR(File));
-			Spec.FormatString = Session.StoreString(reinterpret_cast<const TCHAR*>(EventData.GetAttachment() + strlen(File) + 1));
+			const ANSICHAR* AnsiFile = reinterpret_cast<const ANSICHAR*>(EventData.GetAttachment());
+			if (AnsiFile)
+			{
+				File = Session.StoreString(ANSI_TO_TCHAR(AnsiFile));
+				FormatString = Session.StoreString(reinterpret_cast<const TCHAR*>(EventData.GetAttachment() + TCString<ANSICHAR>::Strlen(AnsiFile) + 1));
+			}
 		}
 
+		FAnalysisSessionEditScope _(Session);
+		FLogMessageSpec& Spec = LogProvider.GetMessageSpec(LogPoint);
+		Spec.Category = &LogProvider.GetCategory(CategoryPointer);
+		Spec.Line = Line;
+		Spec.Verbosity = Verbosity;
+		Spec.File = File;
+		Spec.FormatString = FormatString;
+		LogProvider.UpdateMessageSpec(LogPoint);
 		break;
 	}
 	case RouteId_LogMessage:
 	{
 		uint64 LogPoint = EventData.GetValue<uint64>("LogPoint");
 		uint64 Cycle = EventData.GetValue<uint64>("Cycle");
+		double Time = Context.EventTime.AsSeconds(Cycle);
 		TArrayView<const uint8> FormatArgsView = FTraceAnalyzerUtils::LegacyAttachmentArray("FormatArgs", Context);
-		LogProvider.AppendMessage(LogPoint, Context.EventTime.AsSeconds(Cycle), FormatArgsView.GetData());
+
+		FAnalysisSessionEditScope _(Session);
+		LogProvider.AppendMessage(LogPoint, Time, FormatArgsView);
 		break;
 	}
 	}
diff --git a/Engine/Source/Developer/TraceServices/Private/Analyzers/LogTraceAnalysis.h b/Engine/Source/Developer/TraceServices/Private/Analyzers/LogTraceAnalysis.h
--- a/Engine/Source/Developer/TraceServices/Private/Analyzers/LogTraceAnalysis.h
+++ b/Engine/Source/Developer/TraceServices/Private/Analyzers/LogTraceAnalysis.h
@@ -16,6 +16,7 @@ class FLogTraceAnalyzer
 public:
 	FLogTraceAnalyzer(IAnalysisSession& Session, FLogProvider& LogProvider);
 	virtual void OnAnalysisBegin(const FOnAnalysisContext& Context) override;
+	virtual void OnAnalysisEnd() override;
 	virtual bool OnEvent(uint16 RouteId, EStyle Style, const FOnEventContext& Context) override;
 
 private:
diff --git a/Engine/Source/Developer/TraceServices/Private/Model/Log.cpp b/Engine/Source/Developer/TraceServices/Private/Model/Log.cpp
--- a/Engine/Source/Developer/TraceServices/Private/Model/Log.cpp
+++ b/Engine/Source/Developer/TraceServices/Private/Model/Log.cpp
@@ -41,6 +41,18 @@ FLogProvider::FLogProvider(IAnalysisSession& InSession)
 		AddColumn(&FLogMessageInternal::Message, TEXT("Message"));
 }
 
+FLogProvider::~FLogProvider()
+{
+	for (const auto& KV : MissingLogPoints)
+	{
+		if (KV.Value.FormatArgs)
+		{
+			FMemory::Free(KV.Value.FormatArgs);
+		}
+	}
+	MissingLogPoints.Reset();
+}
+
 uint64 FLogProvider::RegisterCategory()
 {
 	static uint64 IdGenerator = 0;
@@ -50,9 +62,10 @@ uint64 FLogProvider::RegisterCategory()
 FLogCategoryInfo& FLogProvider::GetCategory(uint64 CategoryPointer)
 {
 	Session.WriteAccessCheck();
-	if (CategoryMap.Contains(CategoryPointer))
+	FLogCategoryInfo** FoundCategory = CategoryMap.Find(CategoryPointer);
+	if (FoundCategory)
 	{
-		return *CategoryMap[CategoryPointer];
+		return **FoundCategory;
 	}
 	else
 	{
@@ -67,9 +80,10 @@ FLogCategoryInfo& FLogProvider::GetCategory(uint64 CategoryPointer)
 FLogMessageSpec& FLogProvider::GetMessageSpec(uint64 LogPoint)
 {
 	Session.WriteAccessCheck();
-	if (SpecMap.Contains(LogPoint))
+	FLogMessageSpec** FoundMessageSpec = SpecMap.Find(LogPoint);
+	if (FoundMessageSpec)
 	{
-		return *SpecMap[LogPoint];
+		return **FoundMessageSpec;
 	}
 	else
 	{
@@ -117,9 +131,40 @@ void FLogProvider::UpdateMessageSpec(uint64 LogPoint, uint64 InCategoryPointer,
 	LogMessageSpec.File = InFile;
 	LogMessageSpec.Line = InLine;
 	LogMessageSpec.Verbosity = InVerbosity;
+	UpdateMessageSpec(LogPoint);
 }
 
-FLogMessageInternal& FLogProvider::AppendMessageInternal(double Time)
+void FLogProvider::UpdateMessageSpec(uint64 LogPoint)
+{
+	Session.WriteAccessCheck();
+
+	if (!MissingLogPoints.Contains(LogPoint))
+	{
+		return;
+	}
+
+	FLogMessageSpec** FoundMessageSpec = SpecMap.Find(LogPoint);
+	if (!FoundMessageSpec)
+	{
+		return;
+	}
+	const FLogMessageSpec* Spec = *FoundMessageSpec;
+	check(Spec != nullptr);
+
+	TArray<FMissingMessage> MissingMessages;
+	MissingLogPoints.MultiFind(LogPoint, MissingMessages, false);
+	for (const FMissingMessage& MissingMessage : MissingMessages)
+	{
+		check(MissingMessage.Index < Messages.Num());
+		FLogMessageInternal& InternalMessage = Messages[MissingMessage.Index];
+		check(InternalMessage.Spec == Spec);
+		FFormatArgsHelper::Format(FormatBuffer, FormatBufferSize - 1, TempBuffer, FormatBufferSize - 1, Spec->FormatString, MissingMessage.FormatArgs);
+		InternalMessage.Message = Session.StoreString(FormatBuffer);
+	}
+	MissingLogPoints.Remove(LogPoint);
+}
+
+FLogMessageInternal& FLogProvider::AppendMessageInternal(double Time, uint64* OutIndex)
 {
 	// Performs binary search, resulting in position of the first log message with time > provided time value.
 	uint64 Index = TraceServices::PagedArrayAlgo::UpperBoundBy(Messages, Time,
@@ -130,6 +175,19 @@ FLogMessageInternal& FLogProvider::AppendMessageInternal(double Time)
 		++NumInserts;
 	}
 
+	for (auto& KV : MissingLogPoints)
+	{
+		if (KV.Value.Index >= Index)
+		{
+			KV.Value.Index++;
+		}
+	}
+
+	if (OutIndex)
+	{
+		*OutIndex = Index;
+	}
+
 	FLogMessageInternal& InternalMessage = Messages.Insert(Index);
 	InternalMessage.Time = Time;
 	return InternalMessage;
@@ -152,14 +210,14 @@ void FLogProvider::AppendUnknownMessageInternal(uint64 LogPoint, double Time, co
 void FLogProvider::AppendMessage(uint64 LogPoint, double Time, const uint8* FormatArgs)
 {
 	Session.WriteAccessCheck();
-	FLogMessageSpec** FindSpec = SpecMap.Find(LogPoint);
-	if (FindSpec)
+	FLogMessageSpec* Spec = SpecMap.FindRef(LogPoint);
+	if (Spec)
 	{
-		if ((*FindSpec)->Verbosity != ELogVerbosity::SetColor)
+		if (Spec->Verbosity != ELogVerbosity::SetColor)
 		{
 			FLogMessageInternal& InternalMessage = AppendMessageInternal(Time);
-			InternalMessage.Spec = *FindSpec;
-			FFormatArgsHelper::Format(FormatBuffer, FormatBufferSize - 1, TempBuffer, FormatBufferSize - 1, InternalMessage.Spec->FormatString, FormatArgs);
+			InternalMessage.Spec = Spec;
+			FFormatArgsHelper::Format(FormatBuffer, FormatBufferSize - 1, TempBuffer, FormatBufferSize - 1, Spec->FormatString, FormatArgs);
 			InternalMessage.Message = Session.StoreString(FormatBuffer);
 		}
 	}
@@ -171,36 +229,82 @@ void FLogProvider::AppendMessage(uint64 LogPoint, double Time, const uint8* Form
 	Session.UpdateDurationSeconds(Time);
 }
 
-void FLogProvider::AppendMessage(uint64 LogPoint, double Time, const TCHAR* Text)
+void FLogProvider::AppendMessage(uint64 LogPoint, double Time, TArrayView<const uint8> FormatArgsView)
 {
 	Session.WriteAccessCheck();
-	FLogMessageSpec** FindSpec = SpecMap.Find(LogPoint);
-	if (FindSpec)
+
+	FLogMessageSpec* Spec = SpecMap.FindRef(LogPoint);
+	if (Spec)
 	{
-		if ((*FindSpec)->Verbosity != ELogVerbosity::SetColor)
+		if (Spec->Verbosity != ELogVerbosity::SetColor)
 		{
 			FLogMessageInternal& InternalMessage = AppendMessageInternal(Time);
-			InternalMessage.Spec = *FindSpec;
-			InternalMessage.Message = Text;
-		}
-		else
-		{
-			AppendUnknownMessageInternal(LogPoint, Time, FStringView(Text));
+			InternalMessage.Spec = Spec;
+			FFormatArgsHelper::Format(FormatBuffer, FormatBufferSize - 1, TempBuffer, FormatBufferSize - 1, Spec->FormatString, FormatArgsView.GetData());
+			InternalMessage.Message = Session.StoreString(FormatBuffer);
 		}
 	}
+	else
+	{
+		Spec = &GetMessageSpec(LogPoint);
+		FLogCategoryInfo& Category = GetCategory(~0ull);
+		Spec->Category = &Category;
+		Spec->Line = 0;
+		Spec->Verbosity = ELogVerbosity::Error;
+		Spec->FormatString = TEXT("%s");
+
+		uint64 Index = 0;
+		FLogMessageInternal& InternalMessage = AppendMessageInternal(Time, &Index);
+		InternalMessage.Spec = Spec;
+		FFormatArgsHelper::Format(FormatBuffer, FormatBufferSize - 1, TempBuffer, FormatBufferSize - 1, TEXT("%s"), FormatArgsView.GetData());
+		FStringView Message(FormatBuffer);
+		InternalMessage.Message = Session.StoreString(*FString::Printf(TEXT("Unknown log message spec (LogPoint=0x%llX)! Message: \"%.*s\""), LogPoint, Message.Len(), Message.GetData()));
+
+		uint8* FormatArgs = (uint8*)FMemory::Malloc(FormatArgsView.Num());
+		FMemory::Memcpy(FormatArgs, FormatArgsView.GetData(), FormatArgsView.Num());
+		MissingLogPoints.Add(LogPoint, { Index, FormatArgs });
+		++NumUnknownSpecs;
+	}
+
+	Session.UpdateDurationSeconds(Time);
+}
+
+void FLogProvider::AppendMessage(uint64 LogPoint, double Time, const TCHAR* StoredText)
+{
+	Session.WriteAccessCheck();
+
+	FLogMessageSpec* Spec = SpecMap.FindRef(LogPoint);
+	if (!Spec)
+	{
+		Spec = &GetMessageSpec(LogPoint);
+		FLogCategoryInfo& Category = GetCategory(~0ull);
+		Spec->Category = &Category;
+		Spec->Line = 0;
+		Spec->Verbosity = ELogVerbosity::Error;
+		Spec->FormatString = TEXT("%s");
+	}
+
+	if (Spec->Verbosity != ELogVerbosity::SetColor)
+	{
+		FLogMessageInternal& InternalMessage = AppendMessageInternal(Time);
+		InternalMessage.Spec = Spec;
+		InternalMessage.Message = StoredText;
+	}
+
 	Session.UpdateDurationSeconds(Time);
 }
 
 void FLogProvider::AppendMessage(uint64 LogPoint, double Time, const FStringView Message)
 {
 	Session.WriteAccessCheck();
-	FLogMessageSpec** FindSpec = SpecMap.Find(LogPoint);
-	if (FindSpec)
+
+	FLogMessageSpec* Spec = SpecMap.FindRef(LogPoint);
+	if (Spec)
 	{
-		if ((*FindSpec)->Verbosity != ELogVerbosity::SetColor)
+		if (Spec->Verbosity != ELogVerbosity::SetColor)
 		{
 			FLogMessageInternal& InternalMessage = AppendMessageInternal(Time);
-			InternalMessage.Spec = *FindSpec;
+			InternalMessage.Spec = Spec;
 			InternalMessage.Message = Session.StoreString(Message);
 		}
 	}
@@ -208,6 +312,7 @@ void FLogProvider::AppendMessage(uint64 LogPoint, double Time, const FStringView
 	{
 		AppendUnknownMessageInternal(LogPoint, Time, Message);
 	}
+
 	Session.UpdateDurationSeconds(Time);
 }
 
@@ -343,7 +448,7 @@ uint64 FLogProvider::BinarySearchClosestByTime(double Time) const
 	return TraceServices::PagedArrayAlgo::BinarySearchClosestBy(Messages, Time,
 		[](const FLogMessageInternal& Item) { return Item.Time; });
 }
-	
+
 bool FLogProvider::ConstructMessage(const FLogMessageInternal& InternalMessage, uint64 Index, TFunctionRef<void(const FLogMessageInfo&)> Callback) const
 {
 	FLogMessageInfo Message;
diff --git a/Engine/Source/Developer/TraceServices/Private/Model/LogPrivate.h b/Engine/Source/Developer/TraceServices/Private/Model/LogPrivate.h
--- a/Engine/Source/Developer/TraceServices/Private/Model/LogPrivate.h
+++ b/Engine/Source/Developer/TraceServices/Private/Model/LogPrivate.h
@@ -2,11 +2,15 @@
 
 #pragma once
 
-#include "TraceServices/Model/Log.h"
+#include "CoreTypes.h"
+#include "Containers/Map.h"
 #include "Templates/SharedPointer.h"
+#include "Misc/OutputDeviceHelper.h"
+
+// TraceServices
 #include "Common/PagedArray.h"
 #include "Model/Tables.h"
-#include "Misc/OutputDeviceHelper.h"
+#include "TraceServices/Model/Log.h"
 
 namespace TraceServices
 {
@@ -36,7 +40,7 @@ class FLogProvider
 {
 public:
 	explicit FLogProvider(IAnalysisSession& Session);
-	virtual ~FLogProvider() {}
+	virtual ~FLogProvider();
 
 	//////////////////////////////////////////////////
 	// Read operations
@@ -46,17 +50,21 @@ public:
 	virtual void EnumerateMessagesByIndex(uint64 StartIndex, uint64 EndIndex, TFunctionRef<void(const FLogMessageInfo&)> Callback) const override;
 	virtual void EnumerateMessages(double StartTime, double EndTime, TFunctionRef<void(const FLogMessageInfo&)> Callback) const override;
 	virtual void ReverseEnumerateMessages(double StartTime, double EndTime, TFunctionRef<bool(const FLogMessageInfo&)> Callback) const override;
-	
+
 	virtual uint64 LowerBoundByTime(double Time) const override;
 	virtual uint64 UpperBoundByTime(double Time) const override;
 	virtual uint64 BinarySearchClosestByTime(double Time) const override;
 
+	virtual uint64 GetMessageSpecCount() const override { Session.ReadAccessCheck(); return MessageSpecs.Num(); }
+
 	virtual uint64 GetCategoryCount() const override { Session.ReadAccessCheck(); return Categories.Num(); }
 	virtual void EnumerateCategories(TFunctionRef<void(const FLogCategoryInfo&)> Callback) const override;
 
 	virtual const IUntypedTable& GetMessagesTable() const override { return MessagesTable; }
 
-	virtual uint64 GetInsertCount() const { Session.ReadAccessCheck(); return NumInserts; }
+	virtual uint64 GetInsertCount() const override { Session.ReadAccessCheck(); return NumInserts; }
+	virtual uint64 GetUnknownSpecTotalCount() const override { Session.ReadAccessCheck(); return NumUnknownSpecs; }
+	virtual uint64 GetUnresolvedSpecCount() const override { Session.ReadAccessCheck(); return uint64(MissingLogPoints.Num()); }
 
 	//////////////////////////////////////////////////
 	// Edit operations
@@ -70,14 +78,17 @@ public:
 	virtual void UpdateMessageFile(uint64 LogPoint, const TCHAR* InFile, int32 InLine) override;
 	virtual void UpdateMessageVerbosity(uint64 LogPoint, ELogVerbosity::Type InVerbosity) override;
 	virtual void UpdateMessageSpec(uint64 LogPoint, uint64 InCategoryPointer, const TCHAR* InFormatString, const TCHAR* InFile, int32 InLine, ELogVerbosity::Type InVerbosity) override;
+	virtual void UpdateMessageSpec(uint64 LogPoint) override;
+
 	virtual void AppendMessage(uint64 LogPoint, double Time, const uint8* FormatArgs) override;
-	virtual void AppendMessage(uint64 LogPoint, double Time, const TCHAR* Text) override;
+	virtual void AppendMessage(uint64 LogPoint, double Time, TArrayView<const uint8> FormatArgsView) override;
+	virtual void AppendMessage(uint64 LogPoint, double Time, const TCHAR* StoredText) override;
 	void AppendMessage(uint64 LogPoint, double Time, const FStringView Message);
 
 	//////////////////////////////////////////////////
 
 private:
-	FLogMessageInternal& AppendMessageInternal(double Time);
+	FLogMessageInternal& AppendMessageInternal(double Time, uint64* OutIndex = nullptr);
 	void AppendUnknownMessageInternal(uint64 LogPoint, double Time, const FStringView Message);
 	bool ConstructMessage(const FLogMessageInternal& InternalMessage, uint64 Index, TFunctionRef<void(const FLogMessageInfo&)> Callback) const;
 	bool ConstructMessage(const FLogMessageInternal& InternalMessage, uint64 Index, TFunctionRef<bool(const FLogMessageInfo&)> Callback) const;
@@ -87,16 +98,24 @@ private:
 		FormatBufferSize = 65536
 	};
 
+	struct FMissingMessage // message with missing spec
+	{
+		uint64 Index = ~0ull; // index in Messages array
+		uint8* FormatArgs = nullptr;
+	};
+
 	IAnalysisSession& Session;
-	TMap<uint64, FLogCategoryInfo*> CategoryMap;
-	TMap<uint64, FLogMessageSpec*> SpecMap;
+	TMap<uint64, FLogCategoryInfo*> CategoryMap; // CategoryPointer --> FLogCategoryInfo*
+	TMap<uint64, FLogMessageSpec*> SpecMap; // LogPoint --> FLogMessageSpec*
 	TPagedArray<FLogCategoryInfo> Categories;
 	TPagedArray<FLogMessageSpec> MessageSpecs;
 	TPagedArray<FLogMessageInternal> Messages;
 	TCHAR FormatBuffer[FormatBufferSize];
 	TCHAR TempBuffer[FormatBufferSize];
 	TTableView<FLogMessageInternal> MessagesTable;
-	uint64 NumInserts = 0;
+	TMultiMap<uint64, FMissingMessage> MissingLogPoints; // LogPoint --> FMissingMessage[]
+	uint64 NumInserts = 0; // total number of inserts (out of order messages)
+	uint64 NumUnknownSpecs = 0; // total number of log messages received with unknown spec
 };
 
 } // namespace TraceServices
diff --git a/Engine/Source/Developer/TraceServices/Public/TraceServices/Model/Log.h b/Engine/Source/Developer/TraceServices/Public/TraceServices/Model/Log.h
--- a/Engine/Source/Developer/TraceServices/Public/TraceServices/Model/Log.h
+++ b/Engine/Source/Developer/TraceServices/Public/TraceServices/Model/Log.h
@@ -119,6 +119,13 @@ public:
 	 */
 	virtual uint64 BinarySearchClosestByTime(double Time) const { return ~0ull; }
 
+	/**
+	 * Gets the number of log message specs.
+	 *
+	 * @returns The current number of log message specs.
+	 */
+	virtual uint64 GetMessageSpecCount() const { return 0; }
+
 	/**
 	 * Gets the number of log categories.
 	 *
@@ -148,6 +155,20 @@ public:
 	 * @returns The number of inserts.
 	 */
 	virtual uint64 GetInsertCount() const { return 0; }
+
+	/**
+	 * Gets the total number of messages received with unknown message spec.
+	 *
+	 * @returns The total number of messages received with unknown message spec.
+	 */
+	virtual uint64 GetUnknownSpecTotalCount() const { return 0; }
+
+	/**
+	 * Gets the number of messages received with unknown message spec and still unresolved.
+	 *
+	 * @returns The number of messages received with unknown message spec and still unresolved.
+	 */
+	virtual uint64 GetUnresolvedSpecCount() const { return 0; }
 };
 
 class IEditableLogProvider
@@ -216,6 +237,11 @@ public:
 	 */
 	virtual void UpdateMessageSpec(uint64 LogPoint, uint64 InCategoryPointer, const TCHAR* InFormatString, const TCHAR* InFile, int32 InLine, ELogVerbosity::Type InVerbosity) = 0;
 
+	/**
+	 * Updates information for a log message after message spec was updated.
+	 */
+	virtual void UpdateMessageSpec(uint64 LogPoint) {}
+
 	/**
 	 * Appends a new instance of a message from the trace session.
 	 *
@@ -223,16 +249,26 @@ public:
 	 * @param Time The time in seconds of the event.
 	 * @param FormatArgs The arguments to use in conjunction with the spec's FormatString.
 	 */
-	virtual void AppendMessage(uint64 LogPoint, double Time, const uint8* FormatArgs) = 0;
+	UE_DEPRECATED(5.8, "Use AppendMessage(uint64, double, TArrayView<const uint8>) instead.")
+	virtual void AppendMessage(uint64 LogPoint, double Time, const uint8* FormatArgs) {}
+
+	/**
+	 * Appends a new instance of a message from the trace session.
+	 *
+	 * @param LogPoint The unique identity (memory address) of the message instrumentation.
+	 * @param Time The time in seconds of the event.
+	 * @param FormatArgsView The arguments to use in conjunction with the spec's FormatString.
+	 */
+	virtual void AppendMessage(uint64 LogPoint, double Time, TArrayView<const uint8> FormatArgsView) {}
 
 	/**
 	 * Appends a new instance of a message from the trace session.
 	 *
 	 * @param LogPoint The unique identity (memory address) of the message instrumentation.
 	 * @param Time The time in seconds of the event.
-	 * @param Text The message text.
+	 * @param StoredText The persistent message text.
 	 */
-	virtual void AppendMessage(uint64 LogPoint, double Time, const TCHAR* Text) = 0;
+	virtual void AppendMessage(uint64 LogPoint, double Time, const TCHAR* StoredText) = 0;
 };
 
 TRACESERVICES_API FName GetLogProviderName();
