[Insights] Logs: Fixed memory leak regarding handling of missing log points.

#rnx

diff --git a/Engine/Source/Developer/TraceServices/Private/Model/Log.cpp b/Engine/Source/Developer/TraceServices/Private/Model/Log.cpp
--- a/Engine/Source/Developer/TraceServices/Private/Model/Log.cpp
+++ b/Engine/Source/Developer/TraceServices/Private/Model/Log.cpp
@@ -160,6 +160,10 @@ void FLogProvider::UpdateMessageSpec(uint64 LogPoint)
 		check(InternalMessage.Spec == Spec);
 		FFormatArgsHelper::Format(FormatBuffer, FormatBufferSize - 1, TempBuffer, FormatBufferSize - 1, Spec->FormatString, MissingMessage.FormatArgs);
 		InternalMessage.Message = Session.StoreString(FormatBuffer);
+		if (MissingMessage.FormatArgs)
+		{
+			FMemory::Free(MissingMessage.FormatArgs);
+		}
 	}
 	MissingLogPoints.Remove(LogPoint);
 }
@@ -259,9 +263,12 @@ void FLogProvider::AppendMessage(uint64 LogPoint, double Time, TArrayView<const
 		FFormatArgsHelper::Format(FormatBuffer, FormatBufferSize - 1, TempBuffer, FormatBufferSize - 1, TEXT("%s"), FormatArgsView.GetData());
 		FStringView Message(FormatBuffer);
 		InternalMessage.Message = Session.StoreString(*FString::Printf(TEXT("Unknown log message spec (LogPoint=0x%llX)! Message: \"%.*s\""), LogPoint, Message.Len(), Message.GetData()));
-
-		uint8* FormatArgs = (uint8*)FMemory::Malloc(FormatArgsView.Num());
-		FMemory::Memcpy(FormatArgs, FormatArgsView.GetData(), FormatArgsView.Num());
+		uint8* FormatArgs = nullptr;
+		if (FormatArgsView.Num() > 0)
+		{
+			FormatArgs = (uint8*)FMemory::Malloc(FormatArgsView.Num());
+			FMemory::Memcpy(FormatArgs, FormatArgsView.GetData(), FormatArgsView.Num());
+		}
 		MissingLogPoints.Add(LogPoint, { Index, FormatArgs });
 		++NumUnknownSpecs;
 	}
