[Insights]
* MonotonicTimeline: Fixed "always true" condition. Reason: We have Depth >= 0, so in the "while (Depth < CurrentStackDepth) {}" block CurrentStackDepth > 0. Then CurrentStackDepth++ will make condition "if (CurrentStackDepth > 1)" always true.
* MonotonicTimeline: Minor code refactoring.

#rnx

diff --git a/Engine/Source/Developer/TraceServices/Public/Model/MonotonicTimeline.h b/Engine/Source/Developer/TraceServices/Public/Model/MonotonicTimeline.h
--- a/Engine/Source/Developer/TraceServices/Public/Model/MonotonicTimeline.h
+++ b/Engine/Source/Developer/TraceServices/Public/Model/MonotonicTimeline.h
@@ -111,16 +111,17 @@ public:
 		for (int32 DetailLevelIndex = 1; DetailLevelIndex < SettingsType::DetailLevelsCount; ++DetailLevelIndex)
 		{
 			FDetailLevel& DetailLevel = DetailLevels[DetailLevelIndex];
-			FDetailLevelDepthState& CurrentDepthState = DetailLevel.InsertionState.DepthStates[CurrentDepth];
+			FDetailLevelInsertionState<EventType, SettingsType>& InsertionState = DetailLevel.InsertionState;
+			FDetailLevelDepthState& CurrentDepthState = InsertionState.DepthStates[CurrentDepth];
 
 			if (CurrentDepthState.PendingScopeEnterIndex < 0 || StartTime >= CurrentDepthState.EnterTime + DetailLevel.Resolution)
 			{
 				if (CurrentDepthState.PendingEventIndex >= 0)
 				{
-					check(DetailLevel.InsertionState.PendingDepth < SettingsType::MaxDepth);
-					for (int32 Depth = DetailLevel.InsertionState.PendingDepth; Depth >= CurrentDepth; --Depth)
+					check(InsertionState.PendingDepth < SettingsType::MaxDepth);
+					for (int32 Depth = InsertionState.PendingDepth; Depth >= CurrentDepth; --Depth)
 					{
-						FDetailLevelDepthState& DepthState = DetailLevel.InsertionState.DepthStates[Depth];
+						FDetailLevelDepthState& DepthState = InsertionState.DepthStates[Depth];
 						check(DepthState.PendingScopeEnterIndex >= 0);
 						AddScopeEntry(DetailLevel, DepthState.ExitTime, false);
 
@@ -128,7 +129,7 @@ public:
 						DepthState.PendingEventIndex = -1;
 					}
 				}
-				DetailLevel.InsertionState.PendingDepth = CurrentDepth;
+				InsertionState.PendingDepth = CurrentDepth;
 
 				uint64 EnterScopeIndex = DetailLevel.ScopeEntries.Num();
 				uint64 EventIndex = DetailLevel.Events.Num();
@@ -145,9 +146,9 @@ public:
 				CurrentDepthState.DominatingEvent = Event;
 				//CurrentDepthState.DebugDominatingEventType = Owner.EventTypes[TypeId];
 			}
-			else if (CurrentDepth > DetailLevel.InsertionState.PendingDepth)
+			else if (CurrentDepth > InsertionState.PendingDepth)
 			{
-				DetailLevel.InsertionState.PendingDepth = CurrentDepth;
+				InsertionState.PendingDepth = CurrentDepth;
 			}
 			DetailLevel.SetEvent(CurrentDepthState.PendingEventIndex, Event);
 		}
@@ -173,26 +174,27 @@ public:
 		for (int32 DetailLevelIndex = 1; DetailLevelIndex < SettingsType::DetailLevelsCount; ++DetailLevelIndex)
 		{
 			FDetailLevel& DetailLevel = DetailLevels[DetailLevelIndex];
+			FDetailLevelInsertionState<EventType, SettingsType>& InsertionState = DetailLevel.InsertionState;
 
-			DetailLevel.InsertionState.DepthStates[CurrentDepth].ExitTime = EndTime;
+			InsertionState.DepthStates[CurrentDepth].ExitTime = EndTime;
 
 			UpdateDominatingEvent(DetailLevel, CurrentDepth, EndTime);
 
-			FDetailLevelDepthState& CurrentDepthState = DetailLevel.InsertionState.DepthStates[CurrentDepth];
+			FDetailLevelDepthState& CurrentDepthState = InsertionState.DepthStates[CurrentDepth];
 			check(CurrentDepthState.PendingScopeEnterIndex >= 0);
 			if (EndTime >= CurrentDepthState.EnterTime + DetailLevel.Resolution)
 			{
-				check(DetailLevel.InsertionState.PendingDepth < SettingsType::MaxDepth);
-				for (int32 Depth = DetailLevel.InsertionState.PendingDepth; Depth >= CurrentDepth; --Depth)
+				check(InsertionState.PendingDepth < SettingsType::MaxDepth);
+				for (int32 Depth = InsertionState.PendingDepth; Depth >= CurrentDepth; --Depth)
 				{
-					FDetailLevelDepthState& DepthState = DetailLevel.InsertionState.DepthStates[Depth];
+					FDetailLevelDepthState& DepthState = InsertionState.DepthStates[Depth];
 					check(DepthState.PendingScopeEnterIndex >= 0);
 					AddScopeEntry(DetailLevel, DepthState.ExitTime, false);
 
 					DepthState.PendingScopeEnterIndex = -1;
 					DepthState.PendingEventIndex = -1;
 				}
-				DetailLevel.InsertionState.PendingDepth = CurrentDepth - 1;
+				InsertionState.PendingDepth = CurrentDepth - 1;
 			}
 		}
 
@@ -307,6 +309,7 @@ public:
 				check(CurrentStackDepth > 0);
 				--CurrentStackDepth;
 			}
+
 			ScopeEntry = ScopeEntryIterator.NextItem();
 		}
 		if (CurrentStackDepth == 1 && EventStack[0].StartTime > IntervalEnd)
@@ -344,6 +347,7 @@ public:
 					return;
 				}
 			}
+
 			ScopeEntry = ScopeEntryIterator.NextItem();
 		}
 
@@ -357,6 +361,7 @@ public:
 				bSearchEndTimeUsingPages = true;
 				break;
 			}
+
 			if (ScopeEntry->Time < 0.0) // "begin" event
 			{
 				++ExitDepth;
@@ -435,6 +440,7 @@ public:
 			{
 				if (IsEnter)
 				{
+					check(CurrentDepth < SettingsType::MaxDepth);
 					FStackEntry& StackEntry = EventStack[CurrentDepth];
 					StackEntry.LocalEventIndex = OutputEvents.Num();
 					FOutputEvent& OutputEvent = OutputEvents.AddDefaulted_GetRef();
@@ -447,6 +453,7 @@ public:
 				else
 				{
 					{
+						check(CurrentDepth > 0);
 						FStackEntry& StackEntry = EventStack[--CurrentDepth];
 						FOutputEvent* OutputEvent = OutputEvents.GetData() + StackEntry.LocalEventIndex;
 						OutputEvent->EndTime = Time;
@@ -547,12 +554,13 @@ public:
 		else
 		{
 			// If we start from the end of the session, we use InsertionState as the initial stack.
-			CurrentStackDepth = DetailLevel.InsertionState.CurrentDepth;
+			const FDetailLevelInsertionState<EventType, SettingsType>& InsertionState = DetailLevel.InsertionState;
+			CurrentStackDepth = InsertionState.CurrentDepth;
 			for (int32 InitialStackIndex = 0; InitialStackIndex < CurrentStackDepth; ++InitialStackIndex)
 			{
 				FEnumerationStackEntry& EnumerationStackEntry = EventStack[InitialStackIndex];
-				const FEventStackEntry& EventStackEntry = DetailLevel.InsertionState.EventStack[InitialStackIndex];
-				EnumerationStackEntry.EndTime = DetailLevel.InsertionState.LastTime;
+				const FEventStackEntry& EventStackEntry = InsertionState.EventStack[InitialStackIndex];
+				EnumerationStackEntry.EndTime = InsertionState.LastTime;
 				EnumerationStackEntry.Event = DetailLevel.GetEvent(EventStackEntry.EventIndex);
 			}
 			ScopeEntry = ScopeEntryIterator.GetCurrentItem();
@@ -574,6 +582,7 @@ public:
 				FEnumerationStackEntry& StackEntry = EventStack[CurrentStackDepth++];
 				StackEntry.EndTime = ScopeEntry->Time;
 			}
+
 			ScopeEntry = ScopeEntryIterator.PrevItem();
 		}
 
@@ -612,6 +621,7 @@ public:
 					return;
 				}
 			}
+
 			ScopeEntry = ScopeEntryIterator.PrevItem();
 		}
 
@@ -825,7 +835,6 @@ public:
 		}
 
 		const FDetailLevel& DetailLevel = DetailLevels[0];
-
 		if (DetailLevel.ScopeEntries.Num() == 0)
 		{
 			return false;
@@ -906,11 +915,13 @@ public:
 				StackEntry.Event = *Event;
 				StackEntry.StartTime = -ScopeEntry->Time;
 				StackEntry.ExclTime = 0.0;
+
 				if (CurrentStackDepth > 1)
 				{
 					FEventInfoStackEntry& ParentStackEntry = EventStack[CurrentStackDepth - 2];
 					ParentStackEntry.ExclTime += StackEntry.StartTime - LastTime;
 				}
+
 				Event = EventsIterator.NextItem();
 				LastTime = -ScopeEntry->Time;
 			}
@@ -924,6 +935,7 @@ public:
 					FEventInfoStackEntry& StackEntry = EventStack[CurrentStackDepth];
 					StackEntry.ExclTime += ScopeEntry->Time - LastTime;
 				}
+
 				LastTime = ScopeEntry->Time;
 			}
 
@@ -1043,11 +1055,13 @@ public:
 				FEventInfoStackEntry& StackEntry = EventStack[CurrentStackDepth++];
 				StackEntry.Event = *Event;
 				StackEntry.StartTime = -ScopeEntry->Time;
-				if (CurrentStackDepth > 1)
+
+				//if (ensure(CurrentStackDepth > 1)) //-V547
 				{
 					FEventInfoStackEntry& ParentStackEntry = EventStack[CurrentStackDepth - 2];
 					ParentStackEntry.ExclTime += StackEntry.StartTime - LastTime;
 				}
+
 				Event = EventsIterator.NextItem();
 				LastTime = -ScopeEntry->Time;
 			}
@@ -1055,6 +1069,7 @@ public:
 			{
 				check(CurrentStackDepth > 0);
 				--CurrentStackDepth;
+
 				if (CurrentStackDepth == Depth)
 				{
 					FEventInfoStackEntry& StackEntry = EventStack[CurrentStackDepth];
@@ -1063,6 +1078,7 @@ public:
 
 				LastTime = ScopeEntry->Time;
 			}
+
 			ScopeEntry = ScopeEntryIterator.NextItem();
 		}
 
@@ -1176,7 +1192,8 @@ private:
 
 	void AddScopeEntry(FDetailLevel& DetailLevel, double Time, bool IsEnter)
 	{
-		checkf(Time >= DetailLevel.InsertionState.LastTime, TEXT("Time=%.9f LastTime=%.9f"), Time, DetailLevel.InsertionState.LastTime);
+		FDetailLevelInsertionState<EventType, SettingsType>& InsertionState = DetailLevel.InsertionState;
+		checkf(Time >= InsertionState.LastTime, TEXT("Time=%.9f LastTime=%.9f"), Time, InsertionState.LastTime);
 
 		uint64 EventIndex = DetailLevel.Events.Num();
 		uint64 ScopeIndex = DetailLevel.ScopeEntries.Num();
@@ -1185,21 +1202,21 @@ private:
 		ScopeEntry.Time = IsEnter ? -Time : Time;
 		FEventScopeEntryPage* LastPage = DetailLevel.ScopeEntries.GetLastPage();
 		uint64 LastPageIndex = DetailLevel.ScopeEntries.NumPages() - 1;
-		if (LastPageIndex != DetailLevel.InsertionState.CurrentScopeEntryPageIndex)
+		if (LastPageIndex != InsertionState.CurrentScopeEntryPageIndex)
 		{
 			// At the very first call, CurrentScopeEntryPage will be -1
-			if (DetailLevel.InsertionState.CurrentScopeEntryPageIndex != (uint64) -1)
+			if (InsertionState.CurrentScopeEntryPageIndex != (uint64) -1)
 			{
-				check(DetailLevel.InsertionState.CurrentDepth <= SettingsType::MaxDepth);
+				check(InsertionState.CurrentDepth <= SettingsType::MaxDepth);
 
-				FEventScopeEntryPage* CurrentScopeEntryPage = DetailLevel.ScopeEntries.GetPage(DetailLevel.InsertionState.CurrentScopeEntryPageIndex);
-				const int32 MaxEntryIndex = FMath::Min(CurrentScopeEntryPage->InitialStackCount, DetailLevel.InsertionState.CurrentDepth);
+				FEventScopeEntryPage* CurrentScopeEntryPage = DetailLevel.ScopeEntries.GetPage(InsertionState.CurrentScopeEntryPageIndex);
+				const int32 MaxEntryIndex = FMath::Min(CurrentScopeEntryPage->InitialStackCount, InsertionState.CurrentDepth);
 
 				// Update the open scopes that were also open at the beginning of the last page so the values
 				// represent stats up to and including the current page
 				for (int32 EntryIndex = 0; EntryIndex < MaxEntryIndex; ++EntryIndex)
 				{
-					FEventStackEntry& InsertionStateStackEntry = DetailLevel.InsertionState.EventStack[EntryIndex];
+					FEventStackEntry& InsertionStateStackEntry = InsertionState.EventStack[EntryIndex];
 					FEventStackEntry& PreviousPageStackEntry = CurrentScopeEntryPage->InitialStack[EntryIndex];
 
 					if (InsertionStateStackEntry.EnterScopeIndex == PreviousPageStackEntry.EnterScopeIndex &&
@@ -1214,49 +1231,46 @@ private:
 				}
 			}
 
-			DetailLevel.InsertionState.CurrentScopeEntryPageIndex = LastPageIndex;
+			InsertionState.CurrentScopeEntryPageIndex = LastPageIndex;
 			LastPage->BeginTime = Time;
 			LastPage->BeginEventIndex = DetailLevel.Events.Num();
 			LastPage->EndEventIndex = LastPage->BeginEventIndex;
-			LastPage->InitialStackCount = DetailLevel.InsertionState.CurrentDepth;
+			LastPage->InitialStackCount = InsertionState.CurrentDepth;
 			if (LastPage->InitialStackCount)
 			{
 				LastPage->InitialStack = reinterpret_cast<FEventStackEntry*>(Allocator.Allocate(LastPage->InitialStackCount * sizeof(FEventStackEntry)));
-				memcpy(LastPage->InitialStack, DetailLevel.InsertionState.EventStack, LastPage->InitialStackCount * sizeof(FEventStackEntry));
+				memcpy(LastPage->InitialStack, InsertionState.EventStack, LastPage->InitialStackCount * sizeof(FEventStackEntry));
 			}
 		}
 		LastPage->EndTime = Time;
 
 		if (IsEnter)
 		{
-			++DetailLevel.InsertionState.CurrentDepth;
-			check(DetailLevel.InsertionState.CurrentDepth <= SettingsType::MaxDepth);
-
-			FEventStackEntry& StackEntry = DetailLevel.InsertionState.EventStack[DetailLevel.InsertionState.CurrentDepth - 1];
+			check(InsertionState.CurrentDepth < SettingsType::MaxDepth);
+			FEventStackEntry& StackEntry = InsertionState.EventStack[InsertionState.CurrentDepth++];
 			StackEntry.EventIndex = EventIndex;
 			StackEntry.EnterScopeIndex = ScopeIndex;
 			StackEntry.ExclTime = 0.0;
 			StackEntry.EndTime = -1.0;
 
-			if (DetailLevel.InsertionState.CurrentDepth > 1)
+			if (InsertionState.CurrentDepth > 1)
 			{
-				FEventStackEntry& ParentStackEntry = DetailLevel.InsertionState.EventStack[DetailLevel.InsertionState.CurrentDepth - 2];
-				ParentStackEntry.ExclTime += Time - DetailLevel.InsertionState.LastTime;
+				FEventStackEntry& ParentStackEntry = InsertionState.EventStack[InsertionState.CurrentDepth - 2];
+				ParentStackEntry.ExclTime += Time - InsertionState.LastTime;
 			}
 		}
 		else
 		{
-			check(DetailLevel.InsertionState.CurrentDepth <= SettingsType::MaxDepth);
-			check(DetailLevel.InsertionState.CurrentDepth > 0);
-			--DetailLevel.InsertionState.CurrentDepth;
+			check(InsertionState.CurrentDepth > 0);
+			--InsertionState.CurrentDepth;
 
-			FEventStackEntry& StackEntry = DetailLevel.InsertionState.EventStack[DetailLevel.InsertionState.CurrentDepth];
-			StackEntry.ExclTime += Time - DetailLevel.InsertionState.LastTime;
+			FEventStackEntry& StackEntry = InsertionState.EventStack[InsertionState.CurrentDepth];
+			StackEntry.ExclTime += Time - InsertionState.LastTime;
 
-			FEventScopeEntryPage* CurrentScopeEntryPage = DetailLevel.ScopeEntries.GetPage(DetailLevel.InsertionState.CurrentScopeEntryPageIndex);
-			if (DetailLevel.InsertionState.CurrentDepth < CurrentScopeEntryPage->InitialStackCount)
+			FEventScopeEntryPage* CurrentScopeEntryPage = DetailLevel.ScopeEntries.GetPage(InsertionState.CurrentScopeEntryPageIndex);
+			if (InsertionState.CurrentDepth < CurrentScopeEntryPage->InitialStackCount)
 			{
-				FEventStackEntry& PreviousPageStackEntry = CurrentScopeEntryPage->InitialStack[DetailLevel.InsertionState.CurrentDepth];
+				FEventStackEntry& PreviousPageStackEntry = CurrentScopeEntryPage->InitialStack[InsertionState.CurrentDepth];
 
 				if (StackEntry.EnterScopeIndex == PreviousPageStackEntry.EnterScopeIndex &&
 					StackEntry.EventIndex == PreviousPageStackEntry.EventIndex)
@@ -1267,9 +1281,9 @@ private:
 			}
 		}
 
-		DetailLevel.InsertionState.LastTime = Time;
+		InsertionState.LastTime = Time;
 
-		//ScopeEntry.DebugDepth = DetailLevel.InsertionState.CurrentDepth;
+		//ScopeEntry.DebugDepth = InsertionState.CurrentDepth;
 	}
 
 	void AddEvent(FDetailLevel& DetailLevel, const EventType& Event)
@@ -1357,11 +1371,13 @@ private:
 				StackEntry.Event = *Event;
 				StackEntry.StartTime = -ScopeEntry->Time;
 				StackEntry.ExclTime = 0.0;
+
 				if (IterationState.StackDepth > 1)
 				{
 					FEventInfoStackEntry& ParentStackEntry = IterationState.EventStack[IterationState.StackDepth - 2];
 					ParentStackEntry.ExclTime += StackEntry.StartTime - IterationState.LastIterationTime;
 				}
+
 				Event = IterationState.EventsIterator.NextItem();
 				IterationState.LastIterationTime = -ScopeEntry->Time;
 			}
@@ -1378,6 +1394,7 @@ private:
 					BestMatchEntry = StackEntry;
 					bHasEndEvent = true;
 				}
+
 				IterationState.LastIterationTime = ScopeEntry->Time;
 			}
 
@@ -1409,6 +1426,7 @@ private:
 				StackEntry.Event = *Event;
 				StackEntry.StartTime = -ScopeEntry->Time;
 				StackEntry.ExclTime = 0.0;
+
 				if (IterationState.StackDepth > 1)
 				{
 					FEventInfoStackEntry& ParentStackEntry = IterationState.EventStack[IterationState.StackDepth - 2];
@@ -1435,6 +1453,7 @@ private:
 					FEventInfoStackEntry& StackEntry = IterationState.EventStack[IterationState.StackDepth];
 					StackEntry.ExclTime += ScopeEntry->Time - IterationState.LastIterationTime;
 				}
+
 				IterationState.LastIterationTime = ScopeEntry->Time;
 			}
 
