[Insights]
* MemoryInsights: Updated the way alloc size histogram is computed (the pow2 values are inclusive on the right side of intervals).
* MemoryInsights: Added debug log at end of analysis for the two histograms (by alloc size and by event distance) for each root heap.

diff --git a/Engine/Source/Developer/TraceServices/Private/Model/AllocationsProvider.cpp b/Engine/Source/Developer/TraceServices/Private/Model/AllocationsProvider.cpp
--- a/Engine/Source/Developer/TraceServices/Private/Model/AllocationsProvider.cpp
+++ b/Engine/Source/Developer/TraceServices/Private/Model/AllocationsProvider.cpp
@@ -1881,6 +1881,8 @@ void FAllocationsProvider::EditSwapOp(uint32 ThreadId, double Time, uint64 Unmas
 			Allocation.Tag = 0;
 			Allocation.RootHeap = static_cast<uint8>(RootHeapId);
 			Allocation.Flags = EMemoryTraceHeapAllocationFlags::Swap;
+
+			RootHeap.UpdateHistogramByAllocSize(CompressedPageSize);
 			break;
 		}
 		case EMemoryTraceSwapOperation::PageIn:
@@ -1915,11 +1917,16 @@ void FAllocationsProvider::EditSwapOp(uint32 ThreadId, double Time, uint64 Unmas
 				check(RootHeap.EventIndex > AllocationPtr->StartEventIndex);
 				AllocationPtr->EndEventIndex = RootHeap.EventIndex;
 				AllocationPtr->EndTime = Time;
+
 				AllocationPtr->FreeThreadId = static_cast<uint16>(ThreadId);
 				check(uint32(AllocationPtr->FreeThreadId) == ThreadId);
+
 				AllocationPtr->FreeCallstackId = CallstackId;
 
 				RootHeap.SbTree->AddAlloc(AllocationPtr); // SbTree takes ownership of AllocationPtr
+
+				const uint32 EventDistance = AllocationPtr->EndEventIndex - AllocationPtr->StartEventIndex;
+				RootHeap.UpdateHistogramByEventDistance(EventDistance);
 			}
 
 			break;
@@ -2401,44 +2408,290 @@ void FAllocationsProvider::EditUnmarkAllocationAsHeap(uint32 ThreadId, double Ti
 
 void FAllocationsProvider::FRootHeap::UpdateHistogramByAllocSize(uint64 Size)
 {
-	if (Size > MaxAllocSize)
+	// HistogramIndex : Value Range
+	// 0 : [0]
+	// 1 : [1]
+	// 2 : [2]
+	// 3 : [3 .. 4]
+	// 4 : [5 .. 8]
+	// 5 : [9 .. 16]
+	// ...
+	// i : [2^(i-2)+1 .. 2^(i-1)], i > 2
+	// ...
+	// 64 : [2^62+1 .. 2^63]
+	// 65 : [2^63+1 .. 2^64-1]
+
+	if (Size == 0)
 	{
-		MaxAllocSize = Size;
+		++AllocSizeHistogramPow2[0];
 	}
+	else
+	{
+		uint32 HistogramIndexPow2 = 65 - static_cast<uint32>(FMath::CountLeadingZeros64(Size - 1));
+		++AllocSizeHistogramPow2[HistogramIndexPow2];
+
+		if (Size > MaxAllocSize)
+		{
+			MaxAllocSize = Size;
+		}
+	}
+}
+
+////////////////////////////////////////////////////////////////////////////////////////////////////
 
+void FAllocationsProvider::FRootHeap::UpdateHistogramByEventDistance(uint32 EventDistance)
+{
 	// HistogramIndex : Value Range
 	// 0 : [0]
 	// 1 : [1]
 	// 2 : [2 .. 3]
 	// 3 : [4 .. 7]
 	// ...
-	// i : [2^(i-1) .. 2^i-1], i > 0
+	// i : [2^(i-1) .. 2^i-1], i > 1
 	// ...
-	// 64 : [2^63 .. 2^64-1]
-	uint32 HistogramIndexPow2 = 64 - static_cast<uint32>(FMath::CountLeadingZeros64(Size));
-	++AllocSizeHistogramPow2[HistogramIndexPow2];
+	// 32 : [2^31 .. 2^32-1]
+
+	uint32 HistogramIndexPow2 = 32 - FMath::CountLeadingZeros(EventDistance);
+	++EventDistanceHistogramPow2[HistogramIndexPow2];
+
+	if (EventDistance > MaxEventDistance)
+	{
+		MaxEventDistance = EventDistance;
+	}
 }
 
 ////////////////////////////////////////////////////////////////////////////////////////////////////
 
-void FAllocationsProvider::FRootHeap::UpdateHistogramByEventDistance(uint32 EventDistance)
+void FAllocationsProvider::FRootHeap::LogStats()
 {
-	if (EventDistance > MaxEventDistance)
+	if (EventIndex == 0)
 	{
-		MaxEventDistance = EventDistance;
+		return;
+	}
+
+	check(HeapSpec != nullptr);
+
+	UE_LOG(LogTraceServices, Log, TEXT("[MemAlloc] RootHeap %u (%s) : %u events, %u live allocs, max alloc size = %llu bytes, max event distance = %u"),
+		uint32(HeapSpec->Id), HeapSpec->Name, EventIndex, LiveAllocs->Num(), MaxAllocSize, MaxEventDistance);
+
+	LogAllocSizeHistogram();
+	LogEventDistanceHistogram();
+}
+
+////////////////////////////////////////////////////////////////////////////////////////////////////
+
+static FString BytesToString(uint64 Value)
+{
+	if (Value == 0)
+	{
+		return FString("0");
+	}
+	if (Value == ~0ull)
+	{
+		return FString("MAX64");
+	}
+
+	constexpr uint64 K = 1024ull;
+	constexpr uint64 M = K * K;
+	constexpr uint64 G = M * K;
+	constexpr uint64 T = G * K;
+
+	uint64 Unit = 0;
+	const TCHAR* UnitStr = nullptr;
+	if (Value < K)
+	{
+		return FString::Printf(TEXT("%llu B"), Value);
+	}
+	else if (Value < M)
+	{
+		Unit = K;
+		UnitStr = TEXT("KiB");
+	}
+	else if (Value < G)
+	{
+		Unit = M;
+		UnitStr = TEXT("MiB");
+	}
+	else if (Value < T)
+	{
+		Unit = G;
+		UnitStr = TEXT("GiB");
+	}
+	else
+	{
+		Unit = T;
+		UnitStr = TEXT("TiB");
+	}
+
+	if (Value % Unit == 0)
+	{
+		return FString::Printf(TEXT("%llu %s"), Value / Unit, UnitStr);
+	}
+	else
+	{
+		return FString::Printf(TEXT("%.2f %s"), (double)Value / (double)Unit, UnitStr);
+	}
+}
+
+////////////////////////////////////////////////////////////////////////////////////////////////////
+
+static FString CountToString(uint64 Count)
+{
+	FString Out = FString::Printf(TEXT("%llu"), Count);
+	for (int32 Index = Out.Len() - 3; Index > 0; Index -= 3)
+	{
+		Out.InsertAt(Index, TEXT(','));
+	}
+	return Out;
+}
+
+////////////////////////////////////////////////////////////////////////////////////////////////////
+
+void FAllocationsProvider::FRootHeap::LogAllocSizeHistogram()
+{
+	uint32 StartIndex = 0;
+	while (StartIndex < NumAllocSizeHistogramPow2 && AllocSizeHistogramPow2[StartIndex] == 0)
+	{
+		++StartIndex;
+	}
+
+	uint32 EndIndex = NumAllocSizeHistogramPow2;
+	while (EndIndex > 0 && AllocSizeHistogramPow2[EndIndex - 1] == 0)
+	{
+		--EndIndex;
 	}
 
+	if (StartIndex >= EndIndex)
+	{
+		return;
+	}
+
+	check(HeapSpec != nullptr);
+	UE_LOG(LogTraceServices, Log, TEXT("[MemAlloc] Alloc Size Histogram (%s):"), HeapSpec->Name);
+
+	// HistogramIndex : Value Range
+	// 0 : [0 bytes]
+	// 1 : [1 byte]
+	// 2 : [2 bytes]
+	// 3 : (2 .. 4]
+	// 4 : (4 .. 8]
+	// 5 : (8 .. 16]
+	// ...
+	// i : (2^(i-2) .. 2^(i-1)], i > 2
+	// ...
+	// 64 : (2^62 .. 2^63]
+	// 65 : (2^63 .. 2^64-1]
+
+	uint64 Sum = 0;
+	for (uint32 Index = StartIndex; Index < EndIndex; ++Index)
+	{
+		Sum += AllocSizeHistogramPow2[Index];
+	}
+	uint64 PartialSum = 0;
+	double PercentFactor = 100.0 / (double)Sum;
+
+	UE_LOG(LogTraceServices, Log, TEXT("[MemAlloc] -----------------------------------------------------------"));
+	for (uint32 Index = StartIndex; Index < EndIndex; ++Index)
+	{
+		const uint64 Count = AllocSizeHistogramPow2[Index];
+		PartialSum += Count;
+		const double Percent = (double)Count * PercentFactor;
+		const double CumulativePercent = (double)PartialSum * PercentFactor;
+		if (Index > 2)
+		{
+			// i : (2^(i-2) .. 2^(i-1)], i > 2
+			const uint64 Min = (1ull << (Index - 2));
+			const uint64 Max = (Index == NumAllocSizeHistogramPow2 - 1) ? ~0ull : (1ull << (Index - 1));
+			UE_LOG(LogTraceServices, Log, TEXT("[MemAlloc] (%8s .. %-8s] %11s %4.01f%% %11s %5.01f%%"),
+				*BytesToString(Min), *BytesToString(Max), *CountToString(Count), Percent, *CountToString(PartialSum), CumulativePercent);
+		}
+		else if (Index == 1)
+		{
+			UE_LOG(LogTraceServices, Log, TEXT("[MemAlloc] {     1 byte         } %11s %4.01f%% %11s %5.01f%%"),
+				*CountToString(Count), Percent, *CountToString(PartialSum), CumulativePercent);
+		}
+		else // if (Index == 0 || Index == 2)
+		{
+			UE_LOG(LogTraceServices, Log, TEXT("[MemAlloc] {     %u bytes        } %11s %4.01f%% %11s %5.01f%%"),
+				Index, *CountToString(Count), Percent, *CountToString(PartialSum), CumulativePercent);
+		}
+	}
+	UE_LOG(LogTraceServices, Log, TEXT("[MemAlloc] -----------------------------------------------------------"));
+}
+
+////////////////////////////////////////////////////////////////////////////////////////////////////
+
+void FAllocationsProvider::FRootHeap::LogEventDistanceHistogram()
+{
+	uint32 StartIndex = 0;
+	while (StartIndex < NumEventDistanceHistogramPow2 && EventDistanceHistogramPow2[StartIndex] == 0)
+	{
+		++StartIndex;
+	}
+
+	uint32 EndIndex = NumEventDistanceHistogramPow2;
+	while (EndIndex > 0 && EventDistanceHistogramPow2[EndIndex - 1] == 0)
+	{
+		--EndIndex;
+	}
+
+	if (StartIndex >= EndIndex)
+	{
+		return;
+	}
+
+	check(HeapSpec != nullptr);
+	UE_LOG(LogTraceServices, Log, TEXT("[MemAlloc] Event Distance Histogram (%s):"), HeapSpec->Name);
+
 	// HistogramIndex : Value Range
 	// 0 : [0]
 	// 1 : [1]
 	// 2 : [2 .. 3]
 	// 3 : [4 .. 7]
 	// ...
-	// i : [2^(i-1) .. 2^i-1], i > 0
+	// i : [2^(i-1) .. 2^i-1], i > 1
 	// ...
 	// 32 : [2^31 .. 2^32-1]
-	uint32 HistogramIndexPow2 = 32 - FMath::CountLeadingZeros(EventDistance);
-	++EventDistanceHistogramPow2[HistogramIndexPow2];
+
+	uint64 Sum = LiveAllocs->Num();
+	for (uint32 Index = StartIndex; Index < EndIndex; ++Index)
+	{
+		Sum += EventDistanceHistogramPow2[Index];
+	}
+	uint64 PartialSum = 0;
+	double PercentFactor = 100.0 / (double)Sum;
+
+	UE_LOG(LogTraceServices, Log, TEXT("[MemAlloc] -----------------------------------------------------------------"));
+	for (uint32 Index = StartIndex; Index < EndIndex; ++Index)
+	{
+		const uint64 Count = EventDistanceHistogramPow2[Index];
+		PartialSum += Count;
+		const double Percent = (double)Count * PercentFactor;
+		const double CumulativePercent = (double)PartialSum * PercentFactor;
+		if (Index > 1)
+		{
+			// i : [2^(i-1) .. 2^i-1], i > 1
+			const uint64 Min = (1ull << (Index - 1));
+			const uint64 Max = (1ull << Index) - 1;
+			UE_LOG(LogTraceServices, Log, TEXT("[MemAlloc] [%11s .. %-11s] %11s %4.01f%% %11s %5.01f%%"),
+				*CountToString(Min), *CountToString(Max), *CountToString(Count), Percent, *CountToString(PartialSum), CumulativePercent);
+		}
+		else // if (Index == 0 || Index == 1)
+		{
+			UE_LOG(LogTraceServices, Log, TEXT("[MemAlloc] {          %u               } %11s %4.01f%% %11s %5.01f%%"),
+				Index, *CountToString(Count), Percent, *CountToString(PartialSum), CumulativePercent);
+		}
+	}
+	if (LiveAllocs->Num() > 0)
+	{
+		const uint64 Count = LiveAllocs->Num();
+		PartialSum += Count;
+		const double Percent = (double)Count * PercentFactor;
+		const double CumulativePercent = (double)PartialSum * PercentFactor;
+		UE_LOG(LogTraceServices, Log, TEXT("[MemAlloc] {      \u221E (live allocs)     } %11s %4.01f%% %11s %5.01f%%"),
+			*CountToString(Count), Percent, *CountToString(PartialSum), CumulativePercent);
+	}
+	UE_LOG(LogTraceServices, Log, TEXT("[MemAlloc] -----------------------------------------------------------------"));
 }
 
 ////////////////////////////////////////////////////////////////////////////////////////////////////
@@ -2580,46 +2833,6 @@ void FAllocationsProvider::EditOnAnalysisCompleted(double Time)
 		return;
 	}
 
-#if 0
-	const bool bResetTimelineAtEnd = false;
-
-	// Add all live allocs to SbTree (with infinite end time).
-	uint64 LiveAllocsTotalSize = 0;
-	LiveAllocs.Enumerate([](const FAllocationItem& Alloc)
-	{
-		FAllocationItem* AllocationPtr = const_cast<FAllocationItem*>(&Alloc);
-
-		LiveAllocsTotalSize += AllocationPtr->GetSize();
-
-		// Assign same event index to all live allocs at the end of the session.
-		AllocationPtr->EndEventIndex = EventIndex;
-
-		SbTree->AddAlloc(AllocationPtr);
-
-		uint32 EventDistance = AllocationPtr->EndEventIndex - AllocationPtr->StartEventIndex;
-		UpdateHistogramByEventDistance(EventDistance);
-	});
-	//TODO: LiveAllocs.RemoveAll();
-	check(TotalAllocatedMemory == LiveAllocsTotalSize);
-
-	if (bResetTimelineAtEnd)
-	{
-		AdvanceTimelines(Time + 10 * DefaultTimelineSampleGranularity);
-
-		const uint32 LiveAllocsTotalCount = TotalLiveAllocations;
-		LiveAllocs.Empty();
-
-		// Update stats for the last timeline sample (reset to zero).
-		TotalAllocatedMemory = 0;
-		TotalSwapMemory = 0;
-		SampleMinTotalAllocatedMemory = 0;
-		SampleMinLiveAllocations = 0;
-		SampleMinSwapMemory = 0;
-		SampleMinCompressedSwapMemory = 0;
-		SampleFreeEvents += LiveAllocsTotalCount;
-	}
-#endif
-
 	// Flush the last cached timeline sample.
 	AdvanceTimelines(std::numeric_limits<double>::infinity());
 
@@ -2680,9 +2893,11 @@ void FAllocationsProvider::EditOnAnalysisCompleted(double Time)
 		FRootHeap* RootHeap = RootHeaps[RootHeapIdx];
 		if (RootHeap)
 		{
+			RootHeap->LogStats();
 			TotalEventCount += RootHeap->EventIndex;
 		}
 	}
+
 	uint32 NumHeapSpecs = 0;
 	for (FHeapSpec* HeapSpec : HeapSpecs)
 	{
@@ -2691,7 +2906,8 @@ void FAllocationsProvider::EditOnAnalysisCompleted(double Time)
 			++NumHeapSpecs;
 		}
 	}
-	UE_LOG(LogTraceServices, Log, TEXT("[MemAlloc] Indexed: %llu events"), TotalEventCount);
+
+	UE_LOG(LogTraceServices, Log, TEXT("[MemAlloc] Total Indexed: %llu events"), TotalEventCount);
 	UE_LOG(LogTraceServices, Log, TEXT("[MemAlloc] Traced: %llu alloc events, %llu free events, %llu heap events, %llu swap op events"),
 		AllocEvents, FreeEvents, HeapEvents, SwapOpEvents);
 	UE_LOG(LogTraceServices, Log, TEXT("[MemAlloc] Analysis completed (%llu allocs, %llu frees, %llu heaps, %u heap specs)."),
diff --git a/Engine/Source/Developer/TraceServices/Private/Model/AllocationsProvider.h b/Engine/Source/Developer/TraceServices/Private/Model/AllocationsProvider.h
--- a/Engine/Source/Developer/TraceServices/Private/Model/AllocationsProvider.h
+++ b/Engine/Source/Developer/TraceServices/Private/Model/AllocationsProvider.h
@@ -297,6 +297,10 @@ private:
 		void UpdateHistogramByAllocSize(uint64 Size);
 		void UpdateHistogramByEventDistance(uint32 EventDistance);
 
+		void LogStats();
+		void LogAllocSizeHistogram();
+		void LogEventDistanceHistogram();
+
 		FHeapSpec* HeapSpec = nullptr;
 
 		FLiveAllocCollection* LiveAllocs = nullptr;
@@ -305,11 +309,13 @@ private:
 
 		uint32 EventIndex = 0;
 
+		static constexpr uint32 NumAllocSizeHistogramPow2 = 66;
+		uint64 AllocSizeHistogramPow2[NumAllocSizeHistogramPow2] = { 0 };
 		uint64 MaxAllocSize = 0;
-		uint64 AllocSizeHistogramPow2[65] = { 0 };
 
+		static constexpr uint32 NumEventDistanceHistogramPow2 = 33;
+		uint32 EventDistanceHistogramPow2[NumEventDistanceHistogramPow2] = { 0 };
 		uint32 MaxEventDistance = 0;
-		uint32 EventDistanceHistogramPow2[33] = { 0 };
 	};
 
 public:
