[Insights] Implemented workaround for old trace files with invalid cpu scope timestamps.

#rnx
#Synced-CL 48281143

diff --git a/Engine/Source/Developer/TraceServices/Private/Analyzers/CpuProfilerTraceAnalysis.cpp b/Engine/Source/Developer/TraceServices/Private/Analyzers/CpuProfilerTraceAnalysis.cpp
--- a/Engine/Source/Developer/TraceServices/Private/Analyzers/CpuProfilerTraceAnalysis.cpp
+++ b/Engine/Source/Developer/TraceServices/Private/Analyzers/CpuProfilerTraceAnalysis.cpp
@@ -768,11 +768,21 @@ void FCpuProfilerAnalyzer::OnCpuScopeEnter(const FOnEventContext& Context)
 	uint64 Cycle = Context.EventTime.AsCycle64();
 	double Time = Context.EventTime.AsSeconds();
 
-	check(ThreadState.LastCycle <= Cycle);
-	check(ThreadState.LastPendingEventTime <= Time);
-	ThreadState.LastPendingEventTime = Time;
-
-	ThreadState.PendingEvents.Add({ Cycle, Time, TimerId });
+	if (ensure(Cycle >= ThreadState.LastCycle))
+	{
+		check(ThreadState.LastPendingEventTime <= Time);
+		ThreadState.LastPendingEventTime = Time;
+		ThreadState.PendingEvents.Add({ Cycle, Time, TimerId });
+	}
+	else
+	{
+		// Time needs to increase monotonically.
+		// We are not allowing pending events (with metadata) older than regular CPU timing events.
+		// When this happens we further ignore all pending events on this thread.
+		// The issue can occur in late connect trace sessions with trace protocol <= 6 (i.e. the scoped events have relative timestamps).
+		ThreadState.bShouldIgnorePendingEvents = true;
+		UE_LOG(LogTraceServices, Error, TEXT("[CpuProfiler] Detected scope enter event non-monotonically increasing timestamp. Further CPU timing events with metadata are ignored on thread %u."), ThreadState.ThreadId);
+	}
 }
 
 ////////////////////////////////////////////////////////////////////////////////////////////////////
@@ -795,11 +805,21 @@ void FCpuProfilerAnalyzer::OnCpuScopeLeave(const FOnEventContext& Context)
 	uint64 Cycle = Context.EventTime.AsCycle64();
 	double Time = Context.EventTime.AsSeconds();
 
-	check(ThreadState.LastCycle <= Cycle);
-	check(ThreadState.LastPendingEventTime <= Time);
-	ThreadState.LastPendingEventTime = Time;
-
-	ThreadState.PendingEvents.Add({ ~Cycle, Time, 0 });
+	if (ensure(Cycle >= ThreadState.LastCycle))
+	{
+		check(ThreadState.LastPendingEventTime <= Time);
+		ThreadState.LastPendingEventTime = Time;
+		ThreadState.PendingEvents.Add({ ~Cycle, Time, 0 });
+	}
+	else
+	{
+		// Time needs to increase monotonically.
+		// We are not allowing pending events (with metadata) older than regular CPU timing events.
+		// When this happens we further ignore all pending events on this thread.
+		// The issue can occur in late connect trace sessions with trace protocol <= 6 (i.e. the scoped events have relative timestamps).
+		ThreadState.bShouldIgnorePendingEvents = true;
+		UE_LOG(LogTraceServices, Error, TEXT("[CpuProfiler] Detected scope leave event with non-monotonically increasing timestamp. Further CPU timing events with metadata are ignored on thread %u."), ThreadState.ThreadId);
+	}
 }
 
 ////////////////////////////////////////////////////////////////////////////////////////////////////
