[Insights]
- TraceServices: Added trace trimmer analyzer and Trim functionality (TraceTrimmer.h).
- TraceInsights / Timing View: Added "Trim and Save As..." in context menu. Trims the current trace using the selected time range and saves it as new trace file.

#Synced-CL 49988361

diff --git a/Engine/Source/Developer/TraceAnalysis/Public/Trace/TraceWriter.h b/Engine/Source/Developer/TraceAnalysis/Public/Trace/TraceWriter.h
--- a/Engine/Source/Developer/TraceAnalysis/Public/Trace/TraceWriter.h
+++ b/Engine/Source/Developer/TraceAnalysis/Public/Trace/TraceWriter.h
@@ -45,7 +45,7 @@ public:
 	[[nodiscard]] uint32 GetId() const { return Id; }
 	[[nodiscard]] const FAnsiString& GetName() const { return Name; }
 	[[nodiscard]] uint32 GetSystemId() const { return SystemId; }
-	[[nodiscard]] uint32 GetSortHint() const { return SortHint; }
+	[[nodiscard]] int32 GetSortHint() const { return SortHint; }
 
 	static constexpr uint32 InvalidThreadId = ~0u;
 
@@ -552,9 +552,9 @@ private:
 	void SetErrorInvalidFieldType(uint32 Index);
 	void SetErrorInvalidFieldData(uint32 Index);
 
-	bool CheckValidPodField(uint32 Index, ETraceWriterFieldType FieldType, uint32 Size);
-	bool CheckValidExactFieldType(uint32 Index, ETraceWriterFieldType FieldType);
-	void SetField(uint32 Index, const void* Data, uint32 DataSize);
+	UE_API bool CheckValidPodField(uint32 Index, ETraceWriterFieldType FieldType, uint32 Size);
+	UE_API bool CheckValidExactFieldType(uint32 Index, ETraceWriterFieldType FieldType);
+	UE_API void SetField(uint32 Index, const void* Data, uint32 DataSize);
 
 private:
 	FTraceWriter& Writer;
diff --git a/Engine/Source/Developer/TraceInsights/Private/Insights/TimingProfiler/TimingProfilerCommands.cpp b/Engine/Source/Developer/TraceInsights/Private/Insights/TimingProfiler/TimingProfilerCommands.cpp
--- a/Engine/Source/Developer/TraceInsights/Private/Insights/TimingProfiler/TimingProfilerCommands.cpp
+++ b/Engine/Source/Developer/TraceInsights/Private/Insights/TimingProfiler/TimingProfilerCommands.cpp
@@ -113,6 +113,12 @@ void FTimingViewCommands::RegisterCommands()
 		"Quick find or filter events in the timing view.",
 		EUserInterfaceActionType::Button,
 		FInputChord(EModifierKey::Control, EKeys::F));
+
+	UI_COMMAND(TrimAndSaveAs,
+		"Trim and Save As...",
+		"Trim the current trace using the selected time range and save it as new trace file.",
+		EUserInterfaceActionType::Button,
+		FInputChord());
 }
 UE_ENABLE_OPTIMIZATION_SHIP
 
@@ -160,7 +166,7 @@ IMPLEMENT_TOGGLE_COMMAND(ToggleCallersTreeViewVisibility, IsCallersTreeViewVisib
 IMPLEMENT_TOGGLE_COMMAND(ToggleCalleesTreeViewVisibility, IsCalleesTreeViewVisible, ShowHideCalleesTreeView)
 IMPLEMENT_TOGGLE_COMMAND(ToggleStatsCountersViewVisibility, IsStatsCountersViewVisible, ShowHideStatsCountersView)
 IMPLEMENT_TOGGLE_COMMAND(ToggleLogViewVisibility, IsLogViewVisible, ShowHideLogView)
-IMPLEMENT_TOGGLE_COMMAND(ToggleModulesViewVisibility, IsModulesViewVisible, ShowHideModulesView)	
+IMPLEMENT_TOGGLE_COMMAND(ToggleModulesViewVisibility, IsModulesViewVisible, ShowHideModulesView)
 
 #undef IMPLEMENT_TOGGLE_COMMAND
 
diff --git a/Engine/Source/Developer/TraceInsights/Private/Insights/TimingProfiler/TimingProfilerCommands.h b/Engine/Source/Developer/TraceInsights/Private/Insights/TimingProfiler/TimingProfilerCommands.h
--- a/Engine/Source/Developer/TraceInsights/Private/Insights/TimingProfiler/TimingProfilerCommands.h
+++ b/Engine/Source/Developer/TraceInsights/Private/Insights/TimingProfiler/TimingProfilerCommands.h
@@ -80,6 +80,9 @@ public:
 
 	/** Opens the Quick Find widget. */
 	TSharedPtr<FUICommandInfo> QuickFind;
+
+	/** Trim the current trace using the selected time range and save it as new trace file. */
+	TSharedPtr<FUICommandInfo> TrimAndSaveAs;
 };
 
 ////////////////////////////////////////////////////////////////////////////////////////////////////
diff --git a/Engine/Source/Developer/TraceInsights/Private/Insights/Widgets/STimingView.cpp b/Engine/Source/Developer/TraceInsights/Private/Insights/Widgets/STimingView.cpp
--- a/Engine/Source/Developer/TraceInsights/Private/Insights/Widgets/STimingView.cpp
+++ b/Engine/Source/Developer/TraceInsights/Private/Insights/Widgets/STimingView.cpp
@@ -2,6 +2,7 @@
 
 #include "STimingView.h"
 
+#include "DesktopPlatformModule.h"
 #include "Containers/ArrayBuilder.h"
 #include "Containers/MapBuilder.h"
 #include "Features/IModularFeatures.h"
@@ -35,10 +36,16 @@
 #include "Widgets/SOverlay.h"
 #include "Widgets/Text/STextBlock.h"
 
+// TraceAnalysis
+#include "Trace/DataStream.h"
+#include "Trace/OutDataStream.h"
+#include "Trace/StoreClient.h"
+
 // TraceServices
 #include "Common/ProviderLock.h"
 #include "TraceServices/Model/LoadTimeProfiler.h"
 #include "TraceServices/Model/Regions.h"
+#include "TraceServices/TraceTrimmer.h"
 
 // TraceInsightsCore
 #include "InsightsCore/Common/PaintUtils.h"
@@ -3218,6 +3225,12 @@ void STimingView::ShowContextMenu(const FPointerEvent& MouseEvent)
 			);
 		}
 
+		MenuBuilder.AddMenuEntry(Commands.TrimAndSaveAs,
+			FName("TrimAndSaveAs"),
+			TAttribute<FText>(),
+			TAttribute<FText>(),
+			FSlateIcon(FInsightsStyle::GetStyleSetName(), "Icons.Save"));
+
 		bHasAnyActions = true;
 	}
 	MenuBuilder.EndSection();
@@ -3462,6 +3475,11 @@ void STimingView::BindCommands()
 		FExecuteAction::CreateSP(this, &STimingView::QuickFind_Execute),
 		FCanExecuteAction::CreateSP(this, &STimingView::QuickFind_CanExecute));
 
+	CommandList->MapAction(
+		Commands.TrimAndSaveAs,
+		FExecuteAction::CreateSP(this, &STimingView::TrimAndSaveAs_Execute),
+		FCanExecuteAction::CreateSP(this, &STimingView::TrimAndSaveAs_CanExecute));
+
 	FrameSharedState->BindCommands();
 	ThreadTimingSharedState->BindCommands();
 	LoadingSharedState->BindCommands();
@@ -5573,6 +5591,155 @@ void STimingView::CloseQuickFindTab()
 
 ////////////////////////////////////////////////////////////////////////////////////////////////////
 
+bool STimingView::TrimAndSaveAs_CanExecute() const
+{
+	if (SelectionStartTime >= SelectionEndTime)
+	{
+		return false;
+	}
+
+	TSharedPtr<FInsightsManager> InsightsManager = FInsightsManager::Get();
+
+	ETraceStreamType TraceStreamType = InsightsManager->GetTraceStreamType();
+	if (TraceStreamType != ETraceStreamType::TraceFile &&
+		TraceStreamType != ETraceStreamType::TraceStore)
+	{
+		return false;
+	}
+
+	TSharedPtr<const TraceServices::IAnalysisSession> Session = InsightsManager->GetSession();
+	if (!Session)
+	{
+		return false;
+	}
+
+	return true;
+}
+
+////////////////////////////////////////////////////////////////////////////////////////////////////
+
+void STimingView::TrimAndSaveAs_Execute()
+{
+	if (SelectionStartTime >= SelectionEndTime)
+	{
+		return;
+	}
+
+	TraceServices::FTrimParameters TrimParams;
+	TrimParams.StartTime = SelectionStartTime;
+	TrimParams.EndTime = SelectionEndTime;
+
+	//////////////////////////////////////////////////
+
+	TSharedPtr<FInsightsManager> InsightsManager = FInsightsManager::Get();
+
+	TSharedPtr<const TraceServices::IAnalysisSession> Session = InsightsManager->GetSession();
+	if (!Session)
+	{
+		UE_LOG(LogTimingProfiler, Error, TEXT("[TraceTrimmer] Invalid session!"));
+		return;
+	}
+
+	//////////////////////////////////////////////////
+	// InputDataStream
+
+	TUniquePtr<UE::Trace::IInDataStream> InputDataStream;
+	FString InputTraceName;
+
+	ETraceStreamType TraceStreamType = InsightsManager->GetTraceStreamType();
+	if (TraceStreamType == ETraceStreamType::TraceFile)
+	{
+		InputTraceName = Session->GetName();
+		FPaths::NormalizeFilename(InputTraceName);
+
+		InputDataStream = MakeUnique<UE::Trace::FFileDataStream>();
+		UE::Trace::FFileDataStream* FileDataStream = (UE::Trace::FFileDataStream*)InputDataStream.Get();
+		if (!FileDataStream->Open(*InputTraceName))
+		{
+			UE_LOG(LogTimingProfiler, Error, TEXT("[TraceTrimmer] Cannot open the input trace file ('%s')!"), *InputTraceName);
+			return;
+		}
+	}
+	else if (TraceStreamType == ETraceStreamType::TraceStore)
+	{
+		InputTraceName = Session->GetName();
+		FPaths::NormalizeFilename(InputTraceName);
+
+		const uint32 TraceId = Session->GetTraceId();
+
+		UE::Trace::FStoreClient* StoreClient = InsightsManager->GetStoreClient();
+		if (StoreClient)
+		{
+			InputDataStream = StoreClient->ReadTrace(TraceId);
+		}
+
+		if (!InputDataStream)
+		{
+			UE_LOG(LogTimingProfiler, Error, TEXT("[TraceTrimmer] Cannot open the input trace stream ('%s')!"), *InputTraceName);
+			return;
+		}
+	}
+	else
+	{
+		UE_LOG(LogTimingProfiler, Error, TEXT("[TraceTrimmer] Unsupported input trace stream ('%s')!"), *InsightsManager->GetTraceName());
+		return;
+	}
+
+	//////////////////////////////////////////////////
+	// Choose the output trace file
+
+	TArray<FString> SaveFilenames;
+	bool bDialogResult = false;
+
+	if (IDesktopPlatform* DesktopPlatform = FDesktopPlatformModule::Get())
+	{
+		const FString DefaultPath = FPaths::ProjectSavedDir();
+		const FString DefaultFile = FPaths::GetBaseFilename(InputTraceName) + TEXT("_trim.utrace");
+
+		bDialogResult = DesktopPlatform->SaveFileDialog(
+			FSlateApplication::Get().FindBestParentWindowHandleForDialogs(nullptr),
+			LOCTEXT("SaveTrimmedTraceAs", "Save Trimmed Trace As").ToString(),
+			DefaultPath,
+			DefaultFile,
+			TEXT("Trace Files (*.utrace)|*.utrace|All Files (*.*)|*.*"),
+			EFileDialogFlags::None,
+			SaveFilenames
+		);
+	}
+
+	if (!bDialogResult || SaveFilenames.Num() == 0)
+	{
+		return;
+	}
+
+	const FString& OutputTraceFile = SaveFilenames[0];
+
+	//////////////////////////////////////////////////
+	// OutputDataStream
+
+	TUniquePtr<UE::Trace::FFileOutDataStream> OutputDataStream = MakeUnique<UE::Trace::FFileOutDataStream>();
+	if (!OutputDataStream->Open(*OutputTraceFile))
+	{
+		UE_LOG(LogTimingProfiler, Error, TEXT("[TraceTrimmer] Cannot write the output trace file ('%s')!"), *OutputTraceFile);
+		return;
+	}
+
+	//////////////////////////////////////////////////
+	// Trim
+
+	int32 Result = TraceServices::Trim(TrimParams, *InputDataStream, *OutputDataStream);
+	if (Result == 0)
+	{
+		UE_LOG(LogTimingProfiler, Log, TEXT("[TraceTrimmer] Success!"));
+	}
+	else
+	{
+		UE_LOG(LogTimingProfiler, Error, TEXT("[TraceTrimmer] Failed!"));
+	}
+}
+
+////////////////////////////////////////////////////////////////////////////////////////////////////
+
 void STimingView::HandleTrackVisibilityChanged()
 {
 	if (HoveredTrack.IsValid())
diff --git a/Engine/Source/Developer/TraceInsights/Private/Insights/Widgets/STimingView.h b/Engine/Source/Developer/TraceInsights/Private/Insights/Widgets/STimingView.h
--- a/Engine/Source/Developer/TraceInsights/Private/Insights/Widgets/STimingView.h
+++ b/Engine/Source/Developer/TraceInsights/Private/Insights/Widgets/STimingView.h
@@ -127,6 +127,9 @@ public:
 	bool QuickFind_CanExecute() const;
 	void QuickFind_Execute();
 
+	bool TrimAndSaveAs_CanExecute() const;
+	void TrimAndSaveAs_Execute();
+
 	bool ToggleTrackVisibility_IsChecked(uint64 InTrackId) const;
 	void ToggleTrackVisibility_Execute(uint64 InTrackId);
 
diff --git a/Engine/Source/Developer/TraceServices/Private/Analyzers/TrimAnalyzer.cpp b/Engine/Source/Developer/TraceServices/Private/Analyzers/TrimAnalyzer.cpp
new file mode 100644
--- /dev/null
+++ b/Engine/Source/Developer/TraceServices/Private/Analyzers/TrimAnalyzer.cpp
@@ -0,0 +1,1256 @@
+// Copyright Epic Games, Inc. All Rights Reserved.
+
+#include "TrimAnalyzer.h"
+
+#include "Containers/Map.h"
+#include "Containers/StringConv.h"
+#include "Containers/UnrealString.h"
+#include "HAL/PlatformTime.h"
+#include "Templates/UniquePtr.h"
+
+// TraceAnalysis
+#include "Trace/Analysis.h"
+#include "Trace/TraceWriter.h"
+
+// TraceServices
+#include "Common/Utils.h"
+#include "TraceServices/TraceTrimmer.h"
+
+////////////////////////////////////////////////////////////////////////////////////////////////////
+
+#define LogTraceTrimmerError(Format, ...) \
+	{ \
+		++NumErrors; \
+		wprintf(L"Error: " Format "\n", ##__VA_ARGS__); \
+		UE_LOG(LogTraceServices, Error, TEXT("[TraceTrimmer] " Format), ##__VA_ARGS__); \
+	}
+#define LogTraceTrimmerWarning(Format, ...) \
+	{ \
+		++NumWarnings; \
+		wprintf(L"Warning: " Format "\n", ##__VA_ARGS__); \
+		UE_LOG(LogTraceServices, Warning, TEXT("[TraceTrimmer] " Format), ##__VA_ARGS__); \
+	}
+#define LogTraceTrimmerMessage(Format, ...) \
+	{ \
+		wprintf(L"" Format "\n", ##__VA_ARGS__); \
+		UE_LOG(LogTraceServices, Log, TEXT("[TraceTrimmer] " Format), ##__VA_ARGS__); \
+	}
+
+////////////////////////////////////////////////////////////////////////////////////////////////////
+
+namespace TraceServices
+{
+
+////////////////////////////////////////////////////////////////////////////////////////////////////
+
+static uint32 GetThreadIdField(const UE::Trace::IAnalyzer::FOnEventContext& Context, const ANSICHAR* FieldName = "ThreadId")
+{
+	// Trace analysis was changed to be able to provide a suitable id. Prior to
+	// this users of Trace would send along their own thread ids. For backwards
+	// compatibility we'll bias field thread ids to avoid collision with Trace's.
+	static const uint32 Bias = 0x70000000;
+	uint32 ThreadId = Context.EventData.GetValue<uint32>(FieldName, 0);
+	ThreadId |= ThreadId ? Bias : Context.ThreadInfo.GetId();
+	return ThreadId;
+}
+
+////////////////////////////////////////////////////////////////////////////////////////////////////
+// FTrimAnalyzer
+////////////////////////////////////////////////////////////////////////////////////////////////////
+
+FTrimAnalyzer::FTrimAnalyzer(FTrimParameters& InParameters, UE::Trace::FTraceWriter& InTraceWriter)
+	: Parameters(InParameters)
+	, TraceWriter(InTraceWriter)
+{
+	// Empty IncludeFilter is equivalent with "*.*" (include all).
+	if (!Parameters.Include.Equals("*.*"))
+	{
+		IncludeFilter.Init(Parameters.Include);
+	}
+
+	ExcludeFilter.Init(Parameters.Exclude);
+}
+
+////////////////////////////////////////////////////////////////////////////////////////////////////
+
+FTrimAnalyzer::~FTrimAnalyzer()
+{
+}
+
+////////////////////////////////////////////////////////////////////////////////////////////////////
+
+void FTrimAnalyzer::FNameFilter::Init(FAnsiString Filter)
+{
+	TArray<FAnsiString> FilterPatterns;
+	Filter.ParseIntoArray(FilterPatterns, ",", true);
+
+	for (FAnsiString& Pattern : FilterPatterns)
+	{
+		TArray<FAnsiString> Names;
+		Pattern.ParseIntoArray(Names, ".", true);
+
+		if (Names.Num() == 2)
+		{
+			Patterns.Add({ Names[0], Names[1] });
+		}
+	}
+}
+
+////////////////////////////////////////////////////////////////////////////////////////////////////
+
+void FTrimAnalyzer::OnAnalysisBegin(const FOnAnalysisContext& Context)
+{
+	LogTraceTrimmerMessage("Trimming...");
+
+	AnalysisBeginTimestamp = FPlatformTime::Cycles64();
+
+	Context.InterfaceBuilder.RouteAllEvents(0, false);
+	Context.InterfaceBuilder.RouteAllEvents(1, true);
+
+	KnownEvents[(int)EKnownEvent::Invalid].Id = InvalidEventId;
+	KnownEvents[(int)EKnownEvent::Trace_NewEvent].Id = 0;
+	for (uint32 EventId = 0; EventId <= MaxKnownEventId; ++EventId)
+	{
+		KnownEventMap[EventId] = EKnownEvent::Invalid;
+	}
+	KnownEventMap[0] = EKnownEvent::Trace_NewEvent;
+
+	RegisterCallbacksForKnownEvents();
+
+	TraceWriter.Begin(UE::Trace::ETraceWriterBeginOptions::DeclareDefaultEvents);
+}
+
+////////////////////////////////////////////////////////////////////////////////////////////////////
+
+void FTrimAnalyzer::OnAnalysisEnd()
+{
+	if (bIsTraceNewTraceEventDetected)
+	{
+		TraceWriter.End();
+	}
+	else
+	{
+		LogTraceTrimmerError("Invalid trace ($Trace.NewTrace not detected)!");
+	}
+
+	const uint64 AnalysisDuration = FPlatformTime::Cycles64() - AnalysisBeginTimestamp;
+	const double Duration = double(AnalysisDuration) * FPlatformTime::GetSecondsPerCycle();
+	FString AnalysisDurationString;
+	if (Duration < 1.0)
+	{
+		AnalysisDurationString = FString::Printf(TEXT("%.06f seconds"), Duration);
+	}
+	else if (Duration < 60.0)
+	{
+		AnalysisDurationString = FString::Printf(TEXT("%.02f seconds"), Duration);
+	}
+	else if (Duration < 120.0)
+	{
+		AnalysisDurationString = FString::Printf(TEXT("%.02fs (1 minute %lli seconds)"), Duration, (int64)Duration % 60);
+	}
+	else
+	{
+		AnalysisDurationString = FString::Printf(TEXT("%.02fs (%lli minutes %lli seconds)"), Duration, (int64)Duration / 60, (int64)Duration % 60);
+	}
+	LogTraceTrimmerMessage("  done in %s", *AnalysisDurationString);
+}
+
+////////////////////////////////////////////////////////////////////////////////////////////////////
+
+void FTrimAnalyzer::OnVersion(uint32 InTransportVersion, uint32 InProtocolVersion)
+{
+	TransportVersion = InTransportVersion;
+	ProtocolVersion = InProtocolVersion;
+}
+
+////////////////////////////////////////////////////////////////////////////////////////////////////
+
+void FTrimAnalyzer::RegisterCallbacksForKnownEvents()
+{
+	//////////////////////////////////////////////////
+	// Misc.*
+
+	KnownEvents[(int)EKnownEvent::Misc_RegionBeginWithId].Callback =
+		[this](const FOnEventContext& Context)
+		{
+			const uint64 Cycle = Context.EventData.GetValue<uint64>("CycleAndId");
+			const double Time = Context.EventTime.AsSeconds(Cycle);
+			UpdatedSessionTime(Cycle, Time);
+		};
+
+	//////////////////////////////////////////////////
+	// CpuProfiler.*
+
+	KnownEvents[(int)EKnownEvent::CpuProfiler_EventBatchV3].Callback =
+		[this](const FOnEventContext& Context)
+		{
+			const FEventData& EventData = Context.EventData;
+			TArrayView<const uint8> Data = EventData.GetArrayView<uint8>("Data");
+			if (Data.GetData() != nullptr)
+			{
+				ProcessCpuEventBatchV2V3(Context, Data.GetData(), Data.Num(), 3);
+			}
+		};
+
+	KnownEvents[(int)EKnownEvent::CpuProfiler_EventBatchV2].Callback =
+		[this](const FOnEventContext& Context)
+		{
+			const FEventData& EventData = Context.EventData;
+			TArrayView<const uint8> Data = EventData.GetArrayView<uint8>("Data");
+			if (Data.GetData() != nullptr)
+			{
+				ProcessCpuEventBatchV2V3(Context, Data.GetData(), Data.Num(), 2);
+			}
+		};
+
+	KnownEvents[(int)EKnownEvent::CpuProfiler_EventBatch].Callback =
+	KnownEvents[(int)EKnownEvent::CpuProfiler_EndCapture].Callback =
+		[this](const FOnEventContext& Context)
+		{
+			const FEventData& EventData = Context.EventData;
+			TArrayView<const uint8> Data = EventData.GetArrayView<uint8>("Data");
+			if (Data.GetData() != nullptr)
+			{
+				ProcessCpuEventBatchV1(Context, Data.GetData(), Data.Num());
+			}
+		};
+
+	//////////////////////////////////////////////////
+	// GpuProfiler.*
+
+	KnownEvents[(int)EKnownEvent::GpuProfiler_EventBeginBreadcrumb].Callback =
+	KnownEvents[(int)EKnownEvent::GpuProfiler_EventBeginWork].Callback =
+		
+		[this](const FOnEventContext& Context)
+		{
+			const uint64 Cycle = Context.EventData.GetValue<uint64>("GPUTimestampTOP");
+			const double Time = Context.EventTime.AsSeconds(Cycle);
+			UpdatedSessionTime(Cycle, Time);
+		};
+	KnownEvents[(int)EKnownEvent::GpuProfiler_EventEndBreadcrumb].Callback =
+	KnownEvents[(int)EKnownEvent::GpuProfiler_EventEndWork].Callback =
+		[this](const FOnEventContext& Context)
+		{
+			const uint64 Cycle = Context.EventData.GetValue<uint64>("GPUTimestampBOP");
+			const double Time = Context.EventTime.AsSeconds(Cycle);
+			UpdatedSessionTime(Cycle, Time);
+		};
+	KnownEvents[(int)EKnownEvent::GpuProfiler_EventWait].Callback =
+		[this](const FOnEventContext& Context)
+		{
+			const uint64 Cycle = Context.EventData.GetValue<uint64>("EndTime");
+			const double Time = Context.EventTime.AsSeconds(Cycle);
+			UpdatedSessionTime(Cycle, Time);
+		};
+	KnownEvents[(int)EKnownEvent::GpuProfiler_SignalFence].Callback =
+	KnownEvents[(int)EKnownEvent::GpuProfiler_WaitFence].Callback =
+		[this](const FOnEventContext& Context)
+		{
+			const uint64 Cycle = Context.EventData.GetValue<uint64>("CPUTimestamp");
+			const double Time = Context.EventTime.AsSeconds(Cycle);
+			UpdatedSessionTime(Cycle, Time);
+		};
+
+	//////////////////////////////////////////////////
+	// Known events with a "Cycle" field
+
+	KnownEvents[(int)EKnownEvent::Misc_Bookmark].Callback =
+	KnownEvents[(int)EKnownEvent::Misc_ScreenshotHeader].Callback =
+	KnownEvents[(int)EKnownEvent::Misc_RegionBegin].Callback =
+	KnownEvents[(int)EKnownEvent::Misc_RegionEnd].Callback =
+	KnownEvents[(int)EKnownEvent::Misc_RegionEndWithId].Callback =
+	KnownEvents[(int)EKnownEvent::Misc_BeginFrame].Callback =
+	KnownEvents[(int)EKnownEvent::Misc_EndFrame].Callback =
+	KnownEvents[(int)EKnownEvent::Logging_LogMessage].Callback =
+	KnownEvents[(int)EKnownEvent::PlatformFile_BeginOpen].Callback =
+	KnownEvents[(int)EKnownEvent::PlatformFile_EndOpen].Callback =
+	KnownEvents[(int)EKnownEvent::PlatformFile_BeginRead].Callback =
+	KnownEvents[(int)EKnownEvent::PlatformFile_EndRead].Callback =
+	KnownEvents[(int)EKnownEvent::PlatformFile_BeginClose].Callback =
+	KnownEvents[(int)EKnownEvent::PlatformFile_EndClose].Callback =
+	KnownEvents[(int)EKnownEvent::CpuProfiler_EndThread].Callback =
+	KnownEvents[(int)EKnownEvent::Memory_Marker].Callback =
+	KnownEvents[(int)EKnownEvent::LLM_TagValue].Callback =
+	KnownEvents[(int)EKnownEvent::SlateTrace_AddWidget].Callback =
+		[this](const FOnEventContext& Context)
+		{
+			const uint64 Cycle = Context.EventData.GetValue<uint64>("Cycle");
+			const double Time = Context.EventTime.AsSeconds(Cycle);
+			UpdatedSessionTime(Cycle, Time);
+		};
+}
+
+///////////////////////////////////////////////////////////////////////////////////////////////////
+
+void FTrimAnalyzer::ProcessKnownEvent(uint32 EventId, const ANSICHAR* LoggerName, const ANSICHAR* EventName)
+{
+	#define BEGIN_KNOWN_LOGGER(_Name) \
+		if (FCStringAnsi::Strcmp(LoggerName, _Name) == 0) \
+		{
+
+	#define KNOWN_EVENT(_Name,_Id) \
+			if (FCStringAnsi::Strcmp(EventName, _Name) == 0) \
+			{ \
+				KnownEvents[(int)EKnownEvent::_Id].Id = EventId; \
+				KnownEventMap[EventId] = EKnownEvent::_Id; \
+				break; \
+			}
+
+	#define END_KNOWN_LOGGER() \
+			break; \
+		}
+
+	while (true)
+	{
+		BEGIN_KNOWN_LOGGER("$Trace")
+			KNOWN_EVENT("NewEvent", Trace_NewEvent)
+			//KNOWN_EVENT("NewTrace", Trace_NewTrace)
+			//KNOWN_EVENT("ThreadTiming", Trace_ThreadTiming)
+			//KNOWN_EVENT("ThreadInfo", Trace_ThreadInfo)
+			//KNOWN_EVENT("ThreadGroupBegin", Trace_ThreadGroupBegin)
+			//KNOWN_EVENT("ThreadGroupEnd", Trace_ThreadGroupEnd)
+			//KNOWN_EVENT("ChannelAnnounce", Trace_ChannelAnnounce_Old) // deprecated
+		END_KNOWN_LOGGER()
+
+		//BEGIN_KNOWN_LOGGER("Trace")
+			//KNOWN_EVENT("ChannelAnnounce", Trace_ChannelAnnounce)
+			//KNOWN_EVENT("ChannelAnnounce", Trace_ChannelToggle)
+		//END_KNOWN_LOGGER()
+
+		//BEGIN_KNOWN_LOGGER("Diagnostics")
+			//KNOWN_EVENT("Session", Diagnostics_Session) // deprecated
+			//KNOWN_EVENT("Session2", Diagnostics_Session2)
+		//END_KNOWN_LOGGER()
+
+		BEGIN_KNOWN_LOGGER("Misc")
+			//KNOWN_EVENT("CreateThread", Misc_CreateThread)
+			//KNOWN_EVENT("SetThreadGroup", Misc_SetThreadGroup)
+			//KNOWN_EVENT("BeginThreadGroupScope", Misc_BeginThreadGroupScope)
+			//KNOWN_EVENT("BookmarkSpec", Misc_BookmarkSpec)
+			KNOWN_EVENT("Bookmark", Misc_Bookmark)
+			KNOWN_EVENT("ScreenshotHeader", Misc_ScreenshotHeader)
+			KNOWN_EVENT("RegionBegin", Misc_RegionBegin)
+			KNOWN_EVENT("RegionBeginWithId", Misc_RegionBeginWithId)
+			KNOWN_EVENT("RegionEnd", Misc_RegionEnd)
+			KNOWN_EVENT("RegionEndWithId", Misc_RegionEndWithId)
+			//KNOWN_EVENT("BeginGameFrame", Misc_BeginGameFrame)
+			//KNOWN_EVENT("EndGameFrame", Misc_EndGameFrame)
+			//KNOWN_EVENT("BeginRenderFrame", Misc_BeginRenderFrame)
+			//KNOWN_EVENT("EndRenderFrame", Misc_EndRenderFrame)
+			KNOWN_EVENT("BeginFrame", Misc_BeginFrame)
+			KNOWN_EVENT("EndFrame", Misc_EndFrame)
+		END_KNOWN_LOGGER()
+
+		BEGIN_KNOWN_LOGGER("Logging")
+			//KNOWN_EVENT("LogCategory", Logging_LogCategory)
+			//KNOWN_EVENT("LogMessageSpec", Logging_LogMessageSpec)
+			KNOWN_EVENT("LogMessage", Logging_LogMessage)
+		END_KNOWN_LOGGER()
+
+		BEGIN_KNOWN_LOGGER("PlatformFile")
+			KNOWN_EVENT("BeginOpen", PlatformFile_BeginOpen)
+			KNOWN_EVENT("EndOpen", PlatformFile_EndOpen)
+			KNOWN_EVENT("BeginRead", PlatformFile_BeginRead)
+			KNOWN_EVENT("EndRead", PlatformFile_EndRead)
+			KNOWN_EVENT("BeginClose", PlatformFile_BeginClose)
+			KNOWN_EVENT("EndClose", PlatformFile_EndClose)
+		END_KNOWN_LOGGER()
+
+		//BEGIN_KNOWN_LOGGER("LoadTime")
+			//KNOWN_EVENT("ClassInfo", LoadTime_ClassInfo)
+			//KNOWN_EVENT("NewAsyncPackage", LoadTime_NewAsyncPackage)
+		//END_KNOWN_LOGGER()
+
+		BEGIN_KNOWN_LOGGER("GpuProfiler")
+			//KNOWN_EVENT("EventSpec", GpuProfiler_EventSpec) // deprecated
+			//KNOWN_EVENT("Frame", GpuProfiler_Frame) // deprecated
+			//KNOWN_EVENT("Frame2", GpuProfiler_Frame2) // deprecated
+			//KNOWN_EVENT("Init", GpuProfiler_Init)
+			//KNOWN_EVENT("QueueSpec", GpuProfiler_QueueSpec)
+			//KNOWN_EVENT("EventFrameBoundary", GpuProfiler_EventFrameBoundary)
+			//KNOWN_EVENT("EventBreadcrumbSpec", GpuProfiler_EventBreadcrumbSpec)
+			KNOWN_EVENT("EventBeginBreadcrumb", GpuProfiler_EventBeginBreadcrumb)
+			KNOWN_EVENT("EventEndBreadcrumb", GpuProfiler_EventEndBreadcrumb)
+			KNOWN_EVENT("EventBeginWork", GpuProfiler_EventBeginWork)
+			KNOWN_EVENT("EventEndWork", GpuProfiler_EventEndWork)
+			KNOWN_EVENT("EventWait", GpuProfiler_EventWait)
+			//KNOWN_EVENT("EventStats", GpuProfiler_EventStats)
+			KNOWN_EVENT("SignalFence", GpuProfiler_SignalFence)
+			KNOWN_EVENT("WaitFence", GpuProfiler_WaitFence)
+		END_KNOWN_LOGGER()
+
+		BEGIN_KNOWN_LOGGER("CpuProfiler")
+			//KNOWN_EVENT("EventSpec", CpuProfiler_EventSpec)
+			KNOWN_EVENT("EventBatchV3", CpuProfiler_EventBatchV3)
+			KNOWN_EVENT("EventBatchV2", CpuProfiler_EventBatchV2) // deprecated
+			KNOWN_EVENT("EventBatch", CpuProfiler_EventBatch) // deprecated
+			KNOWN_EVENT("EndCapture", CpuProfiler_EndCapture) // deprecated
+			KNOWN_EVENT("EndThread", CpuProfiler_EndThread)
+		END_KNOWN_LOGGER()
+
+		//BEGIN_KNOWN_LOGGER("Stats")
+			//KNOWN_EVENT("Spec", Stats_Spec)
+		//END_KNOWN_LOGGER()
+
+		//BEGIN_KNOWN_LOGGER("Counters")
+			//KNOWN_EVENT("Spec", Counters_Spec)
+		//END_KNOWN_LOGGER()
+
+		//BEGIN_KNOWN_LOGGER("CsvProfiler")
+			//KNOWN_EVENT("RegisterCategory", CsvProfiler_RegisterCategory)
+			//KNOWN_EVENT("DefineInlineStat", CsvProfiler_DefineInlineStat)
+			//KNOWN_EVENT("DefineDeclaredStat", CsvProfiler_DefineDeclaredStat)
+			//KNOWN_EVENT("Metadata", CsvProfiler_Metadata)
+			//KNOWN_EVENT("BeginCapture", CsvProfiler_BeginCapture)
+		//END_KNOWN_LOGGER()
+
+		BEGIN_KNOWN_LOGGER("Memory")
+			KNOWN_EVENT("Marker", Memory_Marker)
+		END_KNOWN_LOGGER()
+
+		BEGIN_KNOWN_LOGGER("LLM")
+			KNOWN_EVENT("TagValue", LLM_TagValue)
+			//KNOWN_EVENT("TagsSpec", LLM_TagsSpec)
+			//KNOWN_EVENT("TrackerSpec", LLM_TrackerSpec)
+		END_KNOWN_LOGGER()
+
+		BEGIN_KNOWN_LOGGER("SlateTrace")
+			KNOWN_EVENT("AddWidget", SlateTrace_AddWidget)
+		END_KNOWN_LOGGER()
+
+		break;
+	}
+
+	#undef BEGIN_KNOWN_LOGGER
+	#undef KNOWN_EVENT
+	#undef END_KNOWN_LOGGER
+}
+
+////////////////////////////////////////////////////////////////////////////////////////////////////
+
+bool FTrimAnalyzer::OnNewEvent(uint16 RouteId, const FEventTypeInfo& TypeInfo)
+{
+	const uint32 EventId = TypeInfo.GetId();
+
+	const ANSICHAR* LoggerName = TypeInfo.GetLoggerName();
+	check(LoggerName != nullptr);
+
+	const ANSICHAR* EventName = TypeInfo.GetName();
+	check(EventName != nullptr);
+
+	ProcessKnownEvent(EventId, LoggerName, EventName);
+
+	static constexpr FAnsiStringView TraceLoggerName = "$Trace";
+	static constexpr FAnsiStringView TraceNewTraceEventName = "NewTrace";
+	static constexpr FAnsiStringView TraceThreadInfoName = "ThreadInfo";
+
+	FEventInfo* EventInfo = EventInfos.Find(EventId);
+	const UE::Trace::FTraceWriterEventInfo* WriterEventInfo = TraceWriter.FindEvent(LoggerName, EventName);
+
+	// Is it a $Trace.* event?
+	if (TraceLoggerName.Equals(LoggerName))
+	{
+		if (TraceNewTraceEventName.Equals(EventName))
+		{
+			// $Trace.NewTrace
+			TraceNewTraceEventId = EventId;
+		}
+		else if (TraceThreadInfoName.Equals(EventName))
+		{
+			// $Trace.ThreadInfo
+			TraceThreadInfoEventId = EventId;
+		}
+
+		// $Trace.* events are expected to be declared automatically (by TraceWriter.Begin() or by TraceWriter.DeclareDefaultEvents()).
+		if (!WriterEventInfo)
+		{
+			LogTraceTrimmerWarning("Unknown default event (%s.%s)!", ANSI_TO_TCHAR(LoggerName), ANSI_TO_TCHAR(EventName));
+		}
+	}
+	else
+	{
+		// Normal events are not expected to be declared multiple times.
+		if (EventInfo)
+		{
+			if (EventInfo->TypeInfo->GetId() != EventId)
+			{
+				LogTraceTrimmerWarning("Event (%s.%s) with id %u is already declared with id %u!",
+					ANSI_TO_TCHAR(LoggerName), ANSI_TO_TCHAR(EventName), EventId, EventInfo->TypeInfo->GetId());
+
+				EventInfo = nullptr;
+			}
+			else
+			{
+				LogTraceTrimmerWarning("Event (%s.%s) is already declared!", ANSI_TO_TCHAR(LoggerName), ANSI_TO_TCHAR(EventName));
+			}
+		}
+	}
+
+	uint32 WriterEventId = 0;
+
+	if (WriterEventInfo)
+	{
+		WriterEventId = WriterEventInfo->GetId();
+	}
+	else
+	{
+		const UE::Trace::ETraceWriterEventFlags WriterEventFlags = UE::Trace::FTraceWriter::ConvertEventFlags(TypeInfo.GetFlags());
+		auto& Builder = TraceWriter.DeclareEvent(LoggerName, EventName, WriterEventFlags);
+		for (int FieldIndex = 0, FieldCount = TypeInfo.GetFieldCount(); FieldIndex < FieldCount; ++FieldIndex)
+		{
+			const FEventFieldInfo& FieldInfo = *(TypeInfo.GetFieldInfo(FieldIndex));
+			UE::Trace::ETraceWriterFieldType WriterFieldType = UE::Trace::FTraceWriter::ConvertFieldType(FieldInfo);
+			Builder.Field(FieldInfo.GetName(), WriterFieldType);
+		}
+		WriterEventId = Builder.End();
+	}
+
+	if (!EventInfo)
+	{
+		EventInfos.Add(EventId, { &TypeInfo, WriterEventId });
+	}
+
+	return true;
+}
+
+////////////////////////////////////////////////////////////////////////////////////////////////////
+
+bool FTrimAnalyzer::OnEvent(uint16 RouteId, EStyle Style, const FOnEventContext& Context)
+{
+	const FEventData& EventData = Context.EventData;
+
+	const uint32 ContextThreadId = Context.ThreadInfo.GetId();
+
+	const FEventTypeInfo& TypeInfo = EventData.GetTypeInfo();
+
+	const uint32 EventId = TypeInfo.GetId();
+	check(EventId != InvalidEventId);
+
+	const ANSICHAR* LoggerName = TypeInfo.GetLoggerName();
+	check(LoggerName != nullptr);
+
+	const ANSICHAR* EventName = TypeInfo.GetName();
+	check(EventName != nullptr);
+
+	bool bShouldWriteEvent = true;
+
+	//////////////////////////////////////////////////
+	// Known $Trace.* events
+
+	// $Trace.NewTrace
+	if (EventId == TraceNewTraceEventId)
+	{
+		// $Trace.NewTrace should be encountered only once
+		check(bIsTraceNewTraceEventDetected == false);
+		bIsTraceNewTraceEventDetected = true;
+
+		SessionCycleFrequency = Context.EventData.GetValue<uint64>("CycleFrequency");
+		SessionStartCycle = Context.EventData.GetValue<uint64>("StartCycle");
+		SessionStartTimeSinceEpoch = Context.EventData.GetValue<double>("StartDateTime");
+
+		SessionCurrentCycle = SessionStartCycle;
+		SessionCurrentTime = 0.0;
+
+		auto TimeGetter = [this]() { return SessionCurrentCycle; };
+		TraceWriter.SetCustomClock(SessionCycleFrequency, TimeGetter, SessionStartCycle, SessionStartTimeSinceEpoch);
+		TraceWriter.WriteNewTraceEvent();
+		bShouldWriteEvent = false;
+	}
+	else
+	{
+		// $Trace.NewTrace should be declared and emitted before any other event
+		check(bIsTraceNewTraceEventDetected == true);
+	}
+
+	// $Trace.ThreadInfo
+	if (EventId == TraceThreadInfoEventId)
+	{
+		uint32 ThreadId = Context.EventData.GetValue<uint32>("ThreadId");
+		uint32 SystemId = Context.EventData.GetValue<uint32>("SystemId");
+		int32 SortHint = Context.EventData.GetValue<int32>("SortHint");
+		FAnsiStringView ThreadName;
+		Context.EventData.GetString("Name", ThreadName);
+		
+		const UE::Trace::FTraceWriterThreadInfo* ThreadInfo = TraceWriter.GetThreadInfo(ThreadId);
+		if (ThreadInfo != nullptr)
+		{
+			if (!ThreadName.Equals(ThreadInfo->GetName(), ESearchCase::CaseSensitive) ||
+				SystemId != ThreadInfo->GetSystemId() ||
+				SortHint != ThreadInfo->GetSortHint())
+			{
+				LogTraceTrimmerWarning("Thread %u is already registered with different properties (previous Name=\"%s\", SystemId=%u, SortHint=%d; new Name=\"%s\", SystemId=%u, SortHint=%d)!",
+					ThreadId,
+					*FString(ThreadInfo->GetName()), ThreadInfo->GetSystemId(), ThreadInfo->GetSortHint(),
+					*FString(ThreadName), SystemId, SortHint);
+			}
+			else
+			{
+				LogTraceTrimmerMessage("Thread %u is already registered (Name=\"%s\", SystemId=%u, SortHint=%d).",
+					ThreadId,
+					*FString(ThreadName), SystemId, SortHint);
+			}
+			//bShouldWriteEvent = false;
+		}
+		else
+		{
+			LogTraceTrimmerMessage("Registering thread %u (Name=\"%s\", SystemId=%u, SortHint=%d).",
+				ThreadId,
+				*FString(ThreadName), SystemId, SortHint);
+
+			constexpr bool bShouldWriteThreadInfo = false;
+			TraceWriter.RegisterCustomThread(ThreadId, ThreadName, SystemId, SortHint, bShouldWriteThreadInfo);
+			//TraceWriter.WriteThreadInfo(ThreadId);
+			//bShouldWriteEvent = false;
+		}
+	}
+
+	//////////////////////////////////////////////////
+
+	if (RouteId == 0) // Normal events
+	{
+		if (Style != EStyle::Normal)
+		{
+			LogTraceTrimmerWarning("Event (%s.%s) is dispatched as normal event, but with Style=%u!",
+				ANSI_TO_TCHAR(LoggerName), ANSI_TO_TCHAR(EventName), (uint32)Style);
+		}
+	}
+	else if (RouteId == 1) // Scoped events
+	{
+		if (Style == EStyle::EnterScope)
+		{
+			SetCurrentThread(ContextThreadId);
+
+			uint64 Timestamp = Context.EventTime.GetTimestamp();
+			if (Timestamp != 0)
+			{
+				TraceWriter.WriteStampedEnterScopeEvent(SessionStartCycle + Timestamp);
+			}
+			else
+			{
+				TraceWriter.WriteEnterScopeEvent();
+			}
+		}
+		else if (Style == EStyle::LeaveScope)
+		{
+			SetCurrentThread(ContextThreadId);
+
+			uint64 Timestamp = Context.EventTime.GetTimestamp();
+			if (Timestamp != 0)
+			{
+				TraceWriter.WriteStampedLeaveScopeEvent(SessionStartCycle + Timestamp);
+			}
+			else
+			{
+				TraceWriter.WriteLeaveScopeEvent();
+			}
+
+			bShouldWriteEvent = false;
+		}
+		else // unknown Style
+		{
+			LogTraceTrimmerWarning("Event (%s.%s) is dispatched as scoped event, but with Style=%u!",
+				ANSI_TO_TCHAR(LoggerName), ANSI_TO_TCHAR(EventName), (uint32)Style);
+		}
+	}
+	else // unknown RouteId
+	{
+		LogTraceTrimmerWarning("Event (%s.%s) is dispatched on unknown route %u!",
+			ANSI_TO_TCHAR(LoggerName), ANSI_TO_TCHAR(EventName), (uint32)RouteId);
+	}
+
+	//////////////////////////////////////////////////
+
+	FEventInfo* EventInfo = EventInfos.Find(EventId);
+	if (EventInfo == nullptr)
+	{
+		if (LoggerName[0] == 0) // empty dispatch, expected on LeaveScope events
+		{
+			EventInfo = &EmptyEventInfo;
+			bShouldWriteEvent = false;
+		}
+		else
+		{
+			LogTraceTrimmerWarning("Event %u (%s.%s) is not registered!", EventId, ANSI_TO_TCHAR(LoggerName), ANSI_TO_TCHAR(EventName));
+			EventInfo = &EventInfos.Add(EventId, { &TypeInfo, 0 });
+			bShouldWriteEvent = false;
+		}
+	}
+
+	//////////////////////////////////////////////////
+	// Filters
+
+	// Only filter the non-important events.
+	if (bShouldWriteEvent && ContextThreadId != 1)
+	{
+		// Filter by start time.
+		if (SessionCurrentTime < Parameters.StartTime)
+		{
+			bShouldWriteEvent = false;
+		}
+
+		// Include only the specified loggers and events.
+		if (bShouldWriteEvent && !IncludeFilter.IsEmpty())
+		{
+			bool bIsIncluded = false;
+
+			FAnsiString LoggerNameStr(LoggerName);
+			FAnsiString EventNameStr(EventName);
+
+			for (const FNameFilterPattern& FilterPattern : IncludeFilter.Patterns)
+			{
+				if (LoggerNameStr.MatchesWildcard(FilterPattern.LoggerName) &&
+					EventNameStr.MatchesWildcard(FilterPattern.EventName))
+				{
+					bIsIncluded = true;
+					break;
+				}
+			}
+
+			if (!bIsIncluded)
+			{
+				bShouldWriteEvent = false;
+			}
+		}
+
+		// Exclude the specified loggers and events.
+		if (bShouldWriteEvent && !ExcludeFilter.IsEmpty())
+		{
+			bool bIsExcluded = false;
+
+			FAnsiString LoggerNameStr(LoggerName);
+			FAnsiString EventNameStr(EventName);
+
+			for (const FNameFilterPattern& FilterPattern : ExcludeFilter.Patterns)
+			{
+				if (LoggerNameStr.MatchesWildcard(FilterPattern.LoggerName) &&
+					EventNameStr.MatchesWildcard(FilterPattern.EventName))
+				{
+					bIsExcluded = true;
+					break;
+				}
+			}
+
+			if (bIsExcluded)
+			{
+				bShouldWriteEvent = false;
+			}
+		}
+	}
+
+	//////////////////////////////////////////////////
+	// Write event
+
+	if (bShouldWriteEvent)
+	{
+		SetCurrentThread(ContextThreadId);
+
+		UE::Trace::FTraceWriterEventBuilder& EventBuilder = TraceWriter.WriteEvent(EventInfo->WriterEventId);
+
+		const uint32 FieldCount = (Style == EStyle::LeaveScope) ? 0 : TypeInfo.GetFieldCount();
+		for (uint32 FieldIndex = 0; FieldIndex < FieldCount; ++FieldIndex)
+		{
+			const FEventFieldInfo& FieldInfo = *(TypeInfo.GetFieldInfo(FieldIndex));
+			if (FieldInfo.IsArray())
+			{
+				WriteArrayField(FieldIndex, FieldInfo, EventData, EventBuilder);
+			}
+			else
+			{
+				WriteScalarField(FieldIndex, FieldInfo, EventData, EventBuilder);
+
+				// Detect timestamp fields...
+				if (FieldInfo.GetType() == FEventFieldInfo::EType::Integer &&
+					!FieldInfo.IsSigned() &&
+					FieldInfo.GetTypeSize() == 8 &&
+					FCStringAnsi::Strcmp(FieldInfo.GetName(), "Cycle") == 0)
+				{
+					const uint64 Value = EventData.GetValue<uint64>(FieldInfo.GetName());
+					const double Time = Context.EventTime.AsSeconds(Value);
+					UpdatedSessionTime(Value, Time);
+				}
+			}
+		}
+
+		EventBuilder.End();
+	}
+
+	//////////////////////////////////////////////////
+	// Processing known events
+
+	check(EventId <= MaxKnownEventId);
+	EKnownEvent KnownEvent = KnownEventMap[EventId];
+	check((int)KnownEvent < (int)EKnownEvent::Count);
+	const FKnownEvent& Event = KnownEvents[(int)KnownEvent];
+
+	if (Event.Callback)
+	{
+		Event.Callback(Context);
+	}
+
+	//////////////////////////////////////////////////
+
+	// Filter by end time.
+	if (SessionCurrentTime > Parameters.EndTime)
+	{
+		return false; // stop further processing
+	}
+
+	return true;
+}
+
+////////////////////////////////////////////////////////////////////////////////////////////////////
+
+void FTrimAnalyzer::WriteArrayField(
+	uint32 FieldIndex,
+	const UE::Trace::IAnalyzer::FEventFieldInfo& FieldInfo,
+	const UE::Trace::IAnalyzer::FEventData& EventData,
+	UE::Trace::FTraceWriterEventBuilder& EventBuilder)
+{
+	switch (FieldInfo.GetType())
+	{
+		case FEventFieldInfo::EType::Integer:
+		{
+			if (FieldInfo.IsSigned())
+			{
+				switch (FieldInfo.GetTypeSize())
+				{
+					case 1: // int8
+					{
+						const TArrayReader<int8>& Reader = EventData.GetArray<int8>(FieldInfo.GetName());
+						EventBuilder.Field(FieldIndex, TConstArrayView<int8>(Reader.GetData(), Reader.Num()));
+						break;
+					}
+					case 2: // int16
+					{
+						const TArrayReader<int16>& Reader = EventData.GetArray<int16>(FieldInfo.GetName());
+						EventBuilder.Field(FieldIndex, TConstArrayView<int16>(Reader.GetData(), Reader.Num()));
+						break;
+					}
+					case 4: // int32
+					{
+						const TArrayReader<int32>& Reader = EventData.GetArray<int32>(FieldInfo.GetName());
+						EventBuilder.Field(FieldIndex, TConstArrayView<int32>(Reader.GetData(), Reader.Num()));
+						break;
+					}
+					case 8: // int64
+					{
+						const TArrayReader<int64>& Reader = EventData.GetArray<int64>(FieldInfo.GetName());
+						EventBuilder.Field(FieldIndex, TConstArrayView<int64>(Reader.GetData(), Reader.Num()));
+						break;
+					}
+					default:
+					{
+						// error
+					}
+				} // switch (FieldInfo.GetTypeSize())
+			}
+			else // unsigned
+			{
+				switch (FieldInfo.GetTypeSize())
+				{
+					case 1: // uint8
+					{
+						const TArrayReader<uint8>& Reader = EventData.GetArray<uint8>(FieldInfo.GetName());
+						EventBuilder.Field(FieldIndex, TConstArrayView<uint8>(Reader.GetData(), Reader.Num()));
+						break;
+					}
+					case 2: // uint16
+					{
+						const TArrayReader<uint16>& Reader = EventData.GetArray<uint16>(FieldInfo.GetName());
+						EventBuilder.Field(FieldIndex, TConstArrayView<uint16>(Reader.GetData(), Reader.Num()));
+						break;
+					}
+					case 4: // uint32
+					{
+						const TArrayReader<uint32>& Reader = EventData.GetArray<uint32>(FieldInfo.GetName());
+						EventBuilder.Field(FieldIndex, TConstArrayView<uint32>(Reader.GetData(), Reader.Num()));
+						break;
+					}
+					case 8: // uint64
+					{
+						const TArrayReader<uint64>& Reader = EventData.GetArray<uint64>(FieldInfo.GetName());
+						EventBuilder.Field(FieldIndex, TConstArrayView<uint64>(Reader.GetData(), Reader.Num()));
+						break;
+					}
+					default:
+					{
+						// error
+					}
+				} // switch (FieldInfo.GetTypeSize())
+			}
+			break;
+		}
+
+		case FEventFieldInfo::EType::Float: // Float32 or Float64
+		{
+			switch (FieldInfo.GetTypeSize())
+			{
+				case 4: // float
+				{
+					const TArrayReader<float>& Reader = EventData.GetArray<float>(FieldInfo.GetName());
+					EventBuilder.Field(FieldIndex, TConstArrayView<float>(Reader.GetData(), Reader.Num()));
+					break;
+				}
+				case 8: // double
+				{
+					const TArrayReader<double>& Reader = EventData.GetArray<double>(FieldInfo.GetName());
+					EventBuilder.Field(FieldIndex, TConstArrayView<double>(Reader.GetData(), Reader.Num()));
+					break;
+				}
+				default:
+				{
+					// error
+				}
+			} // switch (FieldInfo.GetTypeSize())
+			break;
+		}
+
+		case FEventFieldInfo::EType::AnsiString:
+		{
+			FAnsiStringView Value;
+			EventData.GetString(FieldInfo.GetName(), Value);
+			EventBuilder.Field(FieldIndex, Value);
+			break;
+		}
+
+		case FEventFieldInfo::EType::WideString:
+		{
+			FWideStringView Value;
+			EventData.GetString(FieldInfo.GetName(), Value);
+			EventBuilder.Field(FieldIndex, Value);
+			break;
+		}
+
+		default:
+		{
+			// error
+		}
+	}
+}
+
+////////////////////////////////////////////////////////////////////////////////////////////////////
+
+void FTrimAnalyzer::WriteScalarField(
+	uint32 FieldIndex,
+	const UE::Trace::IAnalyzer::FEventFieldInfo& FieldInfo,
+	const UE::Trace::IAnalyzer::FEventData& EventData,
+	UE::Trace::FTraceWriterEventBuilder& EventBuilder)
+{
+	switch (FieldInfo.GetType())
+	{
+		case FEventFieldInfo::EType::Integer:
+		{
+			if (FieldInfo.IsSigned())
+			{
+				switch (FieldInfo.GetTypeSize())
+				{
+					case 1: // int8
+					{
+						const int8 Value = EventData.GetValue<int8>(FieldInfo.GetName());
+						EventBuilder.Field(FieldIndex, Value);
+						break;
+					}
+					case 2: // int16
+					{
+						const int16 Value = EventData.GetValue<int16>(FieldInfo.GetName());
+						EventBuilder.Field(FieldIndex, Value);
+						break;
+					}
+					case 4: // int32
+					{
+						const int32 Value = EventData.GetValue<int32>(FieldInfo.GetName());
+						EventBuilder.Field(FieldIndex, Value);
+						break;
+					}
+					case 8: // int64
+					{
+						const int64 Value = EventData.GetValue<int64>(FieldInfo.GetName());
+						EventBuilder.Field(FieldIndex, Value);
+						break;
+					}
+					default:
+					{
+						// error
+					}
+				} // switch (FieldInfo.GetTypeSize())
+			}
+			else // unsigned
+			{
+				switch (FieldInfo.GetTypeSize())
+				{
+					case 1: // uint8
+					{
+						const uint8 Value = EventData.GetValue<uint8>(FieldInfo.GetName());
+						EventBuilder.Field(FieldIndex, Value);
+						break;
+					}
+					case 2: // uint16
+					{
+						const uint16 Value = EventData.GetValue<uint16>(FieldInfo.GetName());
+						EventBuilder.Field(FieldIndex, Value);
+						break;
+					}
+					case 4: // uint32
+					{
+						const uint32 Value = EventData.GetValue<uint32>(FieldInfo.GetName());
+						EventBuilder.Field(FieldIndex, Value);
+						break;
+					}
+					case 8: // uint64
+					{
+						const uint64 Value = EventData.GetValue<uint64>(FieldInfo.GetName());
+						EventBuilder.Field(FieldIndex, Value);
+						break;
+					}
+					default:
+					{
+						// error
+					}
+				} // switch (FieldInfo.GetTypeSize())
+			}
+			break;
+		}
+
+		case FEventFieldInfo::EType::Float: // Float32 or Float64
+		{
+			switch (FieldInfo.GetTypeSize())
+			{
+				case 4: // float
+				{
+					const float Value = EventData.GetValue<float>(FieldInfo.GetName());
+					EventBuilder.Field(FieldIndex, Value);
+					break;
+				}
+				case 8: // double
+				{
+					const double Value = EventData.GetValue<double>(FieldInfo.GetName());
+					EventBuilder.Field(FieldIndex, Value);
+					break;
+				}
+				default:
+				{
+					// error
+				}
+			} // switch (FieldInfo.GetTypeSize())
+			break;
+		}
+
+		case FEventFieldInfo::EType::Reference8:
+		{
+			check(FieldInfo.GetSize() == 1);
+			UE::Trace::FEventRef8 RefValue = EventData.GetReferenceValue<uint8>(FieldIndex);
+			//TODO: EventBuilder.Field(FieldIndex, RefValue);
+			break;
+		}
+
+		case FEventFieldInfo::EType::Reference16:
+		{
+			check(FieldInfo.GetSize() == 2);
+			UE::Trace::FEventRef16 RefValue = EventData.GetReferenceValue<uint16>(FieldIndex);
+			//TODO: EventBuilder.Field(FieldIndex, RefValue);
+			break;
+		}
+
+		case FEventFieldInfo::EType::Reference32:
+		{
+			check(FieldInfo.GetSize() == 4);
+			UE::Trace::FEventRef32 RefValue = EventData.GetReferenceValue<uint32>(FieldIndex);
+			//TODO: EventBuilder.Field(FieldIndex, RefValue);
+			break;
+		}
+
+		case FEventFieldInfo::EType::Reference64:
+		{
+			check(FieldInfo.GetSize() == 8);
+			UE::Trace::FEventRef64 RefValue = EventData.GetReferenceValue<uint64>(FieldIndex);
+			//TODO: EventBuilder.Field(FieldIndex, RefValue);
+			break;
+		}
+
+		default:
+		{
+			// error
+		}
+	}
+}
+
+////////////////////////////////////////////////////////////////////////////////////////////////////
+
+void FTrimAnalyzer::SetCurrentThread(uint32 ThreadId)
+{
+	const UE::Trace::FTraceWriterThreadInfo* ThreadInfo = TraceWriter.GetThreadInfo(ThreadId);
+	if (!ThreadInfo)
+	{
+		LogTraceTrimmerMessage("Registering thread %u (unknown).", ThreadId);
+		constexpr bool bShouldWriteThreadInfo = false;
+		TraceWriter.RegisterCustomThread(ThreadId, FAnsiStringView("Unknown"), 0, 0, bShouldWriteThreadInfo);
+	}
+	TraceWriter.SetCurrentThread(ThreadId);
+}
+
+////////////////////////////////////////////////////////////////////////////////////////////////////
+
+void FTrimAnalyzer::UpdatedSessionTime(uint64 Cycle, double Time)
+{
+	if (Cycle > SessionCurrentCycle)
+	{
+		SessionCurrentCycle = Cycle;
+		SessionCurrentTime = Time;
+	}
+}
+
+////////////////////////////////////////////////////////////////////////////////////////////////////
+
+void FTrimAnalyzer::ProcessCpuEventBatchV1(const FOnEventContext& Context, const uint8* BufferPtr, uint32 BufferSize)
+{
+	const FEventData& EventData = Context.EventData;
+
+	const uint32 ThreadId = GetThreadIdField(Context);
+
+	uint64* LastCyclePtr = ThreadMap.Find(ThreadId);
+	if (!LastCyclePtr)
+	{
+		LastCyclePtr = &ThreadMap.Add(ThreadId, 0);
+	}
+	uint64 LastCycle = *LastCyclePtr;
+
+	const uint8* BufferEnd = BufferPtr + BufferSize;
+	while (BufferPtr < BufferEnd)
+	{
+		uint64 DecodedCycle = FTraceAnalyzerUtils::Decode7bit(BufferPtr);
+		LastCycle += (DecodedCycle >> 1);
+
+		if (DecodedCycle & 1ull)
+		{
+			uint32 SpecId = (uint32)FTraceAnalyzerUtils::Decode7bit(BufferPtr);
+		}
+	}
+	if (!ensure(BufferPtr == BufferEnd))
+	{
+		//error
+	}
+
+	*LastCyclePtr = LastCycle;
+	const double Time = Context.EventTime.AsSeconds(LastCycle);
+	UpdatedSessionTime(LastCycle, Time);
+}
+
+////////////////////////////////////////////////////////////////////////////////////////////////////
+
+void FTrimAnalyzer::ProcessCpuEventBatchV2V3(const FOnEventContext& Context, const uint8* BufferPtr, uint32 BufferSize, int32 Version = 3)
+{
+	const FEventData& EventData = Context.EventData;
+
+	const uint32 ThreadId = GetThreadIdField(Context);
+
+	uint64* LastCyclePtr = ThreadMap.Find(ThreadId);
+	if (!LastCyclePtr)
+	{
+		LastCyclePtr = &ThreadMap.Add(ThreadId, 0);
+	}
+
+	uint64 LastCycle = *LastCyclePtr;
+
+	check(Context.EventTime.GetTimestamp() == 0);
+	const uint64 BaseCycle = Context.EventTime.AsCycle64();
+
+	const uint8* BufferEnd = BufferPtr + BufferSize;
+	while (BufferPtr < BufferEnd)
+	{
+		uint64 DecodedCycle = FTraceAnalyzerUtils::Decode7bit(BufferPtr);
+		uint64 ActualCycle = (DecodedCycle >> 2);
+
+		// ActualCycle larger or equal to LastCycle means we have a new base value.
+		if (ActualCycle < LastCycle)
+		{
+			ActualCycle += LastCycle;
+		}
+
+		// If we late connect we will be joining the cycle stream mid-flow and
+		// will have missed out on it's base timestamp. Reconstruct it here.
+		if (ActualCycle < BaseCycle)
+		{
+			ActualCycle += BaseCycle;
+		}
+
+		if (DecodedCycle & 2ull) // coroutine
+		{
+			if (DecodedCycle & 1ull) // Restore
+			{
+				uint64 CoroutineId = FTraceAnalyzerUtils::Decode7bit(BufferPtr);
+				uint32 TimerScopeDepth = (uint32)FTraceAnalyzerUtils::Decode7bit(BufferPtr);
+			}
+			else // Save
+			{
+				uint32 TimerScopeDepth = (uint32)FTraceAnalyzerUtils::Decode7bit(BufferPtr);
+			}
+		}
+		else // normal CPU events
+		{
+			if (DecodedCycle & 1ull)
+			{
+				uint32 SpecId = (uint32)FTraceAnalyzerUtils::Decode7bit(BufferPtr);
+#if 0
+				if (Version == 3)
+				{
+					if (SpecId & 1u) // The last bit is set if this is a metadata id.
+					{
+						uint32 MetadataId = SpecId >> 1;
+					}
+					else
+					{
+						SpecId = SpecId >> 1;
+					}
+				}
+				else // Version == 2
+				{
+					//SpecId
+				}
+#endif
+			}
+		}
+
+		check(ActualCycle > 0);
+		LastCycle = ActualCycle;
+	}
+	if (!ensure(BufferPtr == BufferEnd))
+	{
+		// error
+	}
+
+	*LastCyclePtr = LastCycle;
+	const double Time = Context.EventTime.AsSeconds(LastCycle);
+	UpdatedSessionTime(LastCycle, Time);
+}
+
+////////////////////////////////////////////////////////////////////////////////////////////////////
+
+UE::Trace::FAnalysisProcessor TrimAsync(FTrimParameters& Parameters, UE::Trace::IInDataStream& InputDataStream, UE::Trace::IOutDataStream& OutputDataStream)
+{
+	UE::Trace::FTraceWriter TraceWriter(OutputDataStream);
+	FTrimAnalyzer TrimAnalyzer(Parameters, TraceWriter);
+	UE::Trace::FAnalysisContext Context;
+	Context.AddAnalyzer(TrimAnalyzer);
+	return Context.Process(InputDataStream);
+}
+
+////////////////////////////////////////////////////////////////////////////////////////////////////
+
+int32 Trim(FTrimParameters& Parameters, UE::Trace::IInDataStream& InputDataStream, UE::Trace::IOutDataStream& OutputDataStream)
+{
+	UE::Trace::FTraceWriter TraceWriter(OutputDataStream);
+	FTrimAnalyzer TrimAnalyzer(Parameters, TraceWriter);
+	UE::Trace::FAnalysisContext Context;
+	Context.AddAnalyzer(TrimAnalyzer);
+	Context.Process(InputDataStream).Wait();
+	return (int32)TrimAnalyzer.GetNumErrors();
+}
+
+////////////////////////////////////////////////////////////////////////////////////////////////////
+////////////////////////////////////////////////////////////////////////////////////////////////////
+
+} // namespace TraceServices
+
+#undef LogTraceTrimmerError
+#undef LogTraceTrimmerWarning
+#undef LogTraceTrimmerMessage
diff --git a/Engine/Source/Developer/TraceServices/Private/Analyzers/TrimAnalyzer.h b/Engine/Source/Developer/TraceServices/Private/Analyzers/TrimAnalyzer.h
new file mode 100644
--- /dev/null
+++ b/Engine/Source/Developer/TraceServices/Private/Analyzers/TrimAnalyzer.h
@@ -0,0 +1,219 @@
+// Copyright Epic Games, Inc. All Rights Reserved.
+
+#pragma once
+
+#include "CoreTypes.h"
+#include "Containers/Map.h"
+#include "Containers/AnsiString.h"
+#include "Containers/UnrealString.h"
+#include "Templates/Function.h"
+
+// TraceAnalysis
+#include "Trace/Analyzer.h"
+
+namespace UE::Trace
+{
+	class FTraceWriter;
+	class FTraceWriterEventBuilder;
+}
+
+namespace TraceServices
+{
+
+struct FTrimParameters;
+
+class FTrimAnalyzer : public UE::Trace::IAnalyzer
+{
+private:
+	struct FEventInfo
+	{
+		const FEventTypeInfo* TypeInfo = nullptr;
+		uint32 WriterEventId = 0;
+	};
+
+	struct FNameFilterPattern
+	{
+		FAnsiString LoggerName;
+		FAnsiString EventName;
+	};
+
+	struct FNameFilter
+	{
+		void Init(FAnsiString Filter);
+		bool IsEmpty() const { return Patterns.IsEmpty(); }
+
+		TArray<FNameFilterPattern> Patterns;
+	};
+
+	enum class EKnownEvent : uint32
+	{
+		Invalid,
+
+		Trace_NewEvent,
+		//Trace_NewTrace,
+		//Trace_ThreadTiming,
+		//Trace_ThreadInfo,
+		//Trace_ThreadGroupBegin,
+		//Trace_ThreadGroupEnd,
+		//Trace_ChannelAnnounce_Old, // "$Trace", deprecated in 4.26
+
+		//Trace_ChannelAnnounce, // "Trace", new in 4.26
+		//Trace_ChannelToggle,
+
+		//Diagnostics_Session,
+		//Diagnostics_Session2,
+
+		//Misc_CreateThread,
+		//Misc_SetThreadGroup,
+		//Misc_BeginThreadGroupScope,
+		//Misc_BookmarkSpec,
+		Misc_Bookmark,
+		Misc_ScreenshotHeader,
+		//Misc_ScreenshotChunk,
+		Misc_RegionBegin,
+		Misc_RegionBeginWithId,
+		Misc_RegionEnd,
+		Misc_RegionEndWithId,
+		//Misc_BeginGameFrame, // deprecated
+		//Misc_EndGameFrame, // deprecated
+		//Misc_BeginRenderFrame, // deprecated
+		//Misc_EndRenderFrame, // deprecated
+		Misc_BeginFrame,
+		Misc_EndFrame,
+
+		//Logging_LogCategory,
+		//Logging_LogMessageSpec,
+		Logging_LogMessage,
+
+		PlatformFile_BeginOpen,
+		PlatformFile_EndOpen,
+		PlatformFile_BeginRead,
+		PlatformFile_EndRead,
+		PlatformFile_BeginClose,
+		PlatformFile_EndClose,
+
+		//LoadTime_ClassInfo,
+		//LoadTime_NewAsyncPackage,
+
+		//GpuProfiler_EventSpec, // old GpuProfiler, deprecated
+		//GpuProfiler_Frame, // old GpuProfiler, deprecated
+		//GpuProfiler_Frame2, // old GpuProfiler, deprecated
+		//GpuProfiler_Init,
+		//GpuProfiler_QueueSpec,
+		//GpuProfiler_EventFrameBoundary,
+		//GpuProfiler_EventBreadcrumbSpec,
+		GpuProfiler_EventBeginBreadcrumb,
+		GpuProfiler_EventEndBreadcrumb,
+		GpuProfiler_EventBeginWork,
+		GpuProfiler_EventEndWork,
+		GpuProfiler_EventWait,
+		//GpuProfiler_EventStats,
+		GpuProfiler_SignalFence,
+		GpuProfiler_WaitFence,
+
+		//CpuProfiler_EventSpec,
+		CpuProfiler_EventBatchV3, // added in UE 5.6
+		CpuProfiler_EventBatchV2, // backward compatibility; added in UE 5.1, removed in 5.6
+		CpuProfiler_EventBatch, // backward compatibility; removed in UE 5.1
+		CpuProfiler_EndCapture, // backward compatibility; removed in UE 5.1
+		CpuProfiler_EndThread, // "Cycle" field added in UE 5.4
+
+		//Stats_Spec,
+
+		//Counters_Spec,
+
+		//CsvProfiler_RegisterCategory,
+		//CsvProfiler_DefineInlineStat,
+		//CsvProfiler_DefineDeclaredStat,
+		//CsvProfiler_Metadata,
+		//CsvProfiler_BeginCapture,
+
+		Memory_Marker,
+
+		LLM_TagValue,
+		//LLM_TagsSpec,
+		//LLM_TrackerSpec,
+
+		SlateTrace_AddWidget,
+
+		Count
+	};
+
+	typedef TFunction<void(const FOnEventContext& Context)> KnownEventCallback;
+
+	struct FKnownEvent
+	{
+		uint32 Id = InvalidEventId;
+		KnownEventCallback Callback = nullptr;
+	};
+
+public:
+	FTrimAnalyzer(FTrimParameters& InParameters, UE::Trace::FTraceWriter& InTraceWriter);
+	virtual ~FTrimAnalyzer();
+
+	uint32 GetNumErrors() const { return NumErrors; }
+	uint32 GetNumWarnings() const { return NumWarnings; }
+
+private:
+	virtual void OnAnalysisBegin(const FOnAnalysisContext& Context) override;
+	virtual void OnAnalysisEnd() override;
+	virtual void OnVersion(uint32 InTransportVersion, uint32 InProtocolVersion) /* override */;
+	virtual bool OnNewEvent(uint16 RouteId, const FEventTypeInfo& TypeInfo) override;
+	virtual bool OnEvent(uint16 RouteId, EStyle, const FOnEventContext& Context) override;
+
+	void WriteArrayField(uint32 FieldIndex, const UE::Trace::IAnalyzer::FEventFieldInfo& FieldInfo, const UE::Trace::IAnalyzer::FEventData& EventData, UE::Trace::FTraceWriterEventBuilder& EventBuilder);
+	void WriteScalarField(uint32 FieldIndex, const UE::Trace::IAnalyzer::FEventFieldInfo& FieldInfo, const UE::Trace::IAnalyzer::FEventData& EventData, UE::Trace::FTraceWriterEventBuilder& EventBuilder);
+
+	void SetCurrentThread(uint32 ThreadId);
+
+	void UpdatedSessionTime(uint64 Cycle, double Time);
+
+	bool IsKnownEvent(uint32 EventId, EKnownEvent KnownEvent) const
+	{
+		return EventId == KnownEvents[(int)KnownEvent].Id;
+	}
+	void RegisterCallbacksForKnownEvents();
+	void ProcessKnownEvent(uint32 EventId, const ANSICHAR* LoggerName, const ANSICHAR* EventName);
+
+	void ProcessCpuEventBatchV1(const FOnEventContext& Context, const uint8* BufferPtr, uint32 BufferSize);
+	void ProcessCpuEventBatchV2V3(const FOnEventContext& Context, const uint8* BufferPtr, uint32 BufferSize, int32 Version);
+
+private:
+	FTrimParameters& Parameters;
+	UE::Trace::FTraceWriter& TraceWriter;
+
+	uint32 TransportVersion = 4;
+	uint32 ProtocolVersion = 7;
+
+	TMap<uint32, uint64> ThreadMap;
+	TMap<uint32, FEventInfo> EventInfos; // EventId --> FEventInfo
+	FEventInfo EmptyEventInfo;
+
+	uint64 AnalysisBeginTimestamp = 0; // realtime timestamp
+
+	uint64 SessionCycleFrequency = 0;
+	uint64 SessionStartCycle = 0;
+	double SessionStartTimeSinceEpoch = 0.0;
+
+	uint64 SessionCurrentCycle = 0; // absolute session time
+	double SessionCurrentTime = 0.0; // relative session time, in [seconds]
+
+	FNameFilter IncludeFilter;
+	FNameFilter ExcludeFilter;
+
+	static constexpr uint32 InvalidEventId = (uint32)-1;
+	static constexpr uint32 MaxKnownEventId = 0xFFFF;
+
+	uint32 TraceNewTraceEventId = InvalidEventId;
+	bool bIsTraceNewTraceEventDetected = false;
+
+	uint32 TraceThreadInfoEventId = InvalidEventId;
+
+	FKnownEvent KnownEvents[(int)EKnownEvent::Count];
+	EKnownEvent KnownEventMap[MaxKnownEventId + 1]; // index == EventId --> EKnownEvent
+
+	uint32 NumErrors = 0;
+	uint32 NumWarnings = 0;
+};
+
+} // namespace TraceServices
diff --git a/Engine/Source/Developer/TraceServices/Public/TraceServices/TraceTrimmer.h b/Engine/Source/Developer/TraceServices/Public/TraceServices/TraceTrimmer.h
new file mode 100644
--- /dev/null
+++ b/Engine/Source/Developer/TraceServices/Public/TraceServices/TraceTrimmer.h
@@ -0,0 +1,71 @@
+// Copyright Epic Games, Inc. All Rights Reserved.
+
+#pragma once
+
+#include "CoreTypes.h"
+#include "Containers/AnsiString.h"
+
+// TraceAnalysis
+#include "Trace/Analysis.h"
+#include "Trace/DataStream.h"
+#include "Trace/OutDataStream.h"
+
+#include <limits>
+
+namespace TraceServices {
+
+struct FTrimParameters
+{
+	/**
+	 * The start time filter.
+	 * The analysis will ignore non-important events with timestamp < StartTime.
+	 */
+	double StartTime = -std::numeric_limits<double>::infinity();
+
+	/**
+	 * The end time filter.
+	 * The analysis will stop when a first non-important event with timestamp > EndTime is detected.
+	 */
+	double EndTime = +std::numeric_limits<double>::infinity();
+
+	/**
+	 * Comma separated list of trace events to include.
+	 * The events are specified as "logger.event" wildcard patterns.
+	 * Example: "$trace.*,cpu.*".
+	 * If specified, only the trace events with name that matches one of the patterns in the Include string will be included.
+	 * Empty string (default) is equivalent with "*.*" (include all).
+	 */
+	FAnsiString Include;
+
+	/**
+	 * Comma separated list of trace events to exclude.
+	 * The events are specified as "logger.event" wildcard patterns.
+	 * Example: "$trace.*,cpu.*".
+	 * If specified, the trace events with name that matches one of the pattern in the Exclude string will be excluded.
+	 */
+	FAnsiString Exclude;
+};
+
+/**
+ * Trims a trace stream using the provided parameters.
+ * The function starts the trimming process and returns the FAnalysisProcessor. The caller is responsible for waiting the processor to complete.
+ *
+ * @param Parameters The trim process parameters
+ * @param InputDataStream The input data stream
+ * @param OutputDataStream The output data stream
+ * @returns the analysis processor object
+ */
+TRACESERVICES_API UE::Trace::FAnalysisProcessor TrimAsync(FTrimParameters& Parameters, UE::Trace::IInDataStream& InputDataStream, UE::Trace::IOutDataStream& OutputDataStream);
+
+/**
+ * Trims a trace stream using the provided parameters.
+ * The function waits for the trimming process to complete, before returning.
+ *
+ * @param Parameters The trim process parameters
+ * @param InputDataStream The input data stream
+ * @param OutputDataStream The output data stream
+ * @returns 0 if the trim process is completed successfully, non-zero otherwise
+ */
+TRACESERVICES_API int32 Trim(FTrimParameters& Parameters, UE::Trace::IInDataStream& InputDataStream, UE::Trace::IOutDataStream& OutputDataStream);
+
+} // namespace TraceServices
