[Insights] MemoryInsights: Fixed crash when processing a trace with more than 93 milions live allocations. Replaced the TMap (now limited to 2GB, i.e. for the AllocMap that limit is encounterd at ~93 M allocs) with the TSparseMap.

#rnx

#Synced-CL 50041282

diff --git a/Engine/Source/Developer/TraceServices/Private/Model/AllocMap.h b/Engine/Source/Developer/TraceServices/Private/Model/AllocMap.h
--- a/Engine/Source/Developer/TraceServices/Private/Model/AllocMap.h
+++ b/Engine/Source/Developer/TraceServices/Private/Model/AllocMap.h
@@ -132,7 +132,7 @@ private:
 	};
 
 private:
-	TMap<uint64, FAllocationItem*, FDefaultSetAllocator, TAddressKeyFuncs<FAllocationItem* >> AllocsMap;
+	TSparseMap<uint64, FAllocationItem*, FDefaultSparseSetAllocator, TAddressKeyFuncs<FAllocationItem*>> AllocsMap;
 };
 
 } // namespace TraceServices
