[Insights] TraceTrimmer: Fixed filtering by time range to infer timestamps of trace events more accurately. Now tracking of timestamps is made for each thread.

#rnx

#Synced-CL 50155428

diff --git a/Engine/Source/Developer/TraceServices/Private/Analyzers/TrimAnalyzer.cpp b/Engine/Source/Developer/TraceServices/Private/Analyzers/TrimAnalyzer.cpp
--- a/Engine/Source/Developer/TraceServices/Private/Analyzers/TrimAnalyzer.cpp
+++ b/Engine/Source/Developer/TraceServices/Private/Analyzers/TrimAnalyzer.cpp
@@ -174,8 +174,7 @@ void FTrimAnalyzer::RegisterCallbacksForKnownEvents()
 		[this](const FOnEventContext& Context)
 		{
 			const uint64 Cycle = Context.EventData.GetValue<uint64>("CycleAndId");
-			const double Time = Context.EventTime.AsSeconds(Cycle);
-			UpdatedSessionTime(Cycle, Time);
+			UpdateThreadTime(Context, Cycle);
 		};
 
 	//////////////////////////////////////////////////
@@ -220,35 +219,30 @@ void FTrimAnalyzer::RegisterCallbacksForKnownEvents()
 
 	KnownEvents[(int)EKnownEvent::GpuProfiler_EventBeginBreadcrumb].Callback =
 	KnownEvents[(int)EKnownEvent::GpuProfiler_EventBeginWork].Callback =
-		
 		[this](const FOnEventContext& Context)
 		{
 			const uint64 Cycle = Context.EventData.GetValue<uint64>("GPUTimestampTOP");
-			const double Time = Context.EventTime.AsSeconds(Cycle);
-			UpdatedSessionTime(Cycle, Time);
+			UpdateThreadTime(Context, Cycle);
 		};
 	KnownEvents[(int)EKnownEvent::GpuProfiler_EventEndBreadcrumb].Callback =
 	KnownEvents[(int)EKnownEvent::GpuProfiler_EventEndWork].Callback =
 		[this](const FOnEventContext& Context)
 		{
 			const uint64 Cycle = Context.EventData.GetValue<uint64>("GPUTimestampBOP");
-			const double Time = Context.EventTime.AsSeconds(Cycle);
-			UpdatedSessionTime(Cycle, Time);
+			UpdateThreadTime(Context, Cycle);
 		};
 	KnownEvents[(int)EKnownEvent::GpuProfiler_EventWait].Callback =
 		[this](const FOnEventContext& Context)
 		{
 			const uint64 Cycle = Context.EventData.GetValue<uint64>("EndTime");
-			const double Time = Context.EventTime.AsSeconds(Cycle);
-			UpdatedSessionTime(Cycle, Time);
+			UpdateThreadTime(Context, Cycle);
 		};
 	KnownEvents[(int)EKnownEvent::GpuProfiler_SignalFence].Callback =
 	KnownEvents[(int)EKnownEvent::GpuProfiler_WaitFence].Callback =
 		[this](const FOnEventContext& Context)
 		{
 			const uint64 Cycle = Context.EventData.GetValue<uint64>("CPUTimestamp");
-			const double Time = Context.EventTime.AsSeconds(Cycle);
-			UpdatedSessionTime(Cycle, Time);
+			UpdateThreadTime(Context, Cycle);
 		};
 
 	//////////////////////////////////////////////////
@@ -275,8 +269,7 @@ void FTrimAnalyzer::RegisterCallbacksForKnownEvents()
 		[this](const FOnEventContext& Context)
 		{
 			const uint64 Cycle = Context.EventData.GetValue<uint64>("Cycle");
-			const double Time = Context.EventTime.AsSeconds(Cycle);
-			UpdatedSessionTime(Cycle, Time);
+			UpdateThreadTime(Context, Cycle);
 		};
 }
 
@@ -381,6 +374,8 @@ void FTrimAnalyzer::ProcessKnownEvent(uint32 EventId, const ANSICHAR* LoggerName
 
 		BEGIN_KNOWN_LOGGER("CpuProfiler")
 			//KNOWN_EVENT("EventSpec", CpuProfiler_EventSpec)
+			//KNOWN_EVENT("MetadataSpec", CpuProfiler_MetadataSpec)
+			KNOWN_EVENT("Metadata", CpuProfiler_Metadata)
 			KNOWN_EVENT("EventBatchV3", CpuProfiler_EventBatchV3)
 			KNOWN_EVENT("EventBatchV2", CpuProfiler_EventBatchV2) // deprecated
 			KNOWN_EVENT("EventBatch", CpuProfiler_EventBatch) // deprecated
@@ -519,7 +514,8 @@ bool FTrimAnalyzer::OnEvent(uint16 RouteId, EStyle Style, const FOnEventContext&
 {
 	const FEventData& EventData = Context.EventData;
 
-	const uint32 ContextThreadId = Context.ThreadInfo.GetId();
+	const uint32 CurrentThreadId = Context.ThreadInfo.GetId();
+	FTrimAnalyzer::FThreadInfo& CurrentThread = GetOrAddThread(CurrentThreadId);
 
 	const FEventTypeInfo& TypeInfo = EventData.GetTypeInfo();
 
@@ -534,6 +530,19 @@ bool FTrimAnalyzer::OnEvent(uint16 RouteId, EStyle Style, const FOnEventContext&
 
 	bool bShouldWriteEvent = true;
 
+	//////////////////////////////////////////////////
+	// Processing known events
+
+	check(EventId <= MaxKnownEventId);
+	EKnownEvent KnownEvent = KnownEventMap[EventId];
+	check((int)KnownEvent < (int)EKnownEvent::Count);
+	const FKnownEvent& Event = KnownEvents[(int)KnownEvent];
+
+	if (Event.Callback)
+	{
+		Event.Callback(Context);
+	}
+
 	//////////////////////////////////////////////////
 	// Known $Trace.* events
 
@@ -570,7 +579,7 @@ bool FTrimAnalyzer::OnEvent(uint16 RouteId, EStyle Style, const FOnEventContext&
 		int32 SortHint = Context.EventData.GetValue<int32>("SortHint");
 		FAnsiStringView ThreadName;
 		Context.EventData.GetString("Name", ThreadName);
-		
+
 		const UE::Trace::FTraceWriterThreadInfo* ThreadInfo = TraceWriter.GetThreadInfo(ThreadId);
 		if (ThreadInfo != nullptr)
 		{
@@ -618,7 +627,7 @@ bool FTrimAnalyzer::OnEvent(uint16 RouteId, EStyle Style, const FOnEventContext&
 	{
 		if (Style == EStyle::EnterScope)
 		{
-			SetCurrentThread(ContextThreadId);
+			SetCurrentThread(CurrentThreadId);
 
 			uint64 Timestamp = Context.EventTime.GetTimestamp();
 			if (Timestamp != 0)
@@ -632,7 +641,7 @@ bool FTrimAnalyzer::OnEvent(uint16 RouteId, EStyle Style, const FOnEventContext&
 		}
 		else if (Style == EStyle::LeaveScope)
 		{
-			SetCurrentThread(ContextThreadId);
+			SetCurrentThread(CurrentThreadId);
 
 			uint64 Timestamp = Context.EventTime.GetTimestamp();
 			if (Timestamp != 0)
@@ -676,14 +685,50 @@ bool FTrimAnalyzer::OnEvent(uint16 RouteId, EStyle Style, const FOnEventContext&
 		}
 	}
 
+	//////////////////////////////////////////////////
+
+	bool bShouldStopAnalysis = false;
+
+	// When we encounter a Sync event...
+	if (TypeInfo.IsSync())
+	{
+		double ThreadMinTime = CurrentThread.CurrentTime;
+
+		// ...we can sync times between threads.
+		for (auto& KV : ThreadMap)
+		{
+			FThreadInfo& Thread = KV.Value;
+
+			// Threads with old timestamps can be updated with time of the CurrentThread.
+			if (Thread.CurrentCycle < CurrentThread.CurrentCycle)
+			{
+				Thread.CurrentCycle = CurrentThread.CurrentCycle;
+				Thread.CurrentTime = CurrentThread.CurrentTime;
+			}
+
+			// Compute also the min timestamp of all threads.
+			if (Thread.CurrentTime < ThreadMinTime)
+			{
+				ThreadMinTime = Thread.CurrentTime;
+			}
+		}
+
+		// If all threads have timestamps past the EndTime limit we can stop analysis.
+		if (ThreadMinTime > Parameters.EndTime)
+		{
+			bShouldStopAnalysis = true;
+		}
+	}
+
 	//////////////////////////////////////////////////
 	// Filters
 
 	// Only filter the non-important events.
-	if (bShouldWriteEvent && ContextThreadId != 1)
+	if (bShouldWriteEvent && CurrentThreadId != 1)
 	{
-		// Filter by start time.
-		if (SessionCurrentTime < Parameters.StartTime)
+		// Filter by time range.
+		if (CurrentThread.CurrentTime < Parameters.StartTime ||
+			CurrentThread.CurrentTime > Parameters.EndTime)
 		{
 			bShouldWriteEvent = false;
 		}
@@ -742,7 +787,7 @@ bool FTrimAnalyzer::OnEvent(uint16 RouteId, EStyle Style, const FOnEventContext&
 
 	if (bShouldWriteEvent)
 	{
-		SetCurrentThread(ContextThreadId);
+		SetCurrentThread(CurrentThreadId);
 
 		UE::Trace::FTraceWriterEventBuilder& EventBuilder = TraceWriter.WriteEvent(EventInfo->WriterEventId);
 
@@ -764,9 +809,8 @@ bool FTrimAnalyzer::OnEvent(uint16 RouteId, EStyle Style, const FOnEventContext&
 					FieldInfo.GetTypeSize() == 8 &&
 					FCStringAnsi::Strcmp(FieldInfo.GetName(), "Cycle") == 0)
 				{
-					const uint64 Value = EventData.GetValue<uint64>(FieldInfo.GetName());
-					const double Time = Context.EventTime.AsSeconds(Value);
-					UpdatedSessionTime(Value, Time);
+					const uint64 Cycle = EventData.GetValue<uint64>(FieldInfo.GetName());
+					UpdateThreadTime(CurrentThread, Context, Cycle);
 				}
 			}
 		}
@@ -774,28 +818,9 @@ bool FTrimAnalyzer::OnEvent(uint16 RouteId, EStyle Style, const FOnEventContext&
 		EventBuilder.End();
 	}
 
-	//////////////////////////////////////////////////
-	// Processing known events
-
-	check(EventId <= MaxKnownEventId);
-	EKnownEvent KnownEvent = KnownEventMap[EventId];
-	check((int)KnownEvent < (int)EKnownEvent::Count);
-	const FKnownEvent& Event = KnownEvents[(int)KnownEvent];
-
-	if (Event.Callback)
-	{
-		Event.Callback(Context);
-	}
-
 	//////////////////////////////////////////////////
 
-	// Filter by end time.
-	if (SessionCurrentTime > Parameters.EndTime)
-	{
-		return false; // stop further processing
-	}
-
-	return true;
+	return !bShouldStopAnalysis;
 }
 
 ////////////////////////////////////////////////////////////////////////////////////////////////////
@@ -1090,7 +1115,47 @@ void FTrimAnalyzer::SetCurrentThread(uint32 ThreadId)
 
 ////////////////////////////////////////////////////////////////////////////////////////////////////
 
-void FTrimAnalyzer::UpdatedSessionTime(uint64 Cycle, double Time)
+FTrimAnalyzer::FThreadInfo& FTrimAnalyzer::GetOrAddThread(uint32 ThreadId)
+{
+	FThreadInfo* FoundThread = ThreadMap.Find(ThreadId);
+	if (FoundThread)
+	{
+		return *FoundThread;
+	}
+	return ThreadMap.Add(ThreadId, { ThreadId });
+}
+
+////////////////////////////////////////////////////////////////////////////////////////////////////
+
+void FTrimAnalyzer::UpdateThreadTime(const FOnEventContext& Context, uint64 Cycle)
+{
+	const uint32 ThreadId = Context.ThreadInfo.GetId();
+	FThreadInfo& Thread = GetOrAddThread(ThreadId);
+	UpdateThreadTime(Thread, Context, Cycle);
+}
+
+////////////////////////////////////////////////////////////////////////////////////////////////////
+
+void FTrimAnalyzer::UpdateThreadTime(FThreadInfo& Thread, const FOnEventContext& Context, uint64 Cycle)
+{
+	if (Cycle > Thread.CurrentCycle)
+	{
+		const double Time = Context.EventTime.AsSeconds(Cycle);
+
+		Thread.CurrentCycle = Cycle;
+		Thread.CurrentTime = Time;
+
+		if (Cycle > SessionCurrentCycle)
+		{
+			SessionCurrentCycle = Cycle;
+			SessionCurrentTime = Time;
+		}
+	}
+}
+
+////////////////////////////////////////////////////////////////////////////////////////////////////
+
+void FTrimAnalyzer::UpdateSessionTime(uint64 Cycle, double Time)
 {
 	if (Cycle > SessionCurrentCycle)
 	{
@@ -1106,13 +1171,8 @@ void FTrimAnalyzer::ProcessCpuEventBatchV1(const FOnEventContext& Context, const
 	const FEventData& EventData = Context.EventData;
 
 	const uint32 ThreadId = GetThreadIdField(Context);
-
-	uint64* LastCyclePtr = ThreadMap.Find(ThreadId);
-	if (!LastCyclePtr)
-	{
-		LastCyclePtr = &ThreadMap.Add(ThreadId, 0);
-	}
-	uint64 LastCycle = *LastCyclePtr;
+	FThreadInfo& Thread = GetOrAddThread(ThreadId);
+	uint64 LastCycle = Thread.LastCycle;
 
 	const uint8* BufferEnd = BufferPtr + BufferSize;
 	while (BufferPtr < BufferEnd)
@@ -1130,9 +1190,12 @@ void FTrimAnalyzer::ProcessCpuEventBatchV1(const FOnEventContext& Context, const
 		//error
 	}
 
-	*LastCyclePtr = LastCycle;
-	const double Time = Context.EventTime.AsSeconds(LastCycle);
-	UpdatedSessionTime(LastCycle, Time);
+	Thread.LastCycle = LastCycle;
+	UpdateThreadTime(Thread, Context, LastCycle);
+	if (ThreadId != Context.ThreadInfo.GetId()) // ThreadId has bias (using ThreadId field)
+	{
+		UpdateThreadTime(Context, LastCycle);
+	}
 }
 
 ////////////////////////////////////////////////////////////////////////////////////////////////////
@@ -1141,15 +1204,9 @@ void FTrimAnalyzer::ProcessCpuEventBatchV2V3(const FOnEventContext& Context, con
 {
 	const FEventData& EventData = Context.EventData;
 
-	const uint32 ThreadId = GetThreadIdField(Context);
-
-	uint64* LastCyclePtr = ThreadMap.Find(ThreadId);
-	if (!LastCyclePtr)
-	{
-		LastCyclePtr = &ThreadMap.Add(ThreadId, 0);
-	}
-
-	uint64 LastCycle = *LastCyclePtr;
+	const uint32 ThreadId = Context.ThreadInfo.GetId();
+	FThreadInfo& Thread = GetOrAddThread(ThreadId);
+	uint64 LastCycle = Thread.LastCycle;
 
 	check(Context.EventTime.GetTimestamp() == 0);
 	const uint64 BaseCycle = Context.EventTime.AsCycle64();
@@ -1218,9 +1275,8 @@ void FTrimAnalyzer::ProcessCpuEventBatchV2V3(const FOnEventContext& Context, con
 		// error
 	}
 
-	*LastCyclePtr = LastCycle;
-	const double Time = Context.EventTime.AsSeconds(LastCycle);
-	UpdatedSessionTime(LastCycle, Time);
+	Thread.LastCycle = LastCycle;
+	UpdateThreadTime(Thread, Context, LastCycle);
 }
 
 ////////////////////////////////////////////////////////////////////////////////////////////////////
diff --git a/Engine/Source/Developer/TraceServices/Private/Analyzers/TrimAnalyzer.h b/Engine/Source/Developer/TraceServices/Private/Analyzers/TrimAnalyzer.h
--- a/Engine/Source/Developer/TraceServices/Private/Analyzers/TrimAnalyzer.h
+++ b/Engine/Source/Developer/TraceServices/Private/Analyzers/TrimAnalyzer.h
@@ -112,6 +112,8 @@ private:
 		GpuProfiler_WaitFence,
 
 		//CpuProfiler_EventSpec,
+		//CpuProfiler_MetadataSpec, // added in UE 5.x
+		CpuProfiler_Metadata, // added in UE 5.x
 		CpuProfiler_EventBatchV3, // added in UE 5.6
 		CpuProfiler_EventBatchV2, // backward compatibility; added in UE 5.1, removed in 5.6
 		CpuProfiler_EventBatch, // backward compatibility; removed in UE 5.1
@@ -147,6 +149,14 @@ private:
 		KnownEventCallback Callback = nullptr;
 	};
 
+	struct FThreadInfo
+	{
+		uint32 Id = 0;
+		uint64 LastCycle = 0; // for CpuProfiler.EventBatch decoding
+		uint64 CurrentCycle = 0;
+		double CurrentTime = 0.0;
+	};
+
 public:
 	FTrimAnalyzer(FTrimParameters& InParameters, UE::Trace::FTraceWriter& InTraceWriter);
 	virtual ~FTrimAnalyzer();
@@ -164,9 +174,11 @@ private:
 	void WriteArrayField(uint32 FieldIndex, const UE::Trace::IAnalyzer::FEventFieldInfo& FieldInfo, const UE::Trace::IAnalyzer::FEventData& EventData, UE::Trace::FTraceWriterEventBuilder& EventBuilder);
 	void WriteScalarField(uint32 FieldIndex, const UE::Trace::IAnalyzer::FEventFieldInfo& FieldInfo, const UE::Trace::IAnalyzer::FEventData& EventData, UE::Trace::FTraceWriterEventBuilder& EventBuilder);
 
+	FThreadInfo& GetOrAddThread(uint32 ThreadId);
 	void SetCurrentThread(uint32 ThreadId);
-
-	void UpdatedSessionTime(uint64 Cycle, double Time);
+	void UpdateThreadTime(const FOnEventContext& Context, uint64 Cycle);
+	void UpdateThreadTime(FThreadInfo& Thread, const FOnEventContext& Context, uint64 Cycle);
+	void UpdateSessionTime(uint64 Cycle, double Time);
 
 	bool IsKnownEvent(uint32 EventId, EKnownEvent KnownEvent) const
 	{
@@ -185,7 +197,7 @@ private:
 	uint32 TransportVersion = 4;
 	uint32 ProtocolVersion = 7;
 
-	TMap<uint32, uint64> ThreadMap;
+	TMap<uint32, FThreadInfo> ThreadMap; // ThreadId --> FThreadInfo
 	TMap<uint32, FEventInfo> EventInfos; // EventId --> FEventInfo
 	FEventInfo EmptyEventInfo;
 
