[Insights]
- TraceServices: Deprecated ITimingProfilerProvider::ReadTimer() in favor of newer ITimingProfilerProvider::GetTimerReader().
- TraceServices: Replaced uses of ITimingProfilerProvider:ReadTimer() with GetTimerReader() for CpuProfiler and GpuProfiler.

#Synced-CL 50345856

diff --git a/Engine/Source/Developer/TraceServices/Private/Analyzers/CpuProfilerTraceAnalysis.cpp b/Engine/Source/Developer/TraceServices/Private/Analyzers/CpuProfilerTraceAnalysis.cpp
--- a/Engine/Source/Developer/TraceServices/Private/Analyzers/CpuProfilerTraceAnalysis.cpp
+++ b/Engine/Source/Developer/TraceServices/Private/Analyzers/CpuProfilerTraceAnalysis.cpp
@@ -985,33 +985,25 @@ void FCpuProfilerAnalyzer::OnMetadataSpec(const FOnEventContext& Context)
 	if (Spec.FieldNames.Num() > 0 || Spec.Format != nullptr)
 	{
 		FAnalysisSessionEditScope _(Session);
+
 		uint32 MetadataSpecId = EditableTimingProfilerProvider.AddMetadataSpec(MoveTemp(Spec));
 
 		const ITimingProfilerProvider* TimingProfilerProvider = EditableTimingProfilerProvider.GetReadProvider();
 		if (TimingProfilerProvider)
 		{
-			FTimingProfilerTimer PreviousTimer;
-			TimingProfilerProvider->ReadTimers(
-				[TimerId, &PreviousTimer]
-				(const ITimingProfilerTimerReader& TimerReader)
-				{
-					const FTimingProfilerTimer* PreviousTimerPtr = TimerReader.GetTimer(TimerId);
-					if (PreviousTimerPtr)
-					{
-						PreviousTimer = *PreviousTimerPtr;
-					}
-				});
-			if (PreviousTimer.HasValidMetadataSpecId())
+			const ITimingProfilerTimerReader& TimerReader = TimingProfilerProvider->GetTimerReader();
+			const FTimingProfilerTimer* PreviousTimer = TimerReader.GetTimer(TimerId);
+			if (PreviousTimer && PreviousTimer->HasValidMetadataSpecId())
 			{
-				TimerId = AddCpuTimer(SpecId, *Name, PreviousTimer.File, PreviousTimer.Line);
+				TimerId = AddCpuTimer(SpecId, *Name, PreviousTimer->File, PreviousTimer->Line);
 
 				if (++NumTimerWarnings < NumMaxWarnings)
 				{
-					const FMetadataSpec* PreviousMetadataSpec = TimingProfilerProvider->GetMetadataSpec(PreviousTimer.MetadataSpecId);
+					const FMetadataSpec* PreviousMetadataSpec = TimingProfilerProvider->GetMetadataSpec(PreviousTimer->MetadataSpecId);
 					check(PreviousMetadataSpec);
 					UE_LOG(LogTraceServices, Warning, TEXT("[CpuProfiler] MetadataSpec -- de-duplicated timer for SpecId %u -- old: %u (%s) metadata=\"%s\" (%d fields) --> new: %u (%s) metadata=\"%s\" (%d fields) !!!"),
 						SpecId,
-						PreviousTimer.Id, PreviousTimer.Name, PreviousMetadataSpec->Format, PreviousMetadataSpec->FieldNames.Num(),
+						PreviousTimer->Id, PreviousTimer->Name, PreviousMetadataSpec->Format, PreviousMetadataSpec->FieldNames.Num(),
 						TimerId, *Name, Spec.Format, Spec.FieldNames.Num());
 				}
 			}
@@ -1142,22 +1134,17 @@ uint32 FCpuProfilerAnalyzer::DefineUniqueTimer(uint32 SpecId, const TCHAR* Timer
 
 const TCHAR* FCpuProfilerAnalyzer::GetTimerName(uint32 TimerId) const
 {
-	const TCHAR* TimerName = nullptr;
 	const ITimingProfilerProvider* TimingProfilerProvider = EditableTimingProfilerProvider.GetReadProvider();
 	if (TimingProfilerProvider)
 	{
-		TimingProfilerProvider->ReadTimers(
-			[TimerId, &TimerName]
-			(const ITimingProfilerTimerReader& TimerReader)
-			{
-				const FTimingProfilerTimer* Timer = TimerReader.GetTimer(TimerId);
-				if (Timer)
-				{
-					TimerName = Timer->Name;
-				}
-			});
+		const ITimingProfilerTimerReader& TimerReader = TimingProfilerProvider->GetTimerReader();
+		const FTimingProfilerTimer* Timer = TimerReader.GetTimer(TimerId);
+		if (Timer)
+		{
+			return Timer->Name;
+		}
 	}
-	return TimerName;
+	return nullptr;
 }
 
 ////////////////////////////////////////////////////////////////////////////////////////////////////
diff --git a/Engine/Source/Developer/TraceServices/Private/Analyzers/GpuProfilerTraceAnalysis.cpp b/Engine/Source/Developer/TraceServices/Private/Analyzers/GpuProfilerTraceAnalysis.cpp
--- a/Engine/Source/Developer/TraceServices/Private/Analyzers/GpuProfilerTraceAnalysis.cpp
+++ b/Engine/Source/Developer/TraceServices/Private/Analyzers/GpuProfilerTraceAnalysis.cpp
@@ -184,17 +184,13 @@ bool FGpuProfilerAnalyzer::OnEvent(uint16 RouteId, EStyle Style, const FOnEventC
 		else
 		{
 			FAnalysisSessionEditScope _(Session);
-			TimingProfilerProvider.ReadTimers(
-				[TimerId]
-				(const ITimingProfilerTimerReader& TimerReader)
-				{
-					const FTimingProfilerTimer* Timer = TimerReader.GetTimer(TimerId);
-					check(Timer);
-					if (Timer->HasValidMetadataSpecId())
-					{
-						UE_LOG(LogTraceServices, Warning, TEXT("[GpuProfiler] Empty metadata for a timing event while its timer %u (%s) has valid metadata spec id (%u) !!!"), TimerId, Timer->Name, Timer->MetadataSpecId);
-					}
-				});
+			const ITimingProfilerTimerReader& TimerReader = TimingProfilerProvider.GetTimerReader();
+			const FTimingProfilerTimer* Timer = TimerReader.GetTimer(TimerId);
+			check(Timer);
+			if (Timer->HasValidMetadataSpecId())
+			{
+				UE_LOG(LogTraceServices, Warning, TEXT("[GpuProfiler] Empty metadata for a timing event while its timer %u (%s) has valid metadata spec id (%u) !!!"), TimerId, Timer->Name, Timer->MetadataSpecId);
+			}
 		}
 #endif
 
@@ -607,30 +603,22 @@ void FGpuProfilerAnalyzer::OnEventBreadcrumbSpec(const FOnEventContext& Context)
 	if (Spec.FieldNames.Num() > 0 || Spec.Format != nullptr)
 	{
 		FAnalysisSessionEditScope _(Session);
+
 		uint32 MetadataSpecId = TimingProfilerProvider.AddMetadataSpec(MoveTemp(Spec));
 
-		FTimingProfilerTimer PreviousTimer;
-		TimingProfilerProvider.ReadTimers(
-			[TimerId, &PreviousTimer]
-			(const ITimingProfilerTimerReader& TimerReader)
-			{
-				const FTimingProfilerTimer* PreviousTimerPtr = TimerReader.GetTimer(TimerId);
-				if (PreviousTimerPtr)
-				{
-					PreviousTimer = *PreviousTimerPtr;
-				}
-			});
-		if (PreviousTimer.HasValidMetadataSpecId())
+		const ITimingProfilerTimerReader& TimerReader = TimingProfilerProvider.GetTimerReader();
+		const FTimingProfilerTimer* PreviousTimer = TimerReader.GetTimer(TimerId);
+		if (PreviousTimer && PreviousTimer->HasValidMetadataSpecId())
 		{
 			TimerId = AddGpuTimer(SpecId, Name);
 
 			if (++ErrorData.NumTimerWarnings < ErrorData.NumMaxWarnings)
 			{
-				const FMetadataSpec* PreviousMetadataSpec = TimingProfilerProvider.GetMetadataSpec(PreviousTimer.MetadataSpecId);
+				const FMetadataSpec* PreviousMetadataSpec = TimingProfilerProvider.GetMetadataSpec(PreviousTimer->MetadataSpecId);
 				check(PreviousMetadataSpec);
 				UE_LOG(LogTraceServices, Warning, TEXT("[GpuProfiler] EventBreadcrumbSpec -- de-duplicated timer for SpecId %u -- old: %u (%s) metadata=\"%s\" (%d fields) --> new: %u (%s) metadata=\"%s\" (%d fields) !!!"),
 					SpecId,
-					PreviousTimer.Id, PreviousTimer.Name, PreviousMetadataSpec->Format, PreviousMetadataSpec->FieldNames.Num(),
+					PreviousTimer->Id, PreviousTimer->Name, PreviousMetadataSpec->Format, PreviousMetadataSpec->FieldNames.Num(),
 					TimerId, *Name, Spec.Format, Spec.FieldNames.Num());
 			}
 		}
@@ -669,16 +657,12 @@ uint32 FGpuProfilerAnalyzer::AddGpuTimer(uint32 SpecId, const FString& Breadcrum
 	{
 		FAnalysisSessionEditScope _(Session);
 		TimerId = TimingProfilerProvider.AddGpuTimer(Breadcrumb);
-		TimingProfilerProvider.ReadTimers(
-			[TimerId, &TimerName]
-			(const ITimingProfilerTimerReader& TimerReader)
-			{
-				const FTimingProfilerTimer* Timer = TimerReader.GetTimer(TimerId);
-				if (Timer)
-				{
-					TimerName = Timer->Name;
-				}
-			});
+		const ITimingProfilerTimerReader& TimerReader = TimingProfilerProvider.GetTimerReader();
+		const FTimingProfilerTimer* Timer = TimerReader.GetTimer(TimerId);
+		if (Timer)
+		{
+			TimerName = Timer->Name;
+		}
 	}
 
 	TimerMap.Add(TimerId, TimerName);
diff --git a/Engine/Source/Developer/TraceServices/Public/TraceServices/Model/TimingProfiler.h b/Engine/Source/Developer/TraceServices/Public/TraceServices/Model/TimingProfiler.h
--- a/Engine/Source/Developer/TraceServices/Public/TraceServices/Model/TimingProfiler.h
+++ b/Engine/Source/Developer/TraceServices/Public/TraceServices/Model/TimingProfiler.h
@@ -285,7 +285,8 @@ public:
 	//////////////////////////////////////////////////
 	// Timers
 
-	virtual void ReadTimers(TFunctionRef<void(const ITimingProfilerTimerReader&)> Callback) const = 0;
+	UE_DEPRECATED(5.8, "Use GetTimerReader() instead")
+	virtual void ReadTimers(TFunctionRef<void(const ITimingProfilerTimerReader&)> Callback) const {}
 
 	/**
 	 * Gets the TimerReader provider.
@@ -295,8 +296,13 @@ public:
 	 */
 	virtual const ITimingProfilerTimerReader& GetTimerReader() const
 	{
+		// The concrete class that implements ITimingProfilerProvider should override GetTimerReader()!
+		// GetTimerReader() is not made abstract for deprecation / backward compatibility reasons.
+		ensure(!"Please override GetTimerReader()!");
 		const ITimingProfilerTimerReader* TimerReader = nullptr;
+		PRAGMA_DISABLE_DEPRECATION_WARNINGS
 		ReadTimers([&TimerReader](const TraceServices::ITimingProfilerTimerReader& Out) { TimerReader = &Out; });
+		PRAGMA_ENABLE_DEPRECATION_WARNINGS
 		check(TimerReader != nullptr);
 		return *TimerReader;
 	}
