[Insights]
- TraceServices: Fixed aggregation for large number of timers (exceeded memory limit for TMap; now changed to use TSparseMap).
- TraceServices: Improved memory usage when not computing median (to not allocate histogram data).

#Synced-CL 50345856

diff --git a/Engine/Source/Developer/TraceServices/Private/Common/TimelineStatistics.cpp b/Engine/Source/Developer/TraceServices/Private/Common/TimelineStatistics.cpp
new file mode 100644
--- /dev/null
+++ b/Engine/Source/Developer/TraceServices/Private/Common/TimelineStatistics.cpp
@@ -0,0 +1,231 @@
+// Copyright Epic Games, Inc. All Rights Reserved.
+
+#pragma once
+
+#include "TimelineStatistics.h"
+
+namespace TraceServices
+{
+
+////////////////////////////////////////////////////////////////////////////////////////////////////
+// FInternalHistogram
+////////////////////////////////////////////////////////////////////////////////////////////////////
+
+class FInternalHistogram
+{
+private:
+	static constexpr int32 HistogramLen = 100;
+
+public:
+	FInternalHistogram()
+	{
+	}
+
+	double GetBucketSize() const
+	{
+		return BucketSize;
+	}
+
+	uint64 GetTotalCount() const
+	{
+		return TotalCount;
+	}
+
+	void Init(double MinTime, double MaxTime);
+	void AddValue(double RelativeTime);
+	double ComputeMedian() const;
+
+private:
+	uint32 Histogram[HistogramLen];
+	double BucketSize = 0.0;
+	uint64 TotalCount = 0;
+};
+
+////////////////////////////////////////////////////////////////////////////////////////////////////
+
+void FInternalHistogram::Init(double MinTime, double MaxTime)
+{
+	FMemory::Memzero(Histogram, sizeof(int32) * HistogramLen);
+
+	// Each bucket (Histogram[i]) will be centered on a value.
+	// I.e. First bucket (bucket 0) is centered on Min value: [Min-BucketSize/2, Min+BucketSize/2)
+	// and last bucket (bucket N-1) is centered on Max value: [Max-BucketSize/2, Max+BucketSize/2).
+	if (MaxTime == MinTime)
+	{
+		BucketSize = 1.0; // single large bucket
+	}
+	else
+	{
+		constexpr double InvHistogramLen = 1.0 / double(HistogramLen - 1);
+		BucketSize = (MaxTime - MinTime) * InvHistogramLen;
+	}
+
+	TotalCount = 0;
+}
+
+////////////////////////////////////////////////////////////////////////////////////////////////////
+
+void FInternalHistogram::AddValue(double RelativeTime)
+{
+	int32 Index = static_cast<int32>((RelativeTime + BucketSize / 2) / BucketSize);
+	ensure(Index >= 0);
+	if (Index < 0)
+	{
+		Index = 0;
+	}
+	ensure(Index < HistogramLen);
+	if (Index >= HistogramLen)
+	{
+		Index = HistogramLen - 1;
+	}
+	Histogram[Index]++;
+	TotalCount++;
+}
+
+////////////////////////////////////////////////////////////////////////////////////////////////////
+
+double FInternalHistogram::ComputeMedian() const
+{
+	const uint64 HalfCount = TotalCount / 2;
+	uint64 Count = 0;
+	for (int32 HistogramIndex = 0; HistogramIndex < HistogramLen; HistogramIndex++)
+	{
+		Count += Histogram[HistogramIndex];
+		if (Count > HalfCount)
+		{
+			double Median = HistogramIndex * BucketSize;
+			if (HistogramIndex > 0 &&
+				TotalCount % 2 == 0 &&
+				Count - Histogram[HistogramIndex] == HalfCount)
+			{
+				const double PrevMedian = (HistogramIndex - 1) * BucketSize;
+				Median = (Median + PrevMedian) / 2;
+			}
+			return Median;
+		}
+	}
+	return 0.0;
+}
+
+////////////////////////////////////////////////////////////////////////////////////////////////////
+// FInternalAggregationEntry
+////////////////////////////////////////////////////////////////////////////////////////////////////
+
+FTimelineStatistics::FInternalAggregationEntry::~FInternalAggregationEntry()
+{
+	if (InclHistogram)
+	{
+		InclHistogram->~FInternalHistogram(); // destruct
+		InclHistogram = nullptr;
+	}
+	if (ExclHistogram)
+	{
+		ExclHistogram->~FInternalHistogram(); // destruct
+		ExclHistogram = nullptr;
+	}
+}
+
+////////////////////////////////////////////////////////////////////////////////////////////////////
+// FTimelineStatistics
+////////////////////////////////////////////////////////////////////////////////////////////////////
+
+void FTimelineStatistics::UpdateTotalMinMaxTimerAggregationStats(FTimelineStatistics::FInternalAggregationEntry& AggregationEntry, double InclTime, double ExclTime)
+{
+	UpdateTotalMinMaxTimerStats(AggregationEntry.Stats, InclTime, ExclTime);
+}
+
+////////////////////////////////////////////////////////////////////////////////////////////////////
+
+void FTimelineStatistics::UpdateTotalMinMaxTimerStats(FAggregatedTimingStats& Stats, double InclTime, double ExclTime)
+{
+	Stats.TotalInclusiveTime += InclTime;
+	if (InclTime < Stats.MinInclusiveTime)
+	{
+		Stats.MinInclusiveTime = InclTime;
+	}
+	if (InclTime > Stats.MaxInclusiveTime)
+	{
+		Stats.MaxInclusiveTime = InclTime;
+	}
+
+	Stats.TotalExclusiveTime += ExclTime;
+	if (ExclTime < Stats.MinExclusiveTime)
+	{
+		Stats.MinExclusiveTime = ExclTime;
+	}
+	if (ExclTime > Stats.MaxExclusiveTime)
+	{
+		Stats.MaxExclusiveTime = ExclTime;
+	}
+
+	Stats.InstanceCount++;
+}
+
+////////////////////////////////////////////////////////////////////////////////////////////////////
+
+void FTimelineStatistics::InitHistogram(FTimelineStatistics::FInternalAggregationEntry& AggregationEntry, ILinearAllocator& Allocator)
+{
+	const FAggregatedTimingStats& Stats = AggregationEntry.Stats;
+
+	check(AggregationEntry.InclHistogram == nullptr);
+	AggregationEntry.InclHistogram = (FInternalHistogram*)Allocator.Allocate(sizeof(FInternalHistogram));
+	new (AggregationEntry.InclHistogram) FInternalHistogram(); // construct
+	AggregationEntry.InclHistogram->Init(Stats.MinInclusiveTime, Stats.MaxInclusiveTime);
+
+	check(AggregationEntry.ExclHistogram == nullptr);
+	AggregationEntry.ExclHistogram = (FInternalHistogram*)Allocator.Allocate(sizeof(FInternalHistogram));
+	new (AggregationEntry.ExclHistogram) FInternalHistogram(); // construct
+	AggregationEntry.ExclHistogram->Init(Stats.MinExclusiveTime, Stats.MaxExclusiveTime);
+}
+
+////////////////////////////////////////////////////////////////////////////////////////////////////
+
+void FTimelineStatistics::UpdateHistogram(FTimelineStatistics::FInternalAggregationEntry& AggregationEntry, double InclTime, double ExclTime)
+{
+	const FAggregatedTimingStats& Stats = AggregationEntry.Stats;
+
+	check(AggregationEntry.InclHistogram != nullptr);
+	AggregationEntry.InclHistogram->AddValue(InclTime - Stats.MinInclusiveTime);
+
+	check(AggregationEntry.ExclHistogram != nullptr);
+	AggregationEntry.ExclHistogram->AddValue(ExclTime - Stats.MinExclusiveTime);
+}
+
+////////////////////////////////////////////////////////////////////////////////////////////////////
+
+void FTimelineStatistics::ComputeMedianFromHistogram(FTimelineStatistics::FInternalAggregationEntry& AggregationEntry)
+{
+	FAggregatedTimingStats& Stats = AggregationEntry.Stats;
+
+	check(AggregationEntry.InclHistogram != nullptr);
+	Stats.MedianInclusiveTime = Stats.MinInclusiveTime + AggregationEntry.InclHistogram->ComputeMedian();
+
+	check(AggregationEntry.ExclHistogram != nullptr);
+	Stats.MedianExclusiveTime = Stats.MinExclusiveTime + AggregationEntry.ExclHistogram->ComputeMedian();
+}
+
+////////////////////////////////////////////////////////////////////////////////////////////////////
+
+void FTimelineStatistics::PostProcessTimerStats(FTimelineStatistics::FInternalAggregationEntry& AggregationEntry, bool bComputeMedian)
+{
+	FAggregatedTimingStats& Stats = AggregationEntry.Stats;
+
+	// Compute average inclusive/exclusive times.
+	ensure(Stats.InstanceCount > 0);
+	double InvCount = 1.0f / static_cast<double>(Stats.InstanceCount);
+	Stats.AverageInclusiveTime = Stats.TotalInclusiveTime * InvCount;
+	Stats.AverageExclusiveTime = Stats.TotalExclusiveTime * InvCount;
+	Stats.AverageInstanceCount = static_cast<double>(Stats.InstanceCount);
+
+	if (bComputeMedian)
+	{
+		ComputeMedianFromHistogram(AggregationEntry);
+		check(AggregationEntry.InclHistogram->GetTotalCount() == Stats.InstanceCount);
+		check(AggregationEntry.ExclHistogram->GetTotalCount() == Stats.InstanceCount);
+	}
+}
+
+////////////////////////////////////////////////////////////////////////////////////////////////////
+////////////////////////////////////////////////////////////////////////////////////////////////////
+
+} // namespace TraceServices
diff --git a/Engine/Source/Developer/TraceServices/Private/Common/TimelineStatistics.h b/Engine/Source/Developer/TraceServices/Private/Common/TimelineStatistics.h
--- a/Engine/Source/Developer/TraceServices/Private/Common/TimelineStatistics.h
+++ b/Engine/Source/Developer/TraceServices/Private/Common/TimelineStatistics.h
@@ -2,16 +2,22 @@
 
 #pragma once
 
+#include "CoreTypes.h"
+
 #include "Containers/Map.h"
 #include "Containers/Queue.h"
-#include "CoreTypes.h"
-#include <cfloat>
 #include "Misc/QueuedThreadPool.h"
 #include "Tasks/Task.h"
 #include "Templates/SharedPointer.h"
+
+// TraceServices
 #include "TraceServices/Common/CancellationToken.h"
+#include "TraceServices/Containers/Allocators.h"
+#include "TraceServices/Containers/SlabAllocator.h"
 #include "TraceServices/Containers/Timelines.h"
 
+#include <cfloat>
+
 namespace TraceServices
 {
 
@@ -40,15 +46,68 @@ struct FFrameData
 class FTimelineStatistics
 {
 public:
-	template<typename TimelineType, typename BucketMappingFunc, typename BucketKeyType>
-	static void CreateAggregation(const TArray<const TimelineType*>& Timelines,
-								  BucketMappingFunc BucketMapper,
-								  double IntervalStart,
-								  double IntervalEnd,
-								  TSharedPtr<FCancellationToken> CancellationToken,
-								  TMap<BucketKeyType, FAggregatedTimingStats>& Result)
+	template<typename KeyType, typename ValueType>
+	class TAggregationResult
 	{
-		TMap<BucketKeyType, FInternalAggregationEntry> InternalResult;
+	private:
+		static constexpr uint64 AllocatorSlabSize = 1 * 1024 * 1024;
+
+	public:
+		TAggregationResult()
+			: Allocator(AllocatorSlabSize)
+		{
+		}
+
+		TAggregationResult(TSparseMap<KeyType, ValueType>& InMap)
+			: Allocator(AllocatorSlabSize)
+			, Map(InMap)
+		{
+		}
+
+		// Non-copyable
+		TAggregationResult(const TAggregationResult&) = delete;
+		TAggregationResult& operator=(const TAggregationResult&) = delete;
+
+		~TAggregationResult()
+		{
+		}
+
+		ILinearAllocator& GetAllocator()
+		{
+			return Allocator;
+		}
+
+		TSparseMap<KeyType, ValueType>& GetMap()
+		{
+			return Map;
+		}
+
+		const TSparseMap<KeyType, ValueType>& GetMap() const
+		{
+			return Map;
+		}
+
+	private:
+		FSlabAllocator Allocator; // to allocate extra data for histograms
+		TSparseMap<KeyType, ValueType> Map;
+	};
+
+public:
+	template<
+		typename TimelineType,
+		typename BucketMappingFunc,
+		typename BucketKeyType
+	>
+	static void CreateAggregation(
+		const TArray<const TimelineType*>& Timelines,
+		BucketMappingFunc BucketMapper,
+		double IntervalStart,
+		double IntervalEnd,
+		TSharedPtr<FCancellationToken> CancellationToken,
+		TAggregationResult<BucketKeyType, FAggregatedTimingStats>& Result)
+	{
+		TAggregationResult<BucketKeyType, FInternalAggregationEntry> InternalResult;
+
 		// Compute instance count and total/min/max inclusive/exclusive times for each timer.
 		for (const TimelineType* Timeline : Timelines)
 		{
@@ -58,14 +117,15 @@ public:
 			}
 			ProcessTimeline(Timeline, BucketMapper, UpdateTotalMinMaxTimerAggregationStats, IntervalStart, IntervalEnd, InternalResult);
 		}
+
 		// Now, as we know min/max inclusive/exclusive times for each timer, we can compute histogram and median values.
-		const bool ComputeMedian = true;
-		if (ComputeMedian)
+		const bool bComputeMedian = true;
+		if (bComputeMedian)
 		{
 			// Update bucket size (DT) for computing histogram.
-			for (auto& KV : InternalResult)
+			for (auto& KV : InternalResult.GetMap())
 			{
-				PreComputeHistogram(KV.Value);
+				InitHistogram(KV.Value, InternalResult.GetAllocator());
 			}
 
 			// Compute histogram.
@@ -75,31 +135,42 @@ public:
 				{
 					return;
 				}
-
-				ProcessTimeline(Timeline, BucketMapper, UpdateHistogramForTimerStats, IntervalStart, IntervalEnd, InternalResult);
+				ProcessTimeline(Timeline, BucketMapper, UpdateHistogram, IntervalStart, IntervalEnd, InternalResult);
 			}
 		}
 
 		// Compute average and median inclusive/exclusive times.
-		for (auto& KV : InternalResult)
+		for (auto& KV : InternalResult.GetMap())
 		{
-			PostProcessTimerStats(KV.Value, ComputeMedian);
-			Result.Add(KV.Key, KV.Value.Stats);
+			PostProcessTimerStats(KV.Value, bComputeMedian);
+			Result.GetMap().Add(KV.Key, KV.Value.Stats);
 		}
 	};
 
-	template<typename TimelineType, typename BucketMappingFunc, typename BucketKeyType>
-	static void CreateFrameStatsAggregation(const TArray<const TimelineType*>& Timelines,
-											BucketMappingFunc BucketMapper,
-											const TArray<FFrameData>& Frames,
-											TSharedPtr<FCancellationToken> CancellationToken,
-											TMap<BucketKeyType, FAggregatedTimingStats>& Result,
-											double RatioOfThreadsToUse)
+	template<
+		typename TimelineType,
+		typename BucketMappingFunc,
+		typename BucketKeyType
+	>
+	static void CreateFrameStatsAggregation(
+		const TArray<const TimelineType*>& Timelines,
+		BucketMappingFunc BucketMapper,
+		const TArray<FFrameData>& Frames,
+		TSharedPtr<FCancellationToken> CancellationToken,
+		TAggregationResult<BucketKeyType, FAggregatedTimingStats>& Result,
+		double RatioOfThreadsToUse)
 	{
-		TMap<BucketKeyType, FInternalFrameAggregationEntry> GlobalResult;
-		TMap<BucketKeyType, FAggregatedTimingStats> InitialTimerMap;
+		int32 FramesNum = Frames.Num();
+		if (FramesNum == 0)
+		{
+			return;
+		}
+
+		TAggregationResult<BucketKeyType, FInternalFrameAggregationEntry> GlobalResult;
+		TAggregationResult<BucketKeyType, FAggregatedTimingStats> InitialTimerMap;
+
 		double StartTime = Frames[0].StartTime;
-		double EndTime = Frames[Frames.Num() - 1].EndTime;
+		double EndTime = Frames[FramesNum - 1].EndTime;
 
 		// Gather the keys for all timers so we have a stable map (no insertions necessary) for the next step.
 		// This is to guarantee that GlobalResult and FrameResult have the same iteration order.
@@ -108,12 +179,11 @@ public:
 			GatherKeysFromTimeline(Timeline, BucketMapper, StartTime, EndTime, InitialTimerMap);
 		}
 
-		int32 FramesNum = Frames.Num();
-
 		// For very large numbers of timers or frames, an out of memory is possible so don't compute the median.
 		constexpr int64 MaxSize = 5 * 100 * 1000 * 1000;
-		bool bComputeMedian = InitialTimerMap.Num() * (int64)FramesNum < MaxSize;
-		for (auto& KV : InitialTimerMap)
+		bool bComputeMedian = InitialTimerMap.GetMap().Num() * (int64)FramesNum < MaxSize;
+
+		for (auto& KV : InitialTimerMap.GetMap())
 		{
 			FInternalFrameAggregationEntry Entry;
 			if (bComputeMedian)
@@ -122,11 +192,11 @@ public:
 				Entry.FrameExclusiveTimes.AddUninitialized(Frames.Num());
 			}
 			Entry.Inner.Stats = KV.Value;
-			GlobalResult.Add(KV.Key, MoveTemp(Entry));
+			GlobalResult.GetMap().Add(KV.Key, MoveTemp(Entry));
 		}
 
-		TQueue<TSharedPtr<TMap<BucketKeyType, FAggregatedTimingStats>>, EQueueMode::Mpsc> FrameResultsQueue;
-		
+		TQueue<TSharedPtr<TAggregationResult<BucketKeyType, FAggregatedTimingStats>>, EQueueMode::Mpsc> FrameResultsQueue;
+
 		int32 NumTasks = FMath::Max(1, (int32)((float)GThreadPool->GetNumThreads() * RatioOfThreadsToUse));
 		NumTasks = FMath::Min((int32)NumTasks, Frames.Num());
 		int32 FramesPerTask = Frames.Num() / NumTasks;
@@ -139,65 +209,69 @@ public:
 		TArray<UE::Tasks::TTask<void>> AsyncTasks;
 		for (int32 Index = 0; Index < NumTasks; ++Index)
 		{
-			AsyncTasks.Add(UE::Tasks::Launch(UE_SOURCE_LOCATION, [StartIndex, EndIndex, &Timelines, bComputeMedian, BucketMapper, &Frames, &FrameResultsQueue, &InitialTimerMap, &GlobalResult, CancellationToken]()
-			{
-				if (CancellationToken.IsValid() && CancellationToken->ShouldCancel())
-				{
-					return;
-				}
-
-				TSharedPtr<TMap<BucketKeyType, FAggregatedTimingStats>> TaskResult = MakeShared<TMap<BucketKeyType, FAggregatedTimingStats>>(InitialTimerMap);
-				TMap<BucketKeyType, FAggregatedTimingStats> FrameResult(InitialTimerMap);
-
-				for(int FrameIndex = StartIndex; FrameIndex < EndIndex; ++FrameIndex)
+			AsyncTasks.Add(UE::Tasks::Launch(UE_SOURCE_LOCATION,
+				[StartIndex, EndIndex, &Timelines, bComputeMedian, BucketMapper, &Frames, &FrameResultsQueue, &InitialTimerMap, &GlobalResult, CancellationToken]()
 				{
 					if (CancellationToken.IsValid() && CancellationToken->ShouldCancel())
 					{
 						return;
 					}
 
-					// Compute instance count and total/min/max inclusive/exclusive times for each timer.
-					for (const TimelineType* Timeline : Timelines)
-					{
-						ProcessTimelineForFrameStats(Timeline, BucketMapper, UpdateTotalMinMaxTimerStats, Frames[FrameIndex].StartTime, Frames[FrameIndex].EndTime, FrameResult);
-					}
+					TSharedPtr<TAggregationResult<BucketKeyType, FAggregatedTimingStats>> TaskResult = MakeShared<TAggregationResult<BucketKeyType, FAggregatedTimingStats>>(InitialTimerMap.GetMap());
+					TAggregationResult<BucketKeyType, FAggregatedTimingStats> FrameResult(InitialTimerMap.GetMap());
 
-					typename TMap<BucketKeyType, FAggregatedTimingStats>::TIterator FrameResultIterator(FrameResult);
-					typename TMap<BucketKeyType, FAggregatedTimingStats>::TIterator TaskResultIterator(*TaskResult);
-					typename TMap<BucketKeyType, FInternalFrameAggregationEntry>::TIterator GlobalResultIterator(GlobalResult);
-					while (FrameResultIterator)
+					for (int FrameIndex = StartIndex; FrameIndex < EndIndex; ++FrameIndex)
 					{
-						FAggregatedTimingStats& FrameStats = FrameResultIterator->Value;
-						FAggregatedTimingStats& TaskResultStats = TaskResultIterator->Value;
+						if (CancellationToken.IsValid() && CancellationToken->ShouldCancel())
+						{
+							return;
+						}
 
-						TaskResultStats.InstanceCount += FrameStats.InstanceCount;
+						// Compute instance count and total/min/max inclusive/exclusive times for each timer.
+						for (const TimelineType* Timeline : Timelines)
+						{
+							ProcessTimelineForFrameStats(Timeline, BucketMapper, UpdateTotalMinMaxTimerStats, Frames[FrameIndex].StartTime, Frames[FrameIndex].EndTime, FrameResult);
+						}
 
-						TaskResultStats.TotalInclusiveTime += FrameStats.TotalInclusiveTime;
-						TaskResultStats.MinInclusiveTime = FMath::Min(TaskResultStats.MinInclusiveTime, FrameStats.TotalInclusiveTime);
-						TaskResultStats.MaxInclusiveTime = FMath::Max(TaskResultStats.MaxInclusiveTime, FrameStats.TotalInclusiveTime);
+						typename TSparseMap<BucketKeyType, FAggregatedTimingStats>::TIterator FrameResultIterator(FrameResult.GetMap());
+						typename TSparseMap<BucketKeyType, FAggregatedTimingStats>::TIterator TaskResultIterator(TaskResult->GetMap());
+						typename TSparseMap<BucketKeyType, FInternalFrameAggregationEntry>::TIterator GlobalResultIterator(GlobalResult.GetMap());
+						while (FrameResultIterator)
+						{
+							check(FrameResultIterator->Key == TaskResultIterator->Key);
+							check(GlobalResultIterator->Key == TaskResultIterator->Key);
 
-						TaskResultStats.TotalExclusiveTime += FrameStats.TotalExclusiveTime;
-						TaskResultStats.MinExclusiveTime = FMath::Min(TaskResultStats.MinExclusiveTime, FrameStats.TotalExclusiveTime);
-						TaskResultStats.MaxExclusiveTime = FMath::Max(TaskResultStats.MaxExclusiveTime, FrameStats.TotalExclusiveTime);
+							FAggregatedTimingStats& FrameStats = FrameResultIterator->Value;
+							FAggregatedTimingStats& TaskResultStats = TaskResultIterator->Value;
 
-						if (bComputeMedian)
-						{
-							GlobalResultIterator->Value.FrameInclusiveTimes[FrameIndex] = FrameStats.TotalInclusiveTime;
-							GlobalResultIterator->Value.FrameExclusiveTimes[FrameIndex] = FrameStats.TotalExclusiveTime;
-						}
+							TaskResultStats.InstanceCount += FrameStats.InstanceCount;
 
-						// Reset the per frame stats.
-						FrameStats = FAggregatedTimingStats();
+							TaskResultStats.TotalInclusiveTime += FrameStats.TotalInclusiveTime;
+							TaskResultStats.MinInclusiveTime = FMath::Min(TaskResultStats.MinInclusiveTime, FrameStats.TotalInclusiveTime);
+							TaskResultStats.MaxInclusiveTime = FMath::Max(TaskResultStats.MaxInclusiveTime, FrameStats.TotalInclusiveTime);
 
-						++FrameResultIterator;
-						++TaskResultIterator;
-						++GlobalResultIterator;
-					}
+							TaskResultStats.TotalExclusiveTime += FrameStats.TotalExclusiveTime;
+							TaskResultStats.MinExclusiveTime = FMath::Min(TaskResultStats.MinExclusiveTime, FrameStats.TotalExclusiveTime);
+							TaskResultStats.MaxExclusiveTime = FMath::Max(TaskResultStats.MaxExclusiveTime, FrameStats.TotalExclusiveTime);
 
-				}
+							if (bComputeMedian)
+							{
+								FInternalFrameAggregationEntry& GlobalResultValue = GlobalResultIterator->Value;
+								GlobalResultValue.FrameInclusiveTimes[FrameIndex] = FrameStats.TotalInclusiveTime;
+								GlobalResultValue.FrameExclusiveTimes[FrameIndex] = FrameStats.TotalExclusiveTime;
+							}
 
-				FrameResultsQueue.Enqueue(MoveTemp(TaskResult));
-			}));
+							// Reset the per frame stats.
+							FrameStats = FAggregatedTimingStats();
+
+							++FrameResultIterator;
+							++TaskResultIterator;
+							++GlobalResultIterator;
+						}
+					}
+
+					FrameResultsQueue.Enqueue(MoveTemp(TaskResult));
+				}));
 
 			StartIndex = EndIndex;
 			EndIndex = StartIndex + FramesPerTask + (ExtraFrameTasks > 0 ? 1 : 0);
@@ -221,14 +295,16 @@ public:
 				continue;
 			}
 
-			TSharedPtr<TMap<BucketKeyType, FAggregatedTimingStats>> TaskResult;
+			TSharedPtr<TAggregationResult<BucketKeyType, FAggregatedTimingStats>> TaskResult;
 			ensure(FrameResultsQueue.Dequeue(TaskResult));
 			++ProcessedResults;
 
-			typename TMap<BucketKeyType, FAggregatedTimingStats>::TIterator TaskResultIterator(*TaskResult);
-			typename TMap<BucketKeyType, FInternalFrameAggregationEntry>::TIterator ResultIterator(GlobalResult);
+			typename TSparseMap<BucketKeyType, FAggregatedTimingStats>::TIterator TaskResultIterator(TaskResult->GetMap());
+			typename TSparseMap<BucketKeyType, FInternalFrameAggregationEntry>::TIterator ResultIterator(GlobalResult.GetMap());
 			while (ResultIterator)
 			{
+				check(ResultIterator->Key == TaskResultIterator->Key);
+
 				FAggregatedTimingStats& TaskStats = TaskResultIterator->Value;
 				FAggregatedTimingStats& ResultStats = ResultIterator->Value.Inner.Stats;
 
@@ -250,60 +326,53 @@ public:
 		if (bComputeMedian)
 		{
 			// Compute the median inclusive and exclusive time.
-			for (auto& KV : GlobalResult)
+			for (auto& KV : GlobalResult.GetMap())
 			{
-				if (KV.Value.Inner.Stats.InstanceCount == 0)
+				FInternalFrameAggregationEntry& FrameEntry = KV.Value;
+				FInternalAggregationEntry& AggregationEntry = FrameEntry.Inner;
+
+				if (AggregationEntry.Stats.InstanceCount == 0)
 				{
 					// Some timers might have 0 instances because the gathering phase also searched between frames.
 					continue;
 				}
-				PreComputeHistogram(KV.Value.Inner);
 
-				for (int Index = 0; Index < FramesNum; ++Index)
+				InitHistogram(AggregationEntry, GlobalResult.GetAllocator());
+
+				for (int32 Index = 0; Index < FramesNum; ++Index)
 				{
-					UpdateHistogramForTimerStats(KV.Value.Inner, KV.Value.FrameInclusiveTimes[Index], KV.Value.FrameExclusiveTimes[Index]);
+					UpdateHistogram(AggregationEntry, FrameEntry.FrameInclusiveTimes[Index], FrameEntry.FrameExclusiveTimes[Index]);
 				}
 
-				ComputeFrameStatsMedian(KV.Value.Inner, FramesNum);
+				ComputeMedianFromHistogram(AggregationEntry);
 			}
 		}
 
-		for (auto& KV : GlobalResult)
+		for (auto& KV : GlobalResult.GetMap())
 		{
-			KV.Value.Inner.Stats.AverageInclusiveTime = KV.Value.Inner.Stats.TotalInclusiveTime / FramesNum;
-			KV.Value.Inner.Stats.AverageExclusiveTime = KV.Value.Inner.Stats.TotalExclusiveTime / FramesNum;
-			KV.Value.Inner.Stats.AverageInstanceCount = static_cast<double>(KV.Value.Inner.Stats.InstanceCount) / FramesNum;
+			FAggregatedTimingStats& Stats = KV.Value.Inner.Stats;
+
+			Stats.AverageInclusiveTime = Stats.TotalInclusiveTime / FramesNum;
+			Stats.AverageExclusiveTime = Stats.TotalExclusiveTime / FramesNum;
+			Stats.AverageInstanceCount = static_cast<double>(Stats.InstanceCount) / FramesNum;
 
-			Result.Add(KV.Key, KV.Value.Inner.Stats);
+			Result.GetMap().Add(KV.Key, Stats);
 		}
 	}
 
 private:
-	enum
-	{
-		HistogramLen = 100
-	};
-
 	struct FInternalAggregationEntry
 	{
-		FInternalAggregationEntry()
-		{
-			FMemory::Memzero(InclHistogram, sizeof(int32) * HistogramLen);
-			InclDT = 1.0;
+		FInternalAggregationEntry() = default;
+		~FInternalAggregationEntry();
 
-			FMemory::Memzero(ExclHistogram, sizeof(int32) * HistogramLen);
-			ExclDT = 1.0;
-		}
+		FAggregatedTimingStats Stats;
 
 		// Histogram for computing median inclusive time.
-		int32 InclHistogram[HistogramLen];
-		double InclDT; // bucket size
+		class FInternalHistogram* InclHistogram = nullptr;
 
 		// Histogram for computing median exclusive time.
-		int32 ExclHistogram[HistogramLen];
-		double ExclDT; // bucket size
-
-		FAggregatedTimingStats Stats;
+		class FInternalHistogram* ExclHistogram = nullptr;
 	};
 
 	struct FInternalFrameAggregationEntry
@@ -314,8 +383,19 @@ private:
 		TArray<double> FrameExclusiveTimes;
 	};
 
-	template<typename TimelineType, typename BucketMappingFunc, typename BucketKeyType, typename CallbackType>
-	static void ProcessTimeline(const TimelineType* Timeline, BucketMappingFunc BucketMapper, CallbackType Callback, double IntervalStart, double IntervalEnd, TMap<BucketKeyType, FInternalAggregationEntry>& InternalResult)
+	template<
+		typename TimelineType,
+		typename BucketMappingFunc,
+		typename BucketKeyType,
+		typename CallbackType
+	>
+	static void ProcessTimeline(
+		const TimelineType* Timeline,
+		BucketMappingFunc BucketMapper,
+		CallbackType Callback,
+		double IntervalStart,
+		double IntervalEnd,
+		TAggregationResult<BucketKeyType, FInternalAggregationEntry>& InternalResult)
 	{
 		struct FStackEntry
 		{
@@ -327,51 +407,62 @@ private:
 		TArray<FStackEntry> Stack;
 		Stack.Reserve(1024);
 		double LastTime = 0.0;
-		Timeline->EnumerateEvents(IntervalStart, IntervalEnd, [BucketMapper, Callback, IntervalStart, IntervalEnd, &Stack, &LastTime, &InternalResult](bool IsEnter, double Time, const typename TimelineType::EventType& Event)
-		{
-			Time = FMath::Clamp(Time, IntervalStart, IntervalEnd);
-			BucketKeyType BucketKey = BucketMapper(Event);
-			if (Stack.Num())
+		Timeline->EnumerateEvents(IntervalStart, IntervalEnd,
+			[BucketMapper, Callback, IntervalStart, IntervalEnd, &Stack, &LastTime, &InternalResult]
+			(bool IsEnter, double Time, const typename TimelineType::EventType& Event)
 			{
-				FStackEntry& StackEntry = Stack.Top();
-				StackEntry.ExclusiveTime += Time - LastTime;
-			}
-			LastTime = Time;
-			if (IsEnter)
-			{
-				FStackEntry& StackEntry = Stack.AddDefaulted_GetRef();
-				StackEntry.StartTime = Time;
-				StackEntry.ExclusiveTime = 0.0;
-				StackEntry.BucketKey = BucketKey;
-				if (!InternalResult.Contains(BucketKey))
+				Time = FMath::Clamp(Time, IntervalStart, IntervalEnd);
+				BucketKeyType BucketKey = BucketMapper(Event);
+				if (Stack.Num())
 				{
-					InternalResult.Add(BucketKey);
+					FStackEntry& StackEntry = Stack.Top();
+					StackEntry.ExclusiveTime += Time - LastTime;
 				}
-			}
-			else
-			{
-				FStackEntry& StackEntry = Stack.Last();
-				double EventInclusiveTime = Time - StackEntry.StartTime;
-				check(EventInclusiveTime >= 0.0);
-				double EventExclusiveTime = StackEntry.ExclusiveTime;
-				check(EventExclusiveTime >= 0.0 && EventExclusiveTime <= EventInclusiveTime);
-				Stack.Pop(EAllowShrinking::No);
-				double EventNonRecursiveInclusiveTime = EventInclusiveTime;
-				for (const FStackEntry& AncestorStackEntry : Stack)
+				LastTime = Time;
+				if (IsEnter)
 				{
-					if (AncestorStackEntry.BucketKey == BucketKey)
+					FStackEntry& StackEntry = Stack.AddDefaulted_GetRef();
+					StackEntry.StartTime = Time;
+					StackEntry.ExclusiveTime = 0.0;
+					StackEntry.BucketKey = BucketKey;
+					InternalResult.GetMap().FindOrAdd(BucketKey);
+				}
+				else
+				{
+					check(Stack.Num() > 0);
+					FStackEntry& StackEntry = Stack.Last();
+					double EventInclusiveTime = Time - StackEntry.StartTime;
+					check(EventInclusiveTime >= 0.0);
+					double EventExclusiveTime = StackEntry.ExclusiveTime;
+					check(EventExclusiveTime >= 0.0 && EventExclusiveTime <= EventInclusiveTime);
+					Stack.Pop(EAllowShrinking::No);
+					double EventNonRecursiveInclusiveTime = EventInclusiveTime;
+					for (const FStackEntry& AncestorStackEntry : Stack)
 					{
-						EventNonRecursiveInclusiveTime = 0.0;
+						if (AncestorStackEntry.BucketKey == BucketKey)
+						{
+							EventNonRecursiveInclusiveTime = 0.0;
+						}
 					}
+					Callback(InternalResult.GetMap()[BucketKey], EventNonRecursiveInclusiveTime, EventExclusiveTime);
 				}
-				Callback(InternalResult[BucketKey], EventNonRecursiveInclusiveTime, EventExclusiveTime);
-			}
-			return EEventEnumerate::Continue;
-		});
+				return EEventEnumerate::Continue;
+			});
 	}
 
-	template<typename TimelineType, typename BucketMappingFunc, typename BucketKeyType, typename CallbackType>
-	static void ProcessTimelineForFrameStats(const TimelineType* Timeline, BucketMappingFunc BucketMapper, CallbackType Callback, double IntervalStart, double IntervalEnd, TMap<BucketKeyType, FAggregatedTimingStats>& InternalResult)
+	template<
+		typename TimelineType,
+		typename BucketMappingFunc,
+		typename BucketKeyType,
+		typename CallbackType
+	>
+	static void ProcessTimelineForFrameStats(
+		const TimelineType* Timeline,
+		BucketMappingFunc BucketMapper,
+		CallbackType Callback,
+		double IntervalStart,
+		double IntervalEnd,
+		TAggregationResult<BucketKeyType, FAggregatedTimingStats>& InternalResult)
 	{
 		struct FStackEntry
 		{
@@ -383,253 +474,81 @@ private:
 		TArray<FStackEntry> Stack;
 		Stack.Reserve(1024);
 		double LastTime = 0.0;
-		Timeline->EnumerateEvents(IntervalStart, IntervalEnd, [BucketMapper, Callback, IntervalStart, IntervalEnd, &Stack, &LastTime, &InternalResult](bool IsEnter, double Time, const typename TimelineType::EventType& Event)
-		{
-			Time = FMath::Clamp(Time, IntervalStart, IntervalEnd);
-			BucketKeyType BucketKey = BucketMapper(Event);
-			if (Stack.Num())
+		Timeline->EnumerateEvents(IntervalStart, IntervalEnd,
+			[BucketMapper, Callback, IntervalStart, IntervalEnd, &Stack, &LastTime, &InternalResult]
+			(bool IsEnter, double Time, const typename TimelineType::EventType& Event)
 			{
-				FStackEntry& StackEntry = Stack.Top();
-				StackEntry.ExclusiveTime += Time - LastTime;
-			}
-			LastTime = Time;
-			if (IsEnter)
-			{
-				FStackEntry& StackEntry = Stack.AddDefaulted_GetRef();
-				StackEntry.StartTime = Time;
-				StackEntry.ExclusiveTime = 0.0;
-				StackEntry.BucketKey = BucketKey;
-			}
-			else
-			{
-				FStackEntry& StackEntry = Stack.Last();
-				double EventInclusiveTime = Time - StackEntry.StartTime;
-				check(EventInclusiveTime >= 0.0);
-				double EventExclusiveTime = StackEntry.ExclusiveTime;
-				check(EventExclusiveTime >= 0.0 && EventExclusiveTime <= EventInclusiveTime);
-				Stack.Pop(EAllowShrinking::No);
-				double EventNonRecursiveInclusiveTime = EventInclusiveTime;
-				for (const FStackEntry& AncestorStackEntry : Stack)
+				Time = FMath::Clamp(Time, IntervalStart, IntervalEnd);
+				BucketKeyType BucketKey = BucketMapper(Event);
+				if (Stack.Num())
 				{
-					if (AncestorStackEntry.BucketKey == BucketKey)
-					{
-						EventNonRecursiveInclusiveTime = 0.0;
-					}
+					FStackEntry& StackEntry = Stack.Top();
+					StackEntry.ExclusiveTime += Time - LastTime;
 				}
-				Callback(InternalResult[BucketKey], EventNonRecursiveInclusiveTime, EventExclusiveTime);
-			}
-			return EEventEnumerate::Continue;
-		});
-	}
-
-	static void UpdateTotalMinMaxTimerAggregationStats(FInternalAggregationEntry& AggregationEntry, double InclTime, double ExclTime)
-	{
-		UpdateTotalMinMaxTimerStats(AggregationEntry.Stats, InclTime, ExclTime);
-	}
-
-	static void UpdateTotalMinMaxTimerStats(FAggregatedTimingStats& Stats, double InclTime, double ExclTime)
-	{
-		Stats.TotalInclusiveTime += InclTime;
-		if (InclTime < Stats.MinInclusiveTime)
-		{
-			Stats.MinInclusiveTime = InclTime;
-		}
-		if (InclTime > Stats.MaxInclusiveTime)
-		{
-			Stats.MaxInclusiveTime = InclTime;
-		}
-
-		Stats.TotalExclusiveTime += ExclTime;
-		if (ExclTime < Stats.MinExclusiveTime)
-		{
-			Stats.MinExclusiveTime = ExclTime;
-		}
-		if (ExclTime > Stats.MaxExclusiveTime)
-		{
-			Stats.MaxExclusiveTime = ExclTime;
-		}
-
-		Stats.InstanceCount++;
-	}
-
-	static void PreComputeHistogram(FInternalAggregationEntry& AggregationEntry)
-	{
-		const FAggregatedTimingStats& Stats = AggregationEntry.Stats;
-
-		// Each bucket (Histogram[i]) will be centered on a value.
-		// I.e. First bucket (bucket 0) is centered on Min value: [Min-DT/2, Min+DT/2)
-		// and last bucket (bucket N-1) is centered on Max value: [Max-DT/2, Max+DT/2).
-
-		constexpr double InvHistogramLen = 1.0 / static_cast<double>(HistogramLen - 1);
-
-		if (Stats.MaxInclusiveTime == Stats.MinInclusiveTime)
-		{
-			AggregationEntry.InclDT = 1.0; // single large bucket
-		}
-		else
-		{
-			AggregationEntry.InclDT = (Stats.MaxInclusiveTime - Stats.MinInclusiveTime) * InvHistogramLen;
-		}
-
-		if (Stats.MaxExclusiveTime == Stats.MinExclusiveTime)
-		{
-			AggregationEntry.ExclDT = 1.0; // single large bucket
-		}
-		else
-		{
-			AggregationEntry.ExclDT = (Stats.MaxExclusiveTime - Stats.MinExclusiveTime) * InvHistogramLen;
-		}
-	}
-
-	static void UpdateHistogramForTimerStats(FInternalAggregationEntry& AggregationEntry, double InclTime, double ExclTime)
-	{
-		const FAggregatedTimingStats& Stats = AggregationEntry.Stats;
-
-		int32 InclIndex = static_cast<int32>((InclTime - Stats.MinInclusiveTime + AggregationEntry.InclDT / 2) / AggregationEntry.InclDT);
-		ensure(InclIndex >= 0);
-		if (InclIndex < 0)
-		{
-			InclIndex = 0;
-		}
-		ensure(InclIndex < HistogramLen);
-		if (InclIndex >= HistogramLen)
-		{
-			InclIndex = HistogramLen - 1;
-		}
-		AggregationEntry.InclHistogram[InclIndex]++;
-
-		int32 ExclIndex = static_cast<int32>((ExclTime - Stats.MinExclusiveTime + AggregationEntry.ExclDT / 2) / AggregationEntry.ExclDT);
-		ensure(ExclIndex >= 0);
-		if (ExclIndex < 0)
-		{
-			ExclIndex = 0;
-		}
-		ensure(ExclIndex < HistogramLen);
-		if (ExclIndex >= HistogramLen)
-		{
-			ExclIndex = HistogramLen - 1;
-		}
-		AggregationEntry.ExclHistogram[ExclIndex]++;
-	}
-
-	static void PostProcessTimerStats(FInternalAggregationEntry& AggregationEntry, bool ComputeMedian)
-	{
-		FAggregatedTimingStats& Stats = AggregationEntry.Stats;
-
-		// Compute average inclusive/exclusive times.
-		ensure(Stats.InstanceCount > 0);
-		double InvCount = 1.0f / static_cast<double>(Stats.InstanceCount);
-		Stats.AverageInclusiveTime = Stats.TotalInclusiveTime * InvCount;
-		Stats.AverageExclusiveTime = Stats.TotalExclusiveTime * InvCount;
-		Stats.AverageInstanceCount = static_cast<double>(Stats.InstanceCount);
-
-		if (ComputeMedian)
-		{
-			const int32 HalfCount = static_cast<int32>(Stats.InstanceCount / 2);
-
-			// Compute median inclusive time.
-			int32 InclCount = 0;
-			for (int32 HistogramIndex = 0; HistogramIndex < HistogramLen; HistogramIndex++)
-			{
-				InclCount += AggregationEntry.InclHistogram[HistogramIndex];
-				if (InclCount > HalfCount)
+				LastTime = Time;
+				if (IsEnter)
 				{
-					Stats.MedianInclusiveTime = Stats.MinInclusiveTime + HistogramIndex * AggregationEntry.InclDT;
-					if (HistogramIndex > 0 &&
-						Stats.InstanceCount % 2 == 0 &&
-						InclCount - AggregationEntry.InclHistogram[HistogramIndex] == HalfCount)
-					{
-						const double PrevMedian = Stats.MinInclusiveTime + (HistogramIndex - 1) * AggregationEntry.InclDT;
-						Stats.MedianInclusiveTime = (Stats.MedianInclusiveTime + PrevMedian) / 2;
-					}
-					break;
+					FStackEntry& StackEntry = Stack.AddDefaulted_GetRef();
+					StackEntry.StartTime = Time;
+					StackEntry.ExclusiveTime = 0.0;
+					StackEntry.BucketKey = BucketKey;
 				}
-			}
-
-			// Compute median exclusive time.
-			int32 ExclCount = 0;
-			for (int32 HistogramIndex = 0; HistogramIndex < HistogramLen; HistogramIndex++)
-			{
-				ExclCount += AggregationEntry.ExclHistogram[HistogramIndex];
-				if (ExclCount > HalfCount)
+				else
 				{
-					Stats.MedianExclusiveTime = Stats.MinExclusiveTime + HistogramIndex * AggregationEntry.ExclDT;
-					if (HistogramIndex > 0 &&
-						Stats.InstanceCount % 2 == 0 &&
-						ExclCount - AggregationEntry.ExclHistogram[HistogramIndex] == HalfCount)
+					check(Stack.Num() > 0);
+					FStackEntry& StackEntry = Stack.Last();
+					double EventInclusiveTime = Time - StackEntry.StartTime;
+					check(EventInclusiveTime >= 0.0);
+					double EventExclusiveTime = StackEntry.ExclusiveTime;
+					check(EventExclusiveTime >= 0.0 && EventExclusiveTime <= EventInclusiveTime);
+					Stack.Pop(EAllowShrinking::No);
+					double EventNonRecursiveInclusiveTime = EventInclusiveTime;
+					for (const FStackEntry& AncestorStackEntry : Stack)
 					{
-						const double PrevMedian = Stats.MinExclusiveTime + (HistogramIndex - 1) * AggregationEntry.ExclDT;
-						Stats.MedianExclusiveTime = (Stats.MedianExclusiveTime + PrevMedian) / 2;
+						if (AncestorStackEntry.BucketKey == BucketKey)
+						{
+							EventNonRecursiveInclusiveTime = 0.0;
+						}
 					}
-					break;
-				}
-			}
-		}
-	}
-
-	static void ComputeFrameStatsMedian(FInternalAggregationEntry& AggregationEntry, uint32 NumFrames)
-	{
-		FAggregatedTimingStats& Stats = AggregationEntry.Stats;
-
-		const int32 HalfCount = static_cast<int32>(NumFrames / 2);
-
-		// Compute median inclusive time.
-		int32 InclCount = 0;
-		for (int32 HistogramIndex = 0; HistogramIndex < HistogramLen; HistogramIndex++)
-		{
-			InclCount += AggregationEntry.InclHistogram[HistogramIndex];
-			if (InclCount > HalfCount)
-			{
-				Stats.MedianInclusiveTime = Stats.MinInclusiveTime + HistogramIndex * AggregationEntry.InclDT;
-				if (HistogramIndex > 0 &&
-					NumFrames % 2 == 0 &&
-					InclCount - AggregationEntry.InclHistogram[HistogramIndex] == HalfCount)
-				{
-					const double PrevMedian = Stats.MinInclusiveTime + (HistogramIndex - 1) * AggregationEntry.InclDT;
-					Stats.MedianInclusiveTime = (Stats.MedianInclusiveTime + PrevMedian) / 2;
+					Callback(InternalResult.GetMap()[BucketKey], EventNonRecursiveInclusiveTime, EventExclusiveTime);
 				}
-				break;
-			}
-		}
-
-		// Compute median exclusive time.
-		int32 ExclCount = 0;
-		for (int32 HistogramIndex = 0; HistogramIndex < HistogramLen; HistogramIndex++)
-		{
-			ExclCount += AggregationEntry.ExclHistogram[HistogramIndex];
-			if (ExclCount > HalfCount)
-			{
-				Stats.MedianExclusiveTime = Stats.MinExclusiveTime + HistogramIndex * AggregationEntry.ExclDT;
-				if (HistogramIndex > 0 &&
-					NumFrames % 2 == 0 &&
-					ExclCount - AggregationEntry.ExclHistogram[HistogramIndex] == HalfCount)
-				{
-					const double PrevMedian = Stats.MinExclusiveTime + (HistogramIndex - 1) * AggregationEntry.ExclDT;
-					Stats.MedianExclusiveTime = (Stats.MedianExclusiveTime + PrevMedian) / 2;
-				}
-				break;
-			}
-		}
+				return EEventEnumerate::Continue;
+			});
 	}
 
-	template<typename TimelineType, typename BucketMappingFunc, typename BucketKeyType>
-	static void GatherKeysFromTimeline(const TimelineType* Timeline, BucketMappingFunc BucketMapper, double IntervalStart, double IntervalEnd, TMap<BucketKeyType, FAggregatedTimingStats>& InternalResult)
+	template<
+		typename TimelineType,
+		typename BucketMappingFunc,
+		typename BucketKeyType
+	>
+	static void GatherKeysFromTimeline(
+		const TimelineType* Timeline,
+		BucketMappingFunc BucketMapper,
+		double IntervalStart,
+		double IntervalEnd,
+		TAggregationResult<BucketKeyType, FAggregatedTimingStats>& InternalResult)
 	{
-		Timeline->EnumerateEvents(IntervalStart, IntervalEnd, [BucketMapper, IntervalStart, IntervalEnd, &InternalResult](bool IsEnter, double Time, const typename TimelineType::EventType& Event)
+		Timeline->EnumerateEvents(IntervalStart, IntervalEnd,
+			[BucketMapper, IntervalStart, IntervalEnd, &InternalResult]
+			(bool IsEnter, double Time, const typename TimelineType::EventType& Event)
 			{
 				BucketKeyType BucketKey = BucketMapper(Event);
 
 				if (IsEnter)
 				{
-					if (!InternalResult.Contains(BucketKey))
-					{
-						InternalResult.Add(BucketKey);
-					}
+					InternalResult.GetMap().FindOrAdd(BucketKey);
 				}
 
 				return EEventEnumerate::Continue;
 			});
 	}
+
+	static void UpdateTotalMinMaxTimerAggregationStats(FInternalAggregationEntry& AggregationEntry, double InclTime, double ExclTime);
+	static void UpdateTotalMinMaxTimerStats(FAggregatedTimingStats& Stats, double InclTime, double ExclTime);
+	static void InitHistogram(FInternalAggregationEntry& AggregationEntry, ILinearAllocator& Allocator);
+	static void UpdateHistogram(FInternalAggregationEntry& AggregationEntry, double InclTime, double ExclTime);
+	static void ComputeMedianFromHistogram(FInternalAggregationEntry& AggregationEntry);
+	static void PostProcessTimerStats(FInternalAggregationEntry& AggregationEntry, bool bComputeMedian);
 };
 
 } // namespace TraceServices
diff --git a/Engine/Source/Developer/TraceServices/Private/Model/LoadTimeProfiler.cpp b/Engine/Source/Developer/TraceServices/Private/Model/LoadTimeProfiler.cpp
--- a/Engine/Source/Developer/TraceServices/Private/Model/LoadTimeProfiler.cpp
+++ b/Engine/Source/Developer/TraceServices/Private/Model/LoadTimeProfiler.cpp
@@ -262,29 +262,30 @@ ITable<FLoadTimeProfilerAggregatedStats>* FLoadTimeProfilerProvider::CreateEvent
 	{
 		return Event.Package;
 	};
-	TMap<const FPackageInfo*, FAggregatedTimingStats> Aggregation;
+
+	FTimelineStatistics::TAggregationResult<const FPackageInfo*, FAggregatedTimingStats> Aggregation;
 
 	// Cancellation not currently implemented for LoadTimeProfilerProvider.
 	TSharedPtr<TraceServices::FCancellationToken> bIsCancelRequested = MakeShared<TraceServices::FCancellationToken>();
 
 	FTimelineStatistics::CreateAggregation(Timelines, BucketMapper, IntervalStart, IntervalEnd, bIsCancelRequested, Aggregation);
+
 	TTable<FLoadTimeProfilerAggregatedStats>* Table = new TTable<FLoadTimeProfilerAggregatedStats>(AggregatedStatsTableLayout);
-	for (const auto& KV : Aggregation)
+	for (const auto& KV : Aggregation.GetMap())
 	{
 		const FPackageInfo* Package = KV.Key;
-		if (!Package)
+		if (Package)
 		{
-			continue;
+			FLoadTimeProfilerAggregatedStats& Row = Table->AddRow();
+			Row.Name = Package->Name;
+			const FAggregatedTimingStats& Stats = KV.Value;
+			Row.Count = Stats.InstanceCount;
+			Row.Total = Stats.TotalExclusiveTime;
+			Row.Min = Stats.MinExclusiveTime;
+			Row.Max = Stats.MaxExclusiveTime;
+			Row.Average = Stats.AverageExclusiveTime;
+			Row.Median = Stats.MedianExclusiveTime;
 		}
-		FLoadTimeProfilerAggregatedStats& Row = Table->AddRow();
-		Row.Name = Package->Name;
-		const FAggregatedTimingStats& Stats = KV.Value;
-		Row.Count = Stats.InstanceCount;
-		Row.Total = Stats.TotalExclusiveTime;
-		Row.Min = Stats.MinExclusiveTime;
-		Row.Max = Stats.MaxExclusiveTime;
-		Row.Average = Stats.AverageExclusiveTime;
-		Row.Median = Stats.MedianExclusiveTime;
 	}
 	return Table;
 }
@@ -301,14 +302,15 @@ ITable<FLoadTimeProfilerAggregatedStats>* FLoadTimeProfilerProvider::CreateObjec
 	{
 		return Event.Export ? Event.Export->Class : nullptr;
 	};
-	TMap<const FClassInfo*, FAggregatedTimingStats> Aggregation;
+
+	FTimelineStatistics::TAggregationResult<const FClassInfo*, FAggregatedTimingStats> Aggregation;
 
 	// Cancellation not currently implemented for LoadTimeProfilerProvider.
 	TSharedPtr<TraceServices::FCancellationToken> CancellationToken = MakeShared<TraceServices::FCancellationToken>();
 
 	FTimelineStatistics::CreateAggregation(Timelines, BucketMapper, IntervalStart, IntervalEnd, CancellationToken, Aggregation);
 	TTable<FLoadTimeProfilerAggregatedStats>* Table = new TTable<FLoadTimeProfilerAggregatedStats>(AggregatedStatsTableLayout);
-	for (const auto& KV : Aggregation)
+	for (const auto& KV : Aggregation.GetMap())
 	{
 		const FClassInfo* ClassInfo = KV.Key;
 		if (ClassInfo)
diff --git a/Engine/Source/Developer/TraceServices/Private/Model/TimingProfiler.cpp b/Engine/Source/Developer/TraceServices/Private/Model/TimingProfiler.cpp
--- a/Engine/Source/Developer/TraceServices/Private/Model/TimingProfiler.cpp
+++ b/Engine/Source/Developer/TraceServices/Private/Model/TimingProfiler.cpp
@@ -962,7 +962,8 @@ ITable<FTimingProfilerAggregatedStats>* FTimingProfilerProvider::CreateAggregati
 		return GetTimer(Event.TimerIndex);
 	};
 
-	TMap<const FTimingProfilerTimer*, FAggregatedTimingStats> Aggregation;
+	FTimelineStatistics::TAggregationResult<const FTimingProfilerTimer*, FAggregatedTimingStats> Aggregation;
+
 	if (Params.FrameType == ETraceFrameType::TraceFrameType_Count)
 	{
 		if (Params.IntervalStart <= Session.GetDurationSeconds())
@@ -1015,11 +1016,11 @@ ITable<FTimingProfilerAggregatedStats>* FTimingProfilerProvider::CreateAggregati
 				return First.TotalInclusiveTime < Second.TotalInclusiveTime;
 			};
 
-		Aggregation.ValueStableSort(SortLambda);
+		Aggregation.GetMap().ValueStableSort(SortLambda);
 	}
 
 	int EntryLimit = Params.TableEntryLimit;
-	for (const auto& KV : Aggregation)
+	for (const auto& KV : Aggregation.GetMap())
 	{
 		FTimingProfilerAggregatedStats& Row = Table->AddRow();
 		Row.Timer = KV.Key;
