[Insights] Removed #pragma once in cpp.

#rnx
#Synced-CL 50439474

diff --git a/Engine/Source/Developer/TraceServices/Private/Common/TimelineStatistics.cpp b/Engine/Source/Developer/TraceServices/Private/Common/TimelineStatistics.cpp
--- a/Engine/Source/Developer/TraceServices/Private/Common/TimelineStatistics.cpp
+++ b/Engine/Source/Developer/TraceServices/Private/Common/TimelineStatistics.cpp
@@ -1,7 +1,5 @@
 // Copyright Epic Games, Inc. All Rights Reserved.
 
-#pragma once
-
 #include "TimelineStatistics.h"
 
 namespace TraceServices
