- TraceServices: Added LexToString() for IAllocationsProvider::EQueryRule enum.

#rnx

#Synced-CL 50589729

diff --git a/Engine/Source/Developer/TraceServices/Public/TraceServices/Model/AllocationsProvider.h b/Engine/Source/Developer/TraceServices/Public/TraceServices/Model/AllocationsProvider.h
--- a/Engine/Source/Developer/TraceServices/Public/TraceServices/Model/AllocationsProvider.h
+++ b/Engine/Source/Developer/TraceServices/Public/TraceServices/Model/AllocationsProvider.h
@@ -297,4 +297,32 @@ TRACESERVICES_API const IAllocationsProvider* ReadAllocationsProvider(const IAna
 
 } // namespace TraceServices
 
+inline const TCHAR* LexToString(const TraceServices::IAllocationsProvider::EQueryRule Rule)
+{
+	using EQueryRule = TraceServices::IAllocationsProvider::EQueryRule;
+	switch (Rule)
+	{
+		case EQueryRule::aAf:       return TEXT("aAf");
+		case EQueryRule::afA:       return TEXT("afA");
+		case EQueryRule::Aaf:       return TEXT("Aaf");
+		case EQueryRule::aAfB:      return TEXT("aAfB");
+		case EQueryRule::AaBf:      return TEXT("AaBf");
+		case EQueryRule::aAfaBf:    return TEXT("aAfaBf");
+		case EQueryRule::AfB:       return TEXT("AfB");
+		case EQueryRule::AaB:       return TEXT("AaB");
+		case EQueryRule::AafB:      return TEXT("AafB");
+		case EQueryRule::aABf:      return TEXT("aABf");
+		case EQueryRule::AaBCf:     return TEXT("AaBCf");
+		case EQueryRule::AaBfC:     return TEXT("AaBfC");
+		case EQueryRule::aABfC:     return TEXT("aABfC");
+		case EQueryRule::AaBCfD:    return TEXT("AaBCfD");
+		case EQueryRule::AoB:       return TEXT("AoB");
+		case EQueryRule::AiB:       return TEXT("AiB");
+		//case EQueryRule::A_vs_B:    return TEXT("A_vs_B");
+		//case EQueryRule::A_or_B:    return TEXT("A_or_B");
+		//case EQueryRule::A_xor_B:   return TEXT("A_xor_B");
+		default:                    return TEXT("");
+	}
+}
+
 #undef UE_API
