[Insights] AllocationsProvider: Improved debug display for the Event Distance Histogram.

#rnx

#Synced-CL 50769868

diff --git a/Engine/Source/Developer/TraceServices/Private/Model/AllocationsProvider.cpp b/Engine/Source/Developer/TraceServices/Private/Model/AllocationsProvider.cpp
--- a/Engine/Source/Developer/TraceServices/Private/Model/AllocationsProvider.cpp
+++ b/Engine/Source/Developer/TraceServices/Private/Model/AllocationsProvider.cpp
@@ -2408,6 +2408,11 @@ void FAllocationsProvider::EditUnmarkAllocationAsHeap(uint32 ThreadId, double Ti
 
 void FAllocationsProvider::FRootHeap::UpdateHistogramByAllocSize(uint64 Size)
 {
+	if (Size > MaxAllocSize)
+	{
+		MaxAllocSize = Size;
+	}
+
 	// HistogramIndex : Value Range
 	// 0 : [0]
 	// 1 : [1]
@@ -2416,48 +2421,61 @@ void FAllocationsProvider::FRootHeap::UpdateHistogramByAllocSize(uint64 Size)
 	// 4 : [5 .. 8]
 	// 5 : [9 .. 16]
 	// ...
-	// i : [2^(i-2)+1 .. 2^(i-1)], i > 2
+	// i : [2^(i-2)+1 .. 2^(i-1)], i = 3 .. 64
 	// ...
 	// 64 : [2^62+1 .. 2^63]
 	// 65 : [2^63+1 .. 2^64-1]
 
-	if (Size == 0)
+	uint32 Index = 0;
+	if (Size != 0)
 	{
-		++AllocSizeHistogramPow2[0];
-	}
-	else
-	{
-		uint32 HistogramIndexPow2 = 65 - static_cast<uint32>(FMath::CountLeadingZeros64(Size - 1));
-		++AllocSizeHistogramPow2[HistogramIndexPow2];
-
-		if (Size > MaxAllocSize)
-		{
-			MaxAllocSize = Size;
-		}
+		Index = 65 - static_cast<uint32>(FMath::CountLeadingZeros64(Size - 1));
+		check(Index < NumAllocSizeHistogramPow2);
 	}
+	++AllocSizeHistogramPow2[Index];
 }
 
 ////////////////////////////////////////////////////////////////////////////////////////////////////
 
 void FAllocationsProvider::FRootHeap::UpdateHistogramByEventDistance(uint32 EventDistance)
 {
+	if (EventDistance > MaxEventDistance)
+	{
+		MaxEventDistance = EventDistance;
+	}
+
 	// HistogramIndex : Value Range
-	// 0 : [0]
-	// 1 : [1]
-	// 2 : [2 .. 3]
-	// 3 : [4 .. 7]
+	//  0 : [0]
+	//  1 : [1]
+	//  2 : [2 .. 10]
+	//  3 : [11 .. 100]
 	// ...
-	// i : [2^(i-1) .. 2^i-1], i > 1
+	//  i : [10^(i-2)+1 .. 10^(i-1)], i = 2 .. 11
 	// ...
-	// 32 : [2^31 .. 2^32-1]
+	// 11 : [10^9+1 .. 10^10]
+	// 12 : [MAX_uint32]
 
-	uint32 HistogramIndexPow2 = 32 - FMath::CountLeadingZeros(EventDistance);
-	++EventDistanceHistogramPow2[HistogramIndexPow2];
-
-	if (EventDistance > MaxEventDistance)
+	uint32 Index;
+	if (EventDistance == 0)
 	{
-		MaxEventDistance = EventDistance;
+		Index = 0;
+	}
+	else if (EventDistance != uint32(-1))
+	{
+		--EventDistance;
+		Index = 1;
+		while (EventDistance > 0)
+		{
+			++Index;
+			EventDistance /= 10;
+		}
+		check(Index < NumEventDistanceHistogramPow10 - 1);
+	}
+	else // EventDistance == uint32(-1)
+	{
+		Index = NumEventDistanceHistogramPow10 - 1;
 	}
+	++EventDistanceHistogramPow10[Index];
 }
 
 ////////////////////////////////////////////////////////////////////////////////////////////////////
@@ -2480,57 +2498,88 @@ void FAllocationsProvider::FRootHeap::LogStats()
 
 ////////////////////////////////////////////////////////////////////////////////////////////////////
 
-static FString BytesToString(uint64 Value)
+static FString Pow2ToString(uint32 Pow2, const TCHAR* UnitName = TEXT(""), bool bAddSpaceBeforeUnit = false)
 {
-	if (Value == 0)
-	{
-		return FString("0");
-	}
-	if (Value == ~0ull)
+	const TCHAR* UnitPrefixes[] = { TEXT(""), TEXT("Ki"), TEXT("Mi"), TEXT("Gi"), TEXT("Ti"), TEXT("Pi"), TEXT("Ei") };
+	constexpr int32 NumPrefixes = UE_ARRAY_COUNT(UnitPrefixes);
+
+	uint32 Pow = 10;
+	for (int32 Index = 0; Index < NumPrefixes; ++Index)
 	{
-		return FString("MAX64");
+		if (Pow2 < Pow)
+		{
+			if (bAddSpaceBeforeUnit)
+			{
+				if (UnitPrefixes[Index][0] == 0 && UnitName[0] == 0)
+				{
+					return FString::Printf(TEXT("%u"), 1u << (Pow2 + 10 - Pow));
+				}
+				else
+				{
+					return FString::Printf(TEXT("%u %s%s"), 1u << (Pow2 + 10 - Pow), UnitPrefixes[Index], UnitName);
+				}
+			}
+			else
+			{
+				return FString::Printf(TEXT("%u%s%s"), 1u << (Pow2 + 10 - Pow), UnitPrefixes[Index], UnitName);
+			}
+		}
+		Pow += 10;
 	}
+	return TEXT("inf");
+}
 
-	constexpr uint64 K = 1024ull;
-	constexpr uint64 M = K * K;
-	constexpr uint64 G = M * K;
-	constexpr uint64 T = G * K;
+////////////////////////////////////////////////////////////////////////////////////////////////////
 
-	uint64 Unit = 0;
-	const TCHAR* UnitStr = nullptr;
-	if (Value < K)
-	{
-		return FString::Printf(TEXT("%llu B"), Value);
-	}
-	else if (Value < M)
-	{
-		Unit = K;
-		UnitStr = TEXT("KiB");
-	}
-	else if (Value < G)
-	{
-		Unit = M;
-		UnitStr = TEXT("MiB");
-	}
-	else if (Value < T)
-	{
-		Unit = G;
-		UnitStr = TEXT("GiB");
-	}
-	else
-	{
-		Unit = T;
-		UnitStr = TEXT("TiB");
-	}
+static FString Pow2BytesToString(uint32 Pow2)
+{
+	return Pow2ToString(Pow2, TEXT("B"), true);
+}
 
-	if (Value % Unit == 0)
-	{
-		return FString::Printf(TEXT("%llu %s"), Value / Unit, UnitStr);
-	}
-	else
+////////////////////////////////////////////////////////////////////////////////////////////////////
+
+static FString Pow10ToString(uint32 Pow10, const TCHAR* UnitName = TEXT(""), bool bAddSpaceBeforeUnit = false)
+{
+	const TCHAR* UnitPrefixes[] = { TEXT(""), TEXT("K"), TEXT("M"), TEXT("G"), TEXT("T"), TEXT("P"), TEXT("E") };
+	constexpr int32 NumPrefixes = UE_ARRAY_COUNT(UnitPrefixes);
+
+	uint32 Pow = 3;
+	for (int32 Index = 0; Index < NumPrefixes; ++Index)
 	{
-		return FString::Printf(TEXT("%.2f %s"), (double)Value / (double)Unit, UnitStr);
+		if (Pow10 < Pow)
+		{
+			uint64 Value = 1ull;
+			const uint32 ShiftCount = Pow10 + 3 - Pow;
+			for (uint32 Shift = 0; Shift < ShiftCount; ++Shift)
+			{
+				Value *= 10;
+			}
+			if (bAddSpaceBeforeUnit)
+			{
+				if (UnitPrefixes[Index][0] == 0 && UnitName[0] == 0)
+				{
+					return FString::Printf(TEXT("%llu"), Value);
+				}
+				else
+				{
+					return FString::Printf(TEXT("%llu %s%s"), Value, UnitPrefixes[Index], UnitName);
+				}
+			}
+			else
+			{
+				return FString::Printf(TEXT("%llu%s%s"), Value, UnitPrefixes[Index], UnitName);
+			}
+		}
+		Pow += 3;
 	}
+	return TEXT("inf");
+}
+
+////////////////////////////////////////////////////////////////////////////////////////////////////
+
+static FString Pow10EventsToString(uint32 Pow10)
+{
+	return Pow10ToString(Pow10, TEXT(""), false);
 }
 
 ////////////////////////////////////////////////////////////////////////////////////////////////////
@@ -2570,17 +2619,17 @@ void FAllocationsProvider::FRootHeap::LogAllocSizeHistogram()
 	UE_LOG(LogTraceServices, Log, TEXT("[MemAlloc] Alloc Size Histogram (%s):"), HeapSpec->Name);
 
 	// HistogramIndex : Value Range
-	// 0 : [0 bytes]
-	// 1 : [1 byte]
-	// 2 : [2 bytes]
-	// 3 : (2 .. 4]
-	// 4 : (4 .. 8]
-	// 5 : (8 .. 16]
+	//  0 : [0 bytes]
+	//  1 : [1 byte]
+	//  2 : [2 bytes]
+	//  3 : (2 .. 4]
+	//  4 : (4 .. 8]
+	//  5 : (8 .. 16]
 	// ...
-	// i : (2^(i-2) .. 2^(i-1)], i > 2
+	//  i : (2^(i-2) .. 2^(i-1)], i = 3 .. 64
 	// ...
 	// 64 : (2^62 .. 2^63]
-	// 65 : (2^63 .. 2^64-1]
+	// 65 : (2^63 .. 2^64)
 
 	uint64 Sum = 0;
 	for (uint32 Index = StartIndex; Index < EndIndex; ++Index)
@@ -2594,25 +2643,33 @@ void FAllocationsProvider::FRootHeap::LogAllocSizeHistogram()
 	for (uint32 Index = StartIndex; Index < EndIndex; ++Index)
 	{
 		const uint64 Count = AllocSizeHistogramPow2[Index];
+		if (Count == 0)
+		{
+			continue;
+		}
 		PartialSum += Count;
 		const double Percent = (double)Count * PercentFactor;
 		const double CumulativePercent = (double)PartialSum * PercentFactor;
-		if (Index > 2)
+		if (Index == 65)
 		{
-			// i : (2^(i-2) .. 2^(i-1)], i > 2
-			const uint64 Min = (1ull << (Index - 2));
-			const uint64 Max = (Index == NumAllocSizeHistogramPow2 - 1) ? ~0ull : (1ull << (Index - 1));
+			// 65 : (2^63 .. 2^64)
+			UE_LOG(LogTraceServices, Log, TEXT("[MemAlloc] (%8s .. %-8s) %11s %4.01f%% %11s %5.01f%%"),
+				*Pow2BytesToString(Index - 2), *Pow2BytesToString(Index - 1), *CountToString(Count), Percent, *CountToString(PartialSum), CumulativePercent);
+		}
+		else if (Index > 2)
+		{
+			// i : (2^(i-2) .. 2^(i-1)], i = 3 .. 64
 			UE_LOG(LogTraceServices, Log, TEXT("[MemAlloc] (%8s .. %-8s] %11s %4.01f%% %11s %5.01f%%"),
-				*BytesToString(Min), *BytesToString(Max), *CountToString(Count), Percent, *CountToString(PartialSum), CumulativePercent);
+				*Pow2BytesToString(Index - 2), *Pow2BytesToString(Index - 1), *CountToString(Count), Percent, *CountToString(PartialSum), CumulativePercent);
 		}
 		else if (Index == 1)
 		{
-			UE_LOG(LogTraceServices, Log, TEXT("[MemAlloc] {     1 byte         } %11s %4.01f%% %11s %5.01f%%"),
+			UE_LOG(LogTraceServices, Log, TEXT("[MemAlloc] [     1 byte         ] %11s %4.01f%% %11s %5.01f%%"),
 				*CountToString(Count), Percent, *CountToString(PartialSum), CumulativePercent);
 		}
 		else // if (Index == 0 || Index == 2)
 		{
-			UE_LOG(LogTraceServices, Log, TEXT("[MemAlloc] {     %u bytes        } %11s %4.01f%% %11s %5.01f%%"),
+			UE_LOG(LogTraceServices, Log, TEXT("[MemAlloc] [     %u bytes        ] %11s %4.01f%% %11s %5.01f%%"),
 				Index, *CountToString(Count), Percent, *CountToString(PartialSum), CumulativePercent);
 		}
 	}
@@ -2624,13 +2681,13 @@ void FAllocationsProvider::FRootHeap::LogAllocSizeHistogram()
 void FAllocationsProvider::FRootHeap::LogEventDistanceHistogram()
 {
 	uint32 StartIndex = 0;
-	while (StartIndex < NumEventDistanceHistogramPow2 && EventDistanceHistogramPow2[StartIndex] == 0)
+	while (StartIndex < NumEventDistanceHistogramPow10 && EventDistanceHistogramPow10[StartIndex] == 0)
 	{
 		++StartIndex;
 	}
 
-	uint32 EndIndex = NumEventDistanceHistogramPow2;
-	while (EndIndex > 0 && EventDistanceHistogramPow2[EndIndex - 1] == 0)
+	uint32 EndIndex = NumEventDistanceHistogramPow10;
+	while (EndIndex > 0 && EventDistanceHistogramPow10[EndIndex - 1] == 0)
 	{
 		--EndIndex;
 	}
@@ -2644,42 +2701,55 @@ void FAllocationsProvider::FRootHeap::LogEventDistanceHistogram()
 	UE_LOG(LogTraceServices, Log, TEXT("[MemAlloc] Event Distance Histogram (%s):"), HeapSpec->Name);
 
 	// HistogramIndex : Value Range
-	// 0 : [0]
-	// 1 : [1]
-	// 2 : [2 .. 3]
-	// 3 : [4 .. 7]
+	//  0 : [0]
+	//  1 : [1]
+	//  2 : (1 .. 10]
+	//  3 : (10 .. 100]
 	// ...
-	// i : [2^(i-1) .. 2^i-1], i > 1
+	//  i : (10^(i-2) .. 10^(i-1)], i = 2 .. 11
 	// ...
-	// 32 : [2^31 .. 2^32-1]
+	// 11 : (10^9 .. 10^10]
+	// 12 : [MAX_uint32]
 
 	uint64 Sum = LiveAllocs->Num();
 	for (uint32 Index = StartIndex; Index < EndIndex; ++Index)
 	{
-		Sum += EventDistanceHistogramPow2[Index];
+		Sum += EventDistanceHistogramPow10[Index];
 	}
 	uint64 PartialSum = 0;
 	double PercentFactor = 100.0 / (double)Sum;
 
-	UE_LOG(LogTraceServices, Log, TEXT("[MemAlloc] -----------------------------------------------------------------"));
+	UE_LOG(LogTraceServices, Log, TEXT("[MemAlloc] -----------------------------------------------------------"));
 	for (uint32 Index = StartIndex; Index < EndIndex; ++Index)
 	{
-		const uint64 Count = EventDistanceHistogramPow2[Index];
+		const uint64 Count = EventDistanceHistogramPow10[Index];
+		if (Count == 0)
+		{
+			continue;
+		}
 		PartialSum += Count;
 		const double Percent = (double)Count * PercentFactor;
 		const double CumulativePercent = (double)PartialSum * PercentFactor;
-		if (Index > 1)
+		if (Index == NumEventDistanceHistogramPow10 - 1)
 		{
-			// i : [2^(i-1) .. 2^i-1], i > 1
-			const uint64 Min = (1ull << (Index - 1));
-			const uint64 Max = (1ull << Index) - 1;
-			UE_LOG(LogTraceServices, Log, TEXT("[MemAlloc] [%11s .. %-11s] %11s %4.01f%% %11s %5.01f%%"),
-				*CountToString(Min), *CountToString(Max), *CountToString(Count), Percent, *CountToString(PartialSum), CumulativePercent);
+			UE_LOG(LogTraceServices, Log, TEXT("[MemAlloc] [     MAX_uint32     ] %11s %4.01f%% %11s %5.01f%%"),
+				*CountToString(Count), Percent, *CountToString(PartialSum), CumulativePercent);
 		}
-		else // if (Index == 0 || Index == 1)
+		else if (Index > 1)
 		{
-			UE_LOG(LogTraceServices, Log, TEXT("[MemAlloc] {          %u               } %11s %4.01f%% %11s %5.01f%%"),
-				Index, *CountToString(Count), Percent, *CountToString(PartialSum), CumulativePercent);
+			// i : (10^(i-2) .. 10^(i-1)], i = 2 .. 11
+			UE_LOG(LogTraceServices, Log, TEXT("[MemAlloc] (%8s .. %-8s] %11s %4.01f%% %11s %5.01f%%"),
+				*Pow10EventsToString(Index - 2), *Pow10EventsToString(Index - 1), *CountToString(Count), Percent, *CountToString(PartialSum), CumulativePercent);
+		}
+		else if (Index == 1)
+		{
+			UE_LOG(LogTraceServices, Log, TEXT("[MemAlloc] [       1 event      ] %11s %4.01f%% %11s %5.01f%%"),
+				*CountToString(Count), Percent, *CountToString(PartialSum), CumulativePercent);
+		}
+		else if (Index == 0)
+		{
+			UE_LOG(LogTraceServices, Log, TEXT("[MemAlloc] [       0 events     ] %11s %4.01f%% %11s %5.01f%%"),
+				*CountToString(Count), Percent, *CountToString(PartialSum), CumulativePercent);
 		}
 	}
 	if (LiveAllocs->Num() > 0)
@@ -2688,10 +2758,10 @@ void FAllocationsProvider::FRootHeap::LogEventDistanceHistogram()
 		PartialSum += Count;
 		const double Percent = (double)Count * PercentFactor;
 		const double CumulativePercent = (double)PartialSum * PercentFactor;
-		UE_LOG(LogTraceServices, Log, TEXT("[MemAlloc] {      \u221E (live allocs)     } %11s %4.01f%% %11s %5.01f%%"),
+		UE_LOG(LogTraceServices, Log, TEXT("[MemAlloc] (   \u221E (live allocs)  ) %11s %4.01f%% %11s %5.01f%%"),
 			*CountToString(Count), Percent, *CountToString(PartialSum), CumulativePercent);
 	}
-	UE_LOG(LogTraceServices, Log, TEXT("[MemAlloc] -----------------------------------------------------------------"));
+	UE_LOG(LogTraceServices, Log, TEXT("[MemAlloc] -----------------------------------------------------------"));
 }
 
 ////////////////////////////////////////////////////////////////////////////////////////////////////
diff --git a/Engine/Source/Developer/TraceServices/Private/Model/AllocationsProvider.h b/Engine/Source/Developer/TraceServices/Private/Model/AllocationsProvider.h
--- a/Engine/Source/Developer/TraceServices/Private/Model/AllocationsProvider.h
+++ b/Engine/Source/Developer/TraceServices/Private/Model/AllocationsProvider.h
@@ -313,8 +313,8 @@ private:
 		uint64 AllocSizeHistogramPow2[NumAllocSizeHistogramPow2] = { 0 };
 		uint64 MaxAllocSize = 0;
 
-		static constexpr uint32 NumEventDistanceHistogramPow2 = 33;
-		uint32 EventDistanceHistogramPow2[NumEventDistanceHistogramPow2] = { 0 };
+		static constexpr uint32 NumEventDistanceHistogramPow10 = 13;
+		uint32 EventDistanceHistogramPow10[NumEventDistanceHistogramPow10] = { 0 };
 		uint32 MaxEventDistance = 0;
 	};
 
