Fixing NetTrace to work with TailBuffer.
When active, the TailBuffer is constantly caching trace packets for the duration of a session. This means we need to ensure that each NameId is unique for the whole session, and also included in Important events to make sure there's no missing data.
It also means we may see multiple ConnectionCreated events with the same ConnectionId, as Traces aren't started/stopped in the same way when outputting a TailBuffer, so in the NetTraceAnalyzer we only update the ConnectionState the first time.
ResetLifetimeConditionDebugNames and OnResetPersistentNetDebugNames removed from the API as they are no longer required and would break the Trace Tail Buffer

[RN] Fixing NetTrace to work with TailBuffer by keeping DebugNameIds consistent. ResetLifetimeConditionDebugNames and OnResetPersistentNetDebugNames removed from Iris Net Trace API as they cause broken traces if Trace Tail Buffer is in use
#tests Built and Tested against Dedicated Server Cooked Data and UnrealInsights
#Synced-CL 50900172

diff --git a/Engine/Source/Developer/TraceServices/Private/Analyzers/NetTraceAnalyzer.cpp b/Engine/Source/Developer/TraceServices/Private/Analyzers/NetTraceAnalyzer.cpp
--- a/Engine/Source/Developer/TraceServices/Private/Analyzers/NetTraceAnalyzer.cpp
+++ b/Engine/Source/Developer/TraceServices/Private/Analyzers/NetTraceAnalyzer.cpp
@@ -94,8 +94,8 @@ bool FNetTraceAnalyzer::OnEvent(uint16 RouteId, EStyle Style, const FOnEventCont
 	FAnalysisSessionEditScope _(Session);
 
 	// sometimes when a trace starts, the cache buffer contains some events which get written before the init event can be written,
-	// so we need to discard those
-	if ((NetTraceVersion == 0) && (RouteId != RouteId_InitEvent))
+	// so we need to discard those. NameEvents are part of the ImportantEvents, so are also expected to arrive before InitEvent
+	if ((NetTraceVersion == 0) && (RouteId != RouteId_InitEvent) && (RouteId != RouteId_NameEvent))
 	{
 		return true;
 	}
@@ -670,25 +670,28 @@ void FNetTraceAnalyzer::HandleConnectionCreatedEvent(const FOnEventContext& Cont
 	const uint16 ConnectionId = EventData.GetValue<uint16>("ConnectionId");
 
 	TSharedRef<FNetTraceGameInstanceState> GameInstanceState = GetOrCreateActiveGameInstanceState(GameInstanceId);
-	ensureAlwaysMsgf(!GameInstanceState->ActiveConnections.Contains(ConnectionId), TEXT("Got ConnectionCreatedEvent for already existing connection GameInstanceId: %u ConnectionId: %u"), GameInstanceId, ConnectionId);
 
-	// Add to both active connections and to persistent connections
- 	FNetProfilerConnectionInternal& Connection = NetProfilerProvider.CreateConnection(GameInstanceState->GameInstanceIndex);
-	TSharedRef<FNetTraceConnectionState> ConnectionState = MakeShared<FNetTraceConnectionState>();
-	GameInstanceState->ActiveConnections.Add(ConnectionId, ConnectionState);
-
-	// Fill in Connection data
-	Connection.Connection.ConnectionId = ConnectionId;
-	Connection.Connection.LifeTime.Begin =  GetLastTimestamp();
-	ConnectionState->ConnectionIndex = Connection.Connection.ConnectionIndex;
-	ConnectionState->CurrentPacketStartIndex[ENetProfilerConnectionMode::Outgoing] = 0U;
-	ConnectionState->CurrentPacketStartIndex[ENetProfilerConnectionMode::Incoming] = 0U;
-
-	ConnectionState->CurrentPacketBitOffset[ENetProfilerConnectionMode::Outgoing] = 0U;
-	ConnectionState->CurrentPacketBitOffset[ENetProfilerConnectionMode::Incoming] = 0U;
-
-	ConnectionState->CurrentPacketStatsStartIndex[ENetProfilerConnectionMode::Outgoing] = 0U;
-	ConnectionState->CurrentPacketStatsStartIndex[ENetProfilerConnectionMode::Incoming] = 0U;
+	//Sometimes it's possible for CreateConnection to be called multiple times for a ConnectionId if the capture was made using the TailBuffer (can cause multiple OnConnection calls)
+	if (!GameInstanceState->ActiveConnections.Contains(ConnectionId))
+	{
+		// Add to both active connections and to persistent connections
+		FNetProfilerConnectionInternal& Connection = NetProfilerProvider.CreateConnection(GameInstanceState->GameInstanceIndex);
+		TSharedRef<FNetTraceConnectionState> ConnectionState = MakeShared<FNetTraceConnectionState>();
+		GameInstanceState->ActiveConnections.Add(ConnectionId, ConnectionState);
+
+		// Fill in Connection data
+		Connection.Connection.ConnectionId = ConnectionId;
+		Connection.Connection.LifeTime.Begin = GetLastTimestamp();
+		ConnectionState->ConnectionIndex = Connection.Connection.ConnectionIndex;
+		ConnectionState->CurrentPacketStartIndex[ENetProfilerConnectionMode::Outgoing] = 0U;
+		ConnectionState->CurrentPacketStartIndex[ENetProfilerConnectionMode::Incoming] = 0U;
+
+		ConnectionState->CurrentPacketBitOffset[ENetProfilerConnectionMode::Outgoing] = 0U;
+		ConnectionState->CurrentPacketBitOffset[ENetProfilerConnectionMode::Incoming] = 0U;
+
+		ConnectionState->CurrentPacketStatsStartIndex[ENetProfilerConnectionMode::Outgoing] = 0U;
+		ConnectionState->CurrentPacketStatsStartIndex[ENetProfilerConnectionMode::Incoming] = 0U;
+	}
 }
 
 uint32 FNetTraceAnalyzer::GetOrCreateNetProfilerStatsCounterTypeIndex(FNetProfilerNameIndexType NameId, ENetProfilerStatsCounterType StatsType)
diff --git a/Engine/Source/Runtime/Net/Core/Private/Net/Core/Trace/NetTraceInternal.cpp b/Engine/Source/Runtime/Net/Core/Private/Net/Core/Trace/NetTraceInternal.cpp
--- a/Engine/Source/Runtime/Net/Core/Private/Net/Core/Trace/NetTraceInternal.cpp
+++ b/Engine/Source/Runtime/Net/Core/Private/Net/Core/Trace/NetTraceInternal.cpp
@@ -132,7 +132,6 @@ private:
 
 static FGameInstanceContainer GameInstanceContainer;
 
-FNetTrace::FOnResetPersistentNetDebugNames FNetTrace::OnResetPersistentNetDebugNames;
 FNetTrace::FOnNetTraceStarted FNetTrace::OnNetTraceStarted;
 
 TUniquePtr<FNetTraceInternal::FThreadBuffer>& FNetTraceInternal::GetThreadBufferUniquePtr()
@@ -176,9 +175,7 @@ void FNetTrace::SetupInstanceForNetTraceId(uint32 InNetTraceId)
 void FNetTrace::OnTraceConnected()
 {
 	++CurrentTraceId;
-	NextNameId = 1;
-	UE::Net::ResetPersistentNetDebugNameIds();
-	OnResetPersistentNetDebugNames.Broadcast();
+	//We no longer reset NextNameId or PersistentNetDebugNames as that is incompatible with the Trace Tail Buffer which is always capturing
 
 	FNetTraceInternal::Reporter::ReportInitEvent(FNetTraceInternal::NetTraceVersion);
 
diff --git a/Engine/Source/Runtime/Net/Core/Private/Net/Core/Trace/PersistentNetDebugNames.cpp b/Engine/Source/Runtime/Net/Core/Private/Net/Core/Trace/PersistentNetDebugNames.cpp
--- a/Engine/Source/Runtime/Net/Core/Private/Net/Core/Trace/PersistentNetDebugNames.cpp
+++ b/Engine/Source/Runtime/Net/Core/Private/Net/Core/Trace/PersistentNetDebugNames.cpp
@@ -131,9 +131,4 @@ const FNetDebugName* CreatePersistentNetDebugName(const TCHAR* Name, uint32 Name
 	return DebugName;
 }
 
-void ResetPersistentNetDebugNameIds()
-{
-	FPersistentNetDebugNames::Get().ResetNameIds();
-}
-
 }
diff --git a/Engine/Source/Runtime/Net/Core/Private/Net/Core/Trace/Reporters/NetTraceReporter.cpp b/Engine/Source/Runtime/Net/Core/Private/Net/Core/Trace/Reporters/NetTraceReporter.cpp
--- a/Engine/Source/Runtime/Net/Core/Private/Net/Core/Trace/Reporters/NetTraceReporter.cpp
+++ b/Engine/Source/Runtime/Net/Core/Private/Net/Core/Trace/Reporters/NetTraceReporter.cpp
@@ -24,7 +24,7 @@ UE_TRACE_EVENT_BEGIN(NetTrace, InitEvent)
 UE_TRACE_EVENT_END()
 
 // Trace a name
-UE_TRACE_EVENT_BEGIN(NetTrace, NameEvent)
+UE_TRACE_EVENT_BEGIN(NetTrace, NameEvent, Important | NoSync)
 	UE_TRACE_EVENT_FIELD(uint32, NameId)
 	UE_TRACE_EVENT_FIELD(UE::Trace::AnsiString, Name)
 UE_TRACE_EVENT_END()
@@ -178,7 +178,7 @@ void FNetTraceReporter::ReportInstanceDestroyed(uint32 GameInstanceId)
 
 void FNetTraceReporter::ReportAnsiName(UE::Net::FNetDebugNameId NameId, uint32 NameSize, const char* Name)
 {
-	UE_TRACE_LOG(NetTrace, NameEvent, NetChannel)
+	UE_TRACE_LOG(NetTrace, NameEvent, NetChannel, NameSize * sizeof(char))
 		<< NameEvent.NameId(NameId)
 		<< NameEvent.Name(Name, NameSize);
 }
diff --git a/Engine/Source/Runtime/Net/Core/Public/Net/Core/Trace/NetDebugName.h b/Engine/Source/Runtime/Net/Core/Public/Net/Core/Trace/NetDebugName.h
--- a/Engine/Source/Runtime/Net/Core/Public/Net/Core/Trace/NetDebugName.h
+++ b/Engine/Source/Runtime/Net/Core/Public/Net/Core/Trace/NetDebugName.h
@@ -23,8 +23,6 @@ struct FNetDebugName
 NETCORE_API const FNetDebugName* CreatePersistentNetDebugName(const TCHAR* Name, uint32 NameLen);
 inline const FNetDebugName* CreatePersistentNetDebugName(const TCHAR* Name) { return CreatePersistentNetDebugName(Name, FCString::Strlen(Name)); }
 
-NETCORE_API void ResetPersistentNetDebugNameIds();
-
 NETCORE_API FNetDebugName const* GetReliabilityNetDebugName(bool bIsReliable);
 
 }
diff --git a/Engine/Source/Runtime/Net/Core/Public/Net/Core/Trace/Private/NetTraceInternal.h b/Engine/Source/Runtime/Net/Core/Public/Net/Core/Trace/Private/NetTraceInternal.h
--- a/Engine/Source/Runtime/Net/Core/Public/Net/Core/Trace/Private/NetTraceInternal.h
+++ b/Engine/Source/Runtime/Net/Core/Public/Net/Core/Trace/Private/NetTraceInternal.h
@@ -192,9 +192,6 @@ struct FNetTrace
 
 	static void OnTraceStopped(FTraceAuxiliary::EConnectionType TraceType, const FString& TraceDestination);
 
-	DECLARE_TS_MULTICAST_DELEGATE(FOnResetPersistentNetDebugNames);
-	static NETCORE_API FOnResetPersistentNetDebugNames OnResetPersistentNetDebugNames;
-
 	DECLARE_TS_MULTICAST_DELEGATE(FOnNetTraceStarted);
 	static NETCORE_API FOnNetTraceStarted OnNetTraceStarted;
 
diff --git a/Engine/Source/Runtime/Net/Iris/Private/Iris/IrisCoreModule.cpp b/Engine/Source/Runtime/Net/Iris/Private/Iris/IrisCoreModule.cpp
--- a/Engine/Source/Runtime/Net/Iris/Private/Iris/IrisCoreModule.cpp
+++ b/Engine/Source/Runtime/Net/Iris/Private/Iris/IrisCoreModule.cpp
@@ -66,13 +66,6 @@ private:
 		{
 			RepSystemCount += (RepSystem != nullptr ? 1 : 0);
 		}
-
-#if UE_NET_TRACE_ENABLED
-		FNetTrace::OnResetPersistentNetDebugNames.AddLambda([]() 
-		{
-			UE::Net::ResetLifetimeConditionDebugNames();
-		});
-#endif
 	}
 
 	virtual void ShutdownModule() override
diff --git a/Engine/Source/Runtime/Net/Iris/Private/Iris/ReplicationSystem/ReplicationOperations.cpp b/Engine/Source/Runtime/Net/Iris/Private/Iris/ReplicationSystem/ReplicationOperations.cpp
--- a/Engine/Source/Runtime/Net/Iris/Private/Iris/ReplicationSystem/ReplicationOperations.cpp
+++ b/Engine/Source/Runtime/Net/Iris/Private/Iris/ReplicationSystem/ReplicationOperations.cpp
@@ -1405,14 +1405,6 @@ bool FReplicationProtocolOperations::IsEqualQuantizedState(FNetSerializationCont
 	return true;
 }
 
-void ResetLifetimeConditionDebugNames()
-{
-	for (FNetDebugName& Name : Private::s_LifetimeConditionDebugNames)
-	{
-		Name.DebugNameId = 0;
-	}
-}
-
 }
 
 namespace UE::Net::Private
diff --git a/Engine/Source/Runtime/Net/Iris/Public/Iris/ReplicationSystem/ReplicationOperations.h b/Engine/Source/Runtime/Net/Iris/Public/Iris/ReplicationSystem/ReplicationOperations.h
--- a/Engine/Source/Runtime/Net/Iris/Public/Iris/ReplicationSystem/ReplicationOperations.h
+++ b/Engine/Source/Runtime/Net/Iris/Public/Iris/ReplicationSystem/ReplicationOperations.h
@@ -162,6 +162,4 @@ struct FReplicationProtocolOperations
 	static IRISCORE_API void DeserializeWithMaskDelta(FNetSerializationContext& Context, uint32* DstChangeMaskData, uint8* RESTRICT DstObjectStateBuffer, const uint8* RESTRICT PrevObjectStateBuffer, const FReplicationProtocol* Protocol);
 };
 
-IRISCORE_API void ResetLifetimeConditionDebugNames();
-
 }
