[Insights] TraceServices: Added analyzer & provider for UObject Trace (with support for UObject trace snapshots).

#Synced-CL 51457572

diff --git a/Engine/Source/Developer/TraceServices/Private/Analyzers/ObjectTraceAnalysis.cpp b/Engine/Source/Developer/TraceServices/Private/Analyzers/ObjectTraceAnalysis.cpp
new file mode 100644
--- /dev/null
+++ b/Engine/Source/Developer/TraceServices/Private/Analyzers/ObjectTraceAnalysis.cpp
@@ -0,0 +1,252 @@
+// Copyright Epic Games, Inc. All Rights Reserved.
+
+#include "ObjectTraceAnalysis.h"
+
+#include "Containers/StringConv.h"
+#include "HAL/LowLevelMemTracker.h"
+
+// TraceServices
+#include "Common/ProviderLock.h"
+#include "Common/Utils.h"
+#include "Model/ObjectProviderPrivate.h"
+#include "TraceServices/Model/AnalysisSession.h"
+#include "TraceServices/Model/Definitions.h"
+#include "TraceServices/Model/Strings.h"
+#include "TraceServices/Model/ObjectProvider.h"
+
+namespace TraceServices
+{
+
+FObjectAnalyzer::FObjectAnalyzer(IAnalysisSession& InSession, IEditableObjectProvider& InEditableObjectProvider)
+	: Session(InSession)
+	, EditableProvider(InEditableObjectProvider)
+{
+}
+
+void FObjectAnalyzer::OnAnalysisBegin(const FOnAnalysisContext& Context)
+{
+	auto& Builder = Context.InterfaceBuilder;
+
+	Builder.RouteEvent(RouteId_BeginSnapshot, "CoreUObject", "BeginSnapshot");
+	Builder.RouteEvent(RouteId_EndSnapshot,   "CoreUObject", "EndSnapshot");
+	Builder.RouteEvent(RouteId_ObjectSpec,    "CoreUObject", "ObjectSpec");
+	Builder.RouteEvent(RouteId_ObjectRef,     "CoreUObject", "ObjectRef");
+}
+
+void FObjectAnalyzer::OnAnalysisEnd()
+{
+}
+
+static const TCHAR* ReadString(const UE::Trace::IAnalyzer::FEventData& EventData, const IDefinitionProvider* DefinitionProvider, const ANSICHAR* FieldName)
+{
+	const int32 FieldIndex = EventData.GetTypeInfo().GetFieldIndex(FieldName);
+	if (FieldIndex >= 0)
+	{
+		UE::Trace::FEventRef32 StringRef = EventData.GetReferenceValue<uint32>(FieldIndex);
+		const FStringDefinition* StringDef = DefinitionProvider->Get<FStringDefinition>(StringRef);
+		if (StringDef)
+		{
+			return StringDef->Display;
+		}
+	}
+	return nullptr;
+}
+
+bool FObjectAnalyzer::OnEvent(uint16 RouteId, EStyle Style, const FOnEventContext& Context)
+{
+	LLM_SCOPE_BYNAME(TEXT("Insights/ObjectAnalyzer"));
+
+	const auto& EventData = Context.EventData;
+
+	switch (RouteId)
+	{
+
+	case RouteId_BeginSnapshot:
+	{
+		if (CurrentSnapshot)
+		{
+			UE_LOG(LogTraceServices, Warning, TEXT("[UObj] Ignoring duplicate BeginSnapshot event."));
+			break;
+		}
+
+		const uint32 ThreadId = Context.ThreadInfo.GetId();
+		const uint64 Cycle = EventData.GetValue<uint64>("Cycle");
+		const double Time = Context.EventTime.AsSeconds(Cycle);
+
+		UE_LOG(LogTraceServices, Log, TEXT("[UObj] Tid=%u BeginSnapshot Time=%f"), ThreadId, Time);
+
+		if (Time < SnapshotLastTime)
+		{
+			UE_LOG(LogTraceServices, Warning, TEXT("[UObj] Invalid timestamp for BeginSnapshot event (%f < %f; delta = %f)!"), Time, SnapshotLastTime, SnapshotLastTime - Time);
+		}
+		SnapshotLastTime = Time;
+
+		{
+			FProviderEditScopeLock _(EditableProvider);
+			CurrentSnapshot = EditableProvider.BeginSnapshot(Time);
+		}
+		check(CurrentSnapshot);
+		break;
+	}
+
+	case RouteId_EndSnapshot:
+	{
+		if (!CurrentSnapshot)
+		{
+			UE_LOG(LogTraceServices, Warning, TEXT("[UObj] Ignoring EndSnapshot event received without a matching BeginSnapshot event."));
+			break;
+		}
+
+		const uint32 ThreadId = Context.ThreadInfo.GetId();
+		const uint64 Cycle = EventData.GetValue<uint64>("Cycle");
+		const double Time = Context.EventTime.AsSeconds(Cycle);
+
+		UE_LOG(LogTraceServices, Log, TEXT("[UObj] Tid=%u EndSnapshot Time=%f (Duration=%f)"), ThreadId, Time, Time - SnapshotLastTime);
+
+		if (Time < SnapshotLastTime)
+		{
+			UE_LOG(LogTraceServices, Warning, TEXT("[UObj] Invalid timestamp for EndSnapshot event (%f < %f; delta = %f)!"), Time, SnapshotLastTime, SnapshotLastTime - Time);
+		}
+		SnapshotLastTime = Time;
+
+		{
+			FProviderEditScopeLock _(EditableProvider);
+			IObjectEditableSnapshot* Snapshot = EditableProvider.EndSnapshot(Time);
+			if (Snapshot != nullptr)
+			{
+				check(Snapshot == CurrentSnapshot);
+				UE_LOG(LogTraceServices, Log, TEXT("[UObj] Completed snapshot (%u objects, %u refs)."), Snapshot->GetNumObjects(), Snapshot->GetNumReferences());
+				
+				// debug
+				//Snapshot->SaveAs(*FString::Printf(TEXT("D:/uobjects_%u.csv"), Snapshot->GetId()));
+			}
+			else
+			{
+				UE_LOG(LogTraceServices, Warning, TEXT("[UObj] Completed empty snapshot!"));
+			}
+		}
+		CurrentSnapshot = nullptr;
+		break;
+	}
+
+	case RouteId_ObjectSpec:
+	{
+		if (!CurrentSnapshot)
+		{
+			UE_LOG(LogTraceServices, Warning, TEXT("[UObj] Ignoring ObjectSpec event received outside a Begin-End Snapshot scope."));
+			break;
+		}
+
+		const uint32 ThreadId = Context.ThreadInfo.GetId();
+
+		uint32 Id = EventData.GetValue<uint32>("Id");
+		uint32 ClassId = EventData.GetValue<uint32>("ClassId");
+		uint32 OuterId = EventData.GetValue<uint32>("OuterId");
+		uint32 PackageId = EventData.GetValue<uint32>("PackageId");
+		uint32 Flags = EventData.GetValue<uint32>("Flags");
+
+		int32 StructureSize = EventData.GetValue<int32>("StructureSize");
+		uint64 SystemMemoryBytes = EventData.GetValue<uint64>("SystemMemoryBytes");
+		uint64 VideoMemoryBytes = EventData.GetValue<uint64>("VideoMemoryBytes");
+
+		//////////////////////////////////////////////////
+		// Read Object's Name
+
+		const TCHAR* PersistentName = nullptr;
+#if 0
+		const TraceServices::IDefinitionProvider* DefinitionProvider = TraceServices::ReadDefinitionProvider(Session);
+		if (DefinitionProvider)
+		{
+			TraceServices::FProviderReadScopeLock DefinitionProviderReadLock(*DefinitionProvider);
+			PersistentName = ReadString(EventData, DefinitionProvider, "Name");
+		}
+#endif
+		if (!PersistentName)
+		{
+			FAnsiStringView NameAnsi;
+			if (EventData.GetString("NameAnsi", NameAnsi) &&
+				!NameAnsi.IsEmpty())
+			{
+				FAnalysisSessionEditScope _(Session);
+				PersistentName = Session.StoreString(FString(NameAnsi));
+			}
+		}
+		if (!PersistentName)
+		{
+			FWideStringView NameWide;
+			if (EventData.GetString("NameWide", NameWide) &&
+				!NameWide.IsEmpty())
+			{
+				FAnalysisSessionEditScope _(Session);
+				PersistentName = Session.StoreString(NameWide);
+			}
+		}
+		if (!PersistentName)
+		{
+			FAnalysisSessionEditScope _(Session);
+			PersistentName = Session.StoreString(TEXTVIEW("<noname>"));
+		}
+		check(PersistentName != nullptr);
+
+		const TCHAR* PersistentVersePath = nullptr;
+		{
+			FStringView VersePath;
+			if (EventData.GetString("VersePath", VersePath) &&
+				!VersePath.IsEmpty())
+			{
+				FAnalysisSessionEditScope _(Session);
+				PersistentVersePath = Session.StoreString(VersePath);
+			}
+		}
+		if (!PersistentVersePath)
+		{
+			FAnalysisSessionEditScope _(Session);
+			PersistentVersePath = Session.StoreString(TEXTVIEW(""));
+		}
+
+		//////////////////////////////////////////////////
+
+		FProviderEditScopeLock _(EditableProvider);
+		FObjectInfo* Object = EditableProvider.AddObject(Id, ClassId, OuterId, PackageId, PersistentName, PersistentVersePath, Flags);
+		if (Object)
+		{
+			Object->StructureSize = StructureSize;
+			Object->SystemMemoryBytes = SystemMemoryBytes;
+			Object->VideoMemoryBytes = VideoMemoryBytes;
+		}
+		else
+		{
+			UE_LOG(LogTraceServices, Error, TEXT("[UObj] Failed to add object with Id=%u (Name=\"%s\", VersePath=\"%s\")"), Id, PersistentName, PersistentVersePath);
+		}
+		break;
+	}
+
+	case RouteId_ObjectRef:
+	{
+		if (!CurrentSnapshot)
+		{
+			UE_LOG(LogTraceServices, Warning, TEXT("[UObj] Ignoring ObjectRef event received outside a Begin-End Snapshot scope."));
+			break;
+		}
+
+		const uint32 ThreadId = Context.ThreadInfo.GetId();
+		uint32 ReferencerId = EventData.GetValue<uint32>("ReferencerId");
+		uint32 ObjectId = EventData.GetValue<uint32>("ObjectId");
+
+		//UE_LOG(LogTraceServices, VeryVerbose, TEXT("[UObj] Tid=%u ObjectReference ReferencerSerial=%d ObjectSerial=%d"), ThreadId, ReferencerSerial, ObjectSerial);
+
+		FProviderEditScopeLock _(EditableProvider);
+		FObjectReferenceInfo* Reference = EditableProvider.AddObjectReference(ReferencerId, ObjectId);
+		if (!Reference)
+		{
+			UE_LOG(LogTraceServices, Error, TEXT("[UObj] Failed to add reference (ReferencerId=%u, ObjectId=%u)"), ReferencerId, ObjectId);
+		}
+		break;
+	}
+
+	} // switch (RouteId)
+
+	return true;
+}
+
+} // namespace TraceServices
diff --git a/Engine/Source/Developer/TraceServices/Private/Analyzers/ObjectTraceAnalysis.h b/Engine/Source/Developer/TraceServices/Private/Analyzers/ObjectTraceAnalysis.h
new file mode 100644
--- /dev/null
+++ b/Engine/Source/Developer/TraceServices/Private/Analyzers/ObjectTraceAnalysis.h
@@ -0,0 +1,41 @@
+// Copyright Epic Games, Inc. All Rights Reserved.
+
+#pragma once
+
+#include "CoreTypes.h"
+
+#include "Trace/Analyzer.h"
+#include "TraceServices/Model/ObjectProvider.h"
+
+namespace TraceServices
+{
+
+class IAnalysisSession;
+class IEditableObjectProvider;
+
+class FObjectAnalyzer : public UE::Trace::IAnalyzer
+{
+private:
+	enum : uint16
+	{
+		RouteId_BeginSnapshot,
+		RouteId_EndSnapshot,
+		RouteId_ObjectSpec,
+		RouteId_ObjectRef,
+	};
+
+public:
+	FObjectAnalyzer(IAnalysisSession& Session, IEditableObjectProvider& InEditableObjectProvider);
+	virtual void OnAnalysisBegin(const FOnAnalysisContext& Context) override;
+	virtual void OnAnalysisEnd() override;
+	virtual bool OnEvent(uint16 RouteId, EStyle Style, const FOnEventContext& Context) override;
+
+private:
+	IAnalysisSession& Session;
+	IEditableObjectProvider& EditableProvider;
+
+	IObjectEditableSnapshot* CurrentSnapshot = nullptr;
+	double SnapshotLastTime = 0.0;
+};
+
+} // namespace TraceServices
diff --git a/Engine/Source/Developer/TraceServices/Private/Model/ObjectProvider.cpp b/Engine/Source/Developer/TraceServices/Private/Model/ObjectProvider.cpp
new file mode 100644
--- /dev/null
+++ b/Engine/Source/Developer/TraceServices/Private/Model/ObjectProvider.cpp
@@ -0,0 +1,356 @@
+// Copyright Epic Games, Inc. All Rights Reserved.
+
+#include "TraceServices/Model/ObjectProvider.h"
+#include "ObjectProviderPrivate.h"
+
+#include "Containers/Utf8String.h"
+#include "HAL/PlatformFileManager.h"
+#include "UObject/ObjectMacros.h"
+#include "GenericPlatform/GenericPlatformFile.h"
+
+// TraceServices
+#include "TraceServices/Model/AnalysisSession.h"
+#include "TraceServices/Containers/Allocators.h"
+
+#include <limits>
+
+namespace TraceServices
+{
+
+thread_local FProviderLock::FThreadLocalState GObjectProviderLockState;
+
+////////////////////////////////////////////////////////////////////////////////////////////////////
+// FObjectSnapshot
+////////////////////////////////////////////////////////////////////////////////////////////////////
+
+FObjectSnapshot::FObjectSnapshot()
+{
+}
+
+////////////////////////////////////////////////////////////////////////////////////////////////////
+
+FObjectSnapshot::~FObjectSnapshot()
+{
+}
+
+////////////////////////////////////////////////////////////////////////////////////////////////////
+
+FObjectInfo& FObjectSnapshot::GetOrAddObject(uint32 ObjectId)
+{
+	check(ObjectId < MaxNumObjects);
+	if (ObjectId >= uint32(Objects.Num()))
+	{
+		Objects.AddDefaulted(int32(ObjectId) + 1 - Objects.Num());
+	}
+	return Objects[ObjectId];
+}
+
+////////////////////////////////////////////////////////////////////////////////////////////////////
+
+void FObjectSnapshot::EnumerateObjects(TFunctionRef<bool(const FObjectInfo&)> Callback) const
+{
+	const uint32 ObjectCount = uint32(Objects.Num());
+	for (uint32 ObjectId = 0; ObjectId < ObjectCount; ++ObjectId)
+	{
+		const FObjectInfo& Object = Objects[ObjectId];
+		if (Object.Id == ObjectId)
+		{
+			if (!Callback(Object))
+			{
+				break;
+			}
+		}
+	}
+}
+
+////////////////////////////////////////////////////////////////////////////////////////////////////
+
+const FObjectInfo* FObjectSnapshot::FindObject(FStringView Name) const
+{
+	const uint32 ObjectCount = uint32(Objects.Num());
+	for (uint32 ObjectId = 0; ObjectId < ObjectCount; ++ObjectId)
+	{
+		const FObjectInfo& Object = Objects[ObjectId];
+		if (Object.Id == ObjectId &&
+			Name.Equals(Object.Name, ESearchCase::CaseSensitive))
+		{
+			return &Object;
+		}
+	}
+	return nullptr;
+}
+
+////////////////////////////////////////////////////////////////////////////////////////////////////
+
+void FObjectSnapshot::SaveAs(const TCHAR* CsvFileName) const
+{
+	IPlatformFile& PlatformFile = FPlatformFileManager::Get().GetPlatformFile();
+	IFileHandle* File = PlatformFile.OpenWrite(CsvFileName, false, false);
+	if (File)
+	{
+		FUtf8StringView Header = UTF8TEXTVIEW("Id,ClassId,OuterId,PackageId,Flags,StructureSize,SystemMemoryBytes,VideoMemoryBytes,Name,VersePath\n");
+		File->Write((const uint8*)Header.GetData(), Header.NumBytes());
+		for (const FObjectInfo& Obj : Objects)
+		{
+			if (Obj.Id == InvalidObjectId)
+			{
+				continue;
+			}
+			FUtf8String Str = FUtf8String::Printf(UTF8TEXT("%d,%d,%d,%d,0x%08X,%d,%llu,%llu,%s,%s\n"),
+				int32(Obj.Id),
+				int32(Obj.ClassId),
+				int32(Obj.OuterId),
+				int32(Obj.PackageId),
+				Obj.Flags, //TCHAR_TO_UTF8(*LexToString((EObjectFlags)Flags)),
+				Obj.StructureSize,
+				Obj.SystemMemoryBytes,
+				Obj.VideoMemoryBytes,
+				TCHAR_TO_UTF8(Obj.Name),
+				TCHAR_TO_UTF8(Obj.VersePath));
+			const auto& Data = Str.GetCharArray();
+			File->Write((const uint8*)Data.GetData(), Data.NumBytes());
+		}
+		delete File;
+	}
+}
+
+////////////////////////////////////////////////////////////////////////////////////////////////////
+// FObjectProvider
+////////////////////////////////////////////////////////////////////////////////////////////////////
+
+FObjectProvider::FObjectProvider(IAnalysisSession& InSession)
+	: Session(InSession)
+{
+}
+
+////////////////////////////////////////////////////////////////////////////////////////////////////
+
+FObjectProvider::~FObjectProvider()
+{
+	InternalReset();
+}
+
+////////////////////////////////////////////////////////////////////////////////////////////////////
+
+void FObjectProvider::InternalReset()
+{
+	if (CurrentSnapshot)
+	{
+		delete CurrentSnapshot;
+		CurrentSnapshot = nullptr;
+	}
+
+	for (const FObjectSnapshot* Snapshot : Snapshots)
+	{
+		delete Snapshot;
+	}
+	Snapshots.Reset();
+}
+
+////////////////////////////////////////////////////////////////////////////////////////////////////
+
+void FObjectProvider::EnumerateObjects(uint32 SnapshotId, TFunctionRef<bool(const FObjectInfo&)> Callback) const
+{
+	ReadAccessCheck();
+
+	if (SnapshotId >= uint32(Snapshots.Num()))
+	{
+		return;
+	}
+	checkSlow(Snapshots[SnapshotId]);
+	const FObjectSnapshot& Snapshot = *Snapshots[SnapshotId];
+	Snapshot.EnumerateObjects(Callback);
+}
+
+////////////////////////////////////////////////////////////////////////////////////////////////////
+
+const FObjectInfo* FObjectProvider::GetObject(uint32 SnapshotId, uint32 ObjectId) const
+{
+	ReadAccessCheck();
+
+	if (SnapshotId >= uint32(Snapshots.Num()))
+	{
+		return nullptr;
+	}
+	checkSlow(Snapshots[SnapshotId]);
+	const FObjectSnapshot& Snapshot = *Snapshots[SnapshotId];
+	return Snapshot.GetObject(ObjectId);
+}
+
+////////////////////////////////////////////////////////////////////////////////////////////////////
+
+const FObjectInfo* FObjectProvider::FindObject(uint32 SnapshotId, FStringView Name) const
+{
+	ReadAccessCheck();
+
+	if (SnapshotId >= uint32(Snapshots.Num()))
+	{
+		return nullptr;
+	}
+	checkSlow(Snapshots[SnapshotId]);
+	const FObjectSnapshot& Snapshot = *Snapshots[SnapshotId];
+	return Snapshot.FindObject(Name);
+}
+
+////////////////////////////////////////////////////////////////////////////////////////////////////
+
+uint32 FObjectProvider::GetNumSnapshots() const
+{
+	ReadAccessCheck();
+
+	return uint32(Snapshots.Num());
+}
+
+////////////////////////////////////////////////////////////////////////////////////////////////////
+
+void FObjectProvider::EnumerateSnapshots(TFunctionRef<bool(const IObjectSnapshot&)> Callback) const
+{
+	ReadAccessCheck();
+
+	for (const FObjectSnapshot* Snapshot : Snapshots)
+	{
+		checkSlow(Snapshot);
+		if (!Callback(*Snapshot))
+		{
+			break;
+		}
+	}
+}
+
+////////////////////////////////////////////////////////////////////////////////////////////////////
+
+const IObjectSnapshot* FObjectProvider::GetSnapshot(uint32 SnapshotId) const
+{
+	ReadAccessCheck();
+
+	return (SnapshotId < uint32(Snapshots.Num())) ? Snapshots[SnapshotId] : nullptr;
+}
+
+////////////////////////////////////////////////////////////////////////////////////////////////////
+
+const IObjectSnapshot* FObjectProvider::GetLowerBoundSnapshot(double Time) const
+{
+	ReadAccessCheck();
+
+	const int32 Index = Algo::UpperBound(Snapshots, Time, [](double Value, const FObjectSnapshot* Snapshot) { return Value < Snapshot->StartTime; });
+	if (Index > 0 && Index <= Snapshots.Num())
+	{
+		return Snapshots[Index - 1];
+	}
+	return nullptr;
+}
+
+////////////////////////////////////////////////////////////////////////////////////////////////////
+
+IObjectEditableSnapshot* FObjectProvider::GetCurrentSnapshot()
+{
+	EditAccessCheck();
+
+	return CurrentSnapshot;
+}
+
+////////////////////////////////////////////////////////////////////////////////////////////////////
+
+IObjectEditableSnapshot* FObjectProvider::BeginSnapshot(double Time)
+{
+	EditAccessCheck();
+
+	if (CurrentSnapshot)
+	{
+		return nullptr;
+	}
+
+	CurrentSnapshot = new FObjectSnapshot();
+	CurrentSnapshot->Id = uint32(-1);
+	CurrentSnapshot->StartTime = Time;
+	CurrentSnapshot->EndTime = std::numeric_limits<double>::infinity();
+
+	return CurrentSnapshot;
+}
+
+////////////////////////////////////////////////////////////////////////////////////////////////////
+
+IObjectEditableSnapshot* FObjectProvider::EndSnapshot(double Time)
+{
+	EditAccessCheck();
+
+	if (!CurrentSnapshot)
+	{
+		return nullptr;
+	}
+
+#if 0
+	// Ignores empty snapshots.
+	if (CurrentSnapshot->IsEmpty())
+	{
+		delete CurrentSnapshot;
+		CurrentSnapshot = nullptr;
+		return nullptr;
+	}
+#endif
+
+	CurrentSnapshot->EndTime = FMath::Max(Time, CurrentSnapshot->StartTime);
+	CurrentSnapshot->Id = uint32(Snapshots.Num());
+	Snapshots.Add(CurrentSnapshot);
+
+	FObjectSnapshot* Snapshot = CurrentSnapshot;
+	CurrentSnapshot = nullptr;
+	return Snapshot;
+}
+
+////////////////////////////////////////////////////////////////////////////////////////////////////
+
+FObjectInfo* FObjectProvider::AddObject(uint32 Id, uint32 ClassId, uint32 OuterId, uint32 PackageId, const TCHAR* PersistentName, const TCHAR* PersistentVersePath, uint32 Flags)
+{
+	EditAccessCheck();
+
+	if (!ensure(CurrentSnapshot))
+	{
+		return nullptr;
+	}
+
+	if (!FObjectSnapshot::IsValidObjectId(Id))
+	{
+		return nullptr;
+	}
+
+	FObjectInfo& Object = CurrentSnapshot->GetOrAddObject(Id);
+	Object.Id = Id;
+	Object.ClassId = ClassId;
+	Object.OuterId = OuterId;
+	Object.PackageId = PackageId;
+	Object.Name = PersistentName;
+	Object.VersePath = PersistentVersePath;
+	Object.Flags = Flags;
+
+	return &Object;
+}
+
+////////////////////////////////////////////////////////////////////////////////////////////////////
+
+FObjectReferenceInfo* FObjectProvider::AddObjectReference(uint32 ReferencerId, uint32 ObjectId)
+{
+	EditAccessCheck();
+
+	if (!ensure(CurrentSnapshot))
+	{
+		return nullptr;
+	}
+
+	if (!FObjectSnapshot::IsValidObjectId(ReferencerId) ||
+		!FObjectSnapshot::IsValidObjectId(ObjectId))
+	{
+		return nullptr;
+	}
+
+	FObjectReferenceInfo& ReferenceInfo = CurrentSnapshot->References.AddDefaulted_GetRef();
+	ReferenceInfo.ReferencerId = ReferencerId;
+	ReferenceInfo.ObjectId = ObjectId;
+
+	return &ReferenceInfo;
+}
+
+////////////////////////////////////////////////////////////////////////////////////////////////////
+////////////////////////////////////////////////////////////////////////////////////////////////////
+
+} // namespace TraceServices
diff --git a/Engine/Source/Developer/TraceServices/Private/Model/ObjectProviderPrivate.h b/Engine/Source/Developer/TraceServices/Private/Model/ObjectProviderPrivate.h
new file mode 100644
--- /dev/null
+++ b/Engine/Source/Developer/TraceServices/Private/Model/ObjectProviderPrivate.h
@@ -0,0 +1,125 @@
+// Copyright Epic Games, Inc. All Rights Reserved.
+
+#pragma once
+
+#include "CoreTypes.h"
+#include "Containers/Array.h"
+
+// TraceServices
+#include "Common/ProviderLock.h"
+#include "TraceServices/Model/ObjectProvider.h"
+
+namespace TraceServices
+{
+
+extern thread_local FProviderLock::FThreadLocalState GObjectProviderLockState;
+
+class FObjectSnapshot : public IObjectEditableSnapshot
+{
+	friend class FObjectProvider;
+
+public:
+	static constexpr uint32 MaxNumObjects = 1 << 25; // ~33 million objects
+	static constexpr uint32 InvalidObjectId = uint32(-1);
+
+public:
+	FObjectSnapshot();
+	virtual ~FObjectSnapshot();
+
+	virtual uint32 GetId() const { return Id; }
+
+	static inline bool IsValidObjectId(uint32 ObjectId)
+	{
+		return ObjectId < MaxNumObjects;
+	}
+	inline bool IsValidObject(uint32 ObjectId)
+	{
+		return ObjectId < uint32(Objects.Num()) ? Objects[ObjectId].Id == ObjectId : false;
+	}
+
+	virtual double GetStartTime() const override { return StartTime; }
+	virtual double GetEndTime() const override { return EndTime; }
+
+	virtual uint32 GetNumObjects() const override { return uint32(Objects.Num()); }
+	virtual uint32 GetNumReferences() const override { return uint32(References.Num()); }
+
+	virtual const FObjectInfo* GetObject(uint32 ObjectId) const override
+	{
+		return ObjectId < uint32(Objects.Num()) && Objects[ObjectId].Id == ObjectId ? &Objects[ObjectId] : nullptr;
+	}
+
+	virtual TConstArrayView<FObjectInfo> GetObjectsArray() const override { return Objects; }
+	virtual TConstArrayView<FObjectReferenceInfo> GetReferencesArray() const override { return References; }
+
+	inline bool IsEmpty() const
+	{
+		return Objects.IsEmpty() && References.IsEmpty();
+	}
+
+	FObjectInfo& GetOrAddObject(uint32 ObjectId);
+	void EnumerateObjects(TFunctionRef<bool(const FObjectInfo&)> Callback) const;
+	const FObjectInfo* FindObject(FStringView Name) const;
+
+	virtual void SaveAs(const TCHAR* CsvFileName) const;
+
+private:
+	uint32 Id; // snapshot's unique id
+	double StartTime = 0.0;
+	double EndTime = 0.0;
+	TArray<FObjectInfo> Objects;
+	TArray<FObjectReferenceInfo> References;
+};
+
+class FObjectProvider : public IObjectProvider, public IEditableObjectProvider
+{
+public:
+	explicit FObjectProvider(IAnalysisSession& InSession);
+	virtual ~FObjectProvider();
+
+	//////////////////////////////////////////////////
+	// Read operations
+
+	// IProvider
+	virtual void BeginRead() const override { Lock.BeginRead(GObjectProviderLockState); }
+	virtual void EndRead() const override { Lock.EndRead(GObjectProviderLockState); }
+	virtual void ReadAccessCheck() const override { Lock.ReadAccessCheck(GObjectProviderLockState); }
+
+	// IObjectProvider
+	virtual void EnumerateObjects(uint32 SnapshotId, TFunctionRef<bool(const FObjectInfo&)> Callback) const override;
+	virtual const FObjectInfo* GetObject(uint32 SnapshotId, uint32 ObjectId) const override;
+	virtual const FObjectInfo* FindObject(uint32 SnapshotId, FStringView Name) const override;
+	virtual uint32 GetNumSnapshots() const override;
+	virtual void EnumerateSnapshots(TFunctionRef<bool(const IObjectSnapshot&)> Callback) const override;
+	virtual const IObjectSnapshot* GetSnapshot(uint32 SnapshotId) const  override;
+	virtual const IObjectSnapshot* GetLowerBoundSnapshot(double Time) const override;
+
+	//////////////////////////////////////////////////
+	// Edit operations
+
+	// IEditableProvider
+	virtual void BeginEdit() const override { Lock.BeginWrite(GObjectProviderLockState); }
+	virtual void EndEdit() const override { Lock.EndWrite(GObjectProviderLockState); }
+	virtual void EditAccessCheck() const override { Lock.WriteAccessCheck(GObjectProviderLockState); }
+
+	// IEditableObjectProvider
+	virtual IObjectEditableSnapshot* GetCurrentSnapshot() override;
+	virtual IObjectEditableSnapshot* BeginSnapshot(double Time) override;
+	virtual IObjectEditableSnapshot* EndSnapshot(double Time) override;
+	virtual FObjectInfo* AddObject(uint32 Id, uint32 ClassId, uint32 OuterId, uint32 PackageId, const TCHAR* PersistentName, const TCHAR* PersistentVersePath, uint32 Flags) override;
+	virtual FObjectReferenceInfo* AddObjectReference(uint32 ReferencerId, uint32 ObjectId) override;
+	virtual const IObjectProvider* GetReadProvider() const override { return this; }
+
+	//////////////////////////////////////////////////
+
+private:
+	void InternalReset();
+
+private:
+	mutable FProviderLock Lock;
+	IAnalysisSession& Session;
+
+	TArray<FObjectSnapshot*> Snapshots;
+	FObjectSnapshot* CurrentSnapshot = nullptr;
+};
+
+} // namespace TraceServices
diff --git a/Engine/Source/Developer/TraceServices/Private/Modules/ObjectModule.cpp b/Engine/Source/Developer/TraceServices/Private/Modules/ObjectModule.cpp
new file mode 100644
--- /dev/null
+++ b/Engine/Source/Developer/TraceServices/Private/Modules/ObjectModule.cpp
@@ -0,0 +1,46 @@
+// Copyright Epic Games, Inc. All Rights Reserved.
+
+#include "ObjectModule.h"
+
+#include "Analyzers/ObjectTraceAnalysis.h"
+#include "Model/ObjectProviderPrivate.h"
+
+namespace TraceServices
+{
+
+void FObjectModule::GetModuleInfo(FModuleInfo& OutModuleInfo)
+{
+	static const FName ObjectModuleName("TraceModule_Object");
+
+	OutModuleInfo.Name = ObjectModuleName;
+	OutModuleInfo.DisplayName = TEXT("Object");
+}
+
+void FObjectModule::OnAnalysisBegin(IAnalysisSession& Session)
+{
+	TSharedPtr<FObjectProvider> ObjectProvider = MakeShared<FObjectProvider>(Session);
+	Session.AddProvider(GetObjectProviderName(), ObjectProvider, ObjectProvider);
+	IEditableObjectProvider* EditableObjectProvider = EditObjectProvider(Session);
+	if (EditableObjectProvider)
+	{
+		Session.AddAnalyzer(new FObjectAnalyzer(Session, *EditableObjectProvider));
+	}
+}
+
+FName GetObjectProviderName()
+{
+	static const FName Name("ObjectProvider");
+	return Name;
+}
+
+const IObjectProvider* ReadObjectProvider(const IAnalysisSession& Session)
+{
+	return Session.ReadProvider<IObjectProvider>(GetObjectProviderName());
+}
+
+IEditableObjectProvider* EditObjectProvider(IAnalysisSession& Session)
+{
+	return Session.EditProvider<IEditableObjectProvider>(GetObjectProviderName());
+}
+
+} // namespace TraceServices
diff --git a/Engine/Source/Developer/TraceServices/Private/Modules/ObjectModule.h b/Engine/Source/Developer/TraceServices/Private/Modules/ObjectModule.h
new file mode 100644
--- /dev/null
+++ b/Engine/Source/Developer/TraceServices/Private/Modules/ObjectModule.h
@@ -0,0 +1,17 @@
+// Copyright Epic Games, Inc. All Rights Reserved.
+
+#pragma once
+
+#include "TraceServices/ModuleService.h"
+
+namespace TraceServices
+{
+
+class FObjectModule : public IModule
+{
+public:
+	virtual void GetModuleInfo(FModuleInfo& OutModuleInfo) override;
+	virtual void OnAnalysisBegin(IAnalysisSession& Session) override;
+};
+
+} // namespace TraceServices
diff --git a/Engine/Source/Developer/TraceServices/Private/TraceServicesModule.cpp b/Engine/Source/Developer/TraceServices/Private/TraceServicesModule.cpp
--- a/Engine/Source/Developer/TraceServices/Private/TraceServicesModule.cpp
+++ b/Engine/Source/Developer/TraceServices/Private/TraceServicesModule.cpp
@@ -14,6 +14,7 @@
 #include "Modules/MemoryModule.h"
 #include "Modules/ModuleManager.h"
 #include "Modules/NetProfilerModule.h"
+#include "Modules/ObjectModule.h"
 #include "Modules/PlatformEventsModule.h"
 #include "Modules/StatsModule.h"
 #include "Modules/TasksModule.h"
@@ -49,6 +50,7 @@ private:
 	TraceServices::FPlatformEventsModule PlatformEventsModule;
 	TraceServices::FTasksModule TasksModule;
 	TraceServices::FCookProfilerModule CookProfilingModule;
+	TraceServices::FObjectModule ObjectModule;
 };
 
 TSharedPtr<TraceServices::IAnalysisService> FTraceServicesModule::GetAnalysisService()
@@ -104,7 +106,7 @@ void FTraceServicesModule::StartupModule()
 
 	// Load the analysis module.
 	IModuleInterface& TraceAnalysisModule = FModuleManager::LoadModuleChecked<IModuleInterface>("TraceAnalysis");
-	
+
 	IModularFeatures::Get().RegisterModularFeature(TraceServices::ModuleFeatureName, &TimingProfilerModule);
 	IModularFeatures::Get().RegisterModularFeature(TraceServices::ModuleFeatureName, &CsvProfilerModule);
 	IModularFeatures::Get().RegisterModularFeature(TraceServices::ModuleFeatureName, &CountersModule);
@@ -116,12 +118,14 @@ void FTraceServicesModule::StartupModule()
 	IModularFeatures::Get().RegisterModularFeature(TraceServices::ModuleFeatureName, &MemoryModule);
 	IModularFeatures::Get().RegisterModularFeature(TraceServices::ModuleFeatureName, &TasksModule);
 	IModularFeatures::Get().RegisterModularFeature(TraceServices::ModuleFeatureName, &CookProfilingModule);
+	IModularFeatures::Get().RegisterModularFeature(TraceServices::ModuleFeatureName, &ObjectModule);
 }
 
 void FTraceServicesModule::ShutdownModule()
 {
 	LLM_SCOPE_BYTAG(Insights_TraceServices);
 
+	IModularFeatures::Get().UnregisterModularFeature(TraceServices::ModuleFeatureName, &ObjectModule);
 	IModularFeatures::Get().UnregisterModularFeature(TraceServices::ModuleFeatureName, &CookProfilingModule);
 	IModularFeatures::Get().UnregisterModularFeature(TraceServices::ModuleFeatureName, &TasksModule);
 	IModularFeatures::Get().UnregisterModularFeature(TraceServices::ModuleFeatureName, &MemoryModule);
diff --git a/Engine/Source/Developer/TraceServices/Public/TraceServices/Model/ObjectProvider.h b/Engine/Source/Developer/TraceServices/Public/TraceServices/Model/ObjectProvider.h
new file mode 100644
--- /dev/null
+++ b/Engine/Source/Developer/TraceServices/Public/TraceServices/Model/ObjectProvider.h
@@ -0,0 +1,164 @@
+// Copyright Epic Games, Inc. All Rights Reserved.
+
+#pragma once
+
+#include "CoreTypes.h"
+
+#include "Templates/Function.h"
+#include "UObject/NameTypes.h"
+
+// TraceServices
+#include "TraceServices/Model/AnalysisSession.h"
+
+namespace TraceServices
+{
+
+struct FObjectInfo
+{
+	uint32 Id = uint32(-1);
+	uint32 ClassId = uint32(-1);
+	uint32 OuterId = uint32(-1);
+	uint32 PackageId = uint32(-1);
+	const TCHAR* Name = nullptr;
+	const TCHAR* VersePath = nullptr;
+	uint32 Flags = 0;
+	int32 StructureSize = 0;
+	uint64 SystemMemoryBytes = 0;
+	uint64 VideoMemoryBytes = 0;
+	uint64 _Reserved = 0;
+};
+
+struct FObjectReferenceInfo
+{
+	uint32 ReferencerId = 0;
+	uint32 ObjectId = 0;
+};
+
+class IObjectSnapshot
+{
+public:
+	virtual ~IObjectSnapshot() = default;
+	virtual uint32 GetId() const = 0;
+	virtual double GetStartTime() const = 0;
+	virtual double GetEndTime() const = 0;
+	virtual const FObjectInfo* GetObject(uint32 Id) const = 0;
+	virtual uint32 GetNumObjects() const = 0;
+	virtual uint32 GetNumReferences() const = 0;
+	virtual TConstArrayView<FObjectInfo> GetObjectsArray() const = 0;
+	virtual TConstArrayView<FObjectReferenceInfo> GetReferencesArray() const = 0;
+	virtual void SaveAs(const TCHAR* CsvFileName) const = 0;
+};
+
+class IObjectEditableSnapshot : public IObjectSnapshot
+{
+public:
+	virtual ~IObjectEditableSnapshot() = default;
+};
+
+class IObjectProvider : public IProvider
+{
+public:
+	virtual ~IObjectProvider() = default;
+
+	virtual void BeginRead() const = 0;
+	virtual void EndRead() const = 0;
+	virtual void ReadAccessCheck() const = 0;
+
+	/**
+	 * Enumerates the available objects.
+	 * The enumeration stops if the callback returns false.
+	*/
+	virtual void EnumerateObjects(uint32 SnapshotId, TFunctionRef<bool(const FObjectInfo&)> Callback) const = 0;
+
+	/**
+	 * Gets the object with the specified Id.
+	 * The resulting pointer is valid only under the current read lock.
+	 */
+	virtual const FObjectInfo* GetObject(uint32 SnapshotId, uint32 ObjectId) const = 0;
+
+	/**
+	 * Finds the object with the specified name.
+	 * The resulting pointer is valid only under the current read lock.
+	 */
+	virtual const FObjectInfo* FindObject(uint32 SnapshotId, FStringView Name) const = 0;
+
+	/**
+	 * Gets the number of available snapshots.
+	 */
+	virtual uint32 GetNumSnapshots() const = 0;
+
+	/**
+	 * Enumerates the available snapshots.
+	 * The enumeration stops if the callback returns false.
+	 */
+	virtual void EnumerateSnapshots(TFunctionRef<bool(const IObjectSnapshot&)> Callback) const = 0;
+
+	/**
+	 * Gets the snapshot with the specified Id.
+	 * The resulting pointer is valid only under the current read lock.
+	 */
+	virtual const IObjectSnapshot* GetSnapshot(uint32 SnapshotId) const = 0;
+
+	/**
+	 * Finds the last snapshot with StartTime <= provided Time.
+	 * The resulting pointer is valid only under the current read lock.
+	 */
+	virtual const IObjectSnapshot* GetLowerBoundSnapshot(double Time) const = 0;
+};
+
+class IEditableObjectProvider : public IEditableProvider
+{
+public:
+	virtual ~IEditableObjectProvider() = default;
+
+	/**
+	 * Gets the current snapshot.
+	 * The resulting pointer is valid only under the current edit lock.
+	 */
+	virtual IObjectEditableSnapshot* GetCurrentSnapshot() = 0;
+
+	/**
+	 * Begins a new snapshot.
+	 * The resulting pointer is valid only under the current edit lock.
+	 */
+	virtual IObjectEditableSnapshot* BeginSnapshot(double Time) = 0;
+
+	/**
+	 * Ends the current snapshot.
+	 * The resulting pointer is valid only under the current edit lock.
+	 */
+	virtual IObjectEditableSnapshot* EndSnapshot(double Time) = 0;
+
+	/**
+	 * Adds an object declaration, to the current snapshot.
+	 * It is only valid to call this between BeginSnapshot() and EndSnapshot() calls.
+	 * The resulting pointer is valid only under the current edit lock.
+	 */
+	virtual FObjectInfo* AddObject(
+		uint32 ObjectId,
+		uint32 ClassId,
+		uint32 OuterId,
+		uint32 PackageId,
+		const TCHAR* PersistentName,
+		const TCHAR* PersistentVersePath,
+		uint32 Flags) = 0;
+
+	/**
+	 * Adds a dependency between objects, to the current snapshot.
+	 * It is only valid to call this between BeginSnapshot() and EndSnapshot() calls.
+	 * The resulting pointer is valid only under the current edit lock.
+	 */
+	virtual FObjectReferenceInfo* AddObjectReference(uint32 ReferencerId, uint32 ObjectId) = 0;
+
+	/**
+	 * Gets the read provider.
+	 * @returns The read only provider or nullptr if not available.
+	 */
+	virtual const IObjectProvider* GetReadProvider() const { return nullptr; }
+};
+
+TRACESERVICES_API FName GetObjectProviderName();
+TRACESERVICES_API const IObjectProvider* ReadObjectProvider(const IAnalysisSession& Session);
+TRACESERVICES_API IEditableObjectProvider* EditObjectProvider(IAnalysisSession& Session);
+
+} // namespace TraceServices
diff --git a/Engine/Source/Developer/TraceServices/TraceServices.Build.cs b/Engine/Source/Developer/TraceServices/TraceServices.Build.cs
--- a/Engine/Source/Developer/TraceServices/TraceServices.Build.cs
+++ b/Engine/Source/Developer/TraceServices/TraceServices.Build.cs
@@ -12,6 +12,7 @@ public class TraceServices : ModuleRules
 			new string[] {
 				"Cbor",
 				"Core",
+				"CoreUObject",
 				"SymsLib",
 				"TraceAnalysis",
 			}
