[Insights]
* Extended the ObjectProfiler analysis and UI to support the new extra-spec trace events (VerseExtraSpec, ResourceSizeExtraSpec, FieldExtraSpec, StructExtraSpec, ClassExtraSpec, PackageExtraSpec).
* FObjectInfo enriched with type hierarchy flags (EObjectInfoFlags), super/inheritance ids, inclusive memory sizes,
and package path.
* Extended IEditableObjectProvider with GetEditableObject() to allow post-creation patching of object data by
extra-spec events.
* Refactored FObjectAnalyzer: extracted common snapshot validation and warning/error logging into helpers; added warning/error throttling; updated log messages for the processed trace events.
* ObjectTable: Added separate table columns for inclusive estimated memory sizes.

#rnx
#Synced-CL 51579960

diff --git a/Engine/Source/Developer/TraceInsights/Private/Insights/ObjectProfiler/ViewModels/ObjectNode.cpp b/Engine/Source/Developer/TraceInsights/Private/Insights/ObjectProfiler/ViewModels/ObjectNode.cpp
--- a/Engine/Source/Developer/TraceInsights/Private/Insights/ObjectProfiler/ViewModels/ObjectNode.cpp
+++ b/Engine/Source/Developer/TraceInsights/Private/Insights/ObjectProfiler/ViewModels/ObjectNode.cpp
@@ -21,14 +21,31 @@ FObjectNode::FObjectNode(TWeakPtr<FObjectTable> InParentTable, const TraceServic
 	, ObjectId(ObjectInfo.Id)
 	, ClassId(ObjectInfo.ClassId)
 	, OuterId(ObjectInfo.OuterId)
-	, PackageId(ObjectInfo.PackageId)
 	, ObjectName(ObjectInfo.Name)
 	, VersePath(ObjectInfo.VersePath)
 	, ObjectFlags(ObjectInfo.Flags)
 	, StructureSize(ObjectInfo.StructureSize)
 	, SystemMemorySize(int64(ObjectInfo.SystemMemoryBytes))
 	, VideoMemorySize(int64(ObjectInfo.VideoMemoryBytes))
+	, InclSystemMemorySize(int64(ObjectInfo.TotalSystemMemoryBytes))
+	, InclVideoMemorySize(int64(ObjectInfo.TotalVideoMemoryBytes))
 {
+	if (EnumHasAnyFlags(ObjectInfo.FlagsEx, TraceServices::EObjectInfoFlags::IsField))
+	{
+		SetIsField();
+	}
+	if (EnumHasAnyFlags(ObjectInfo.FlagsEx, TraceServices::EObjectInfoFlags::IsStruct))
+	{
+		SetIsStruct();
+	}
+	if (EnumHasAnyFlags(ObjectInfo.FlagsEx, TraceServices::EObjectInfoFlags::IsClass))
+	{
+		SetIsClass();
+	}
+	if (EnumHasAnyFlags(ObjectInfo.FlagsEx, TraceServices::EObjectInfoFlags::IsPackage))
+	{
+		SetIsPackage();
+	}
 }
 
 ////////////////////////////////////////////////////////////////////////////////////////////////////
@@ -38,10 +55,8 @@ FObjectNode::FObjectNode(const FObjectNode& Clone)
 	, ObjectId(Clone.ObjectId)
 	, ClassId(Clone.ClassId)
 	, OuterId(Clone.OuterId)
-	, PackageId(Clone.PackageId)
 	, Class(Clone.Class)
 	, Outer(Clone.Outer)
-	, Package(Clone.Package)
 	, ObjectName(Clone.ObjectName)
 	, VersePath(Clone.VersePath)
 	, ObjectPath(Clone.ObjectPath)
@@ -49,6 +64,8 @@ FObjectNode::FObjectNode(const FObjectNode& Clone)
 	, StructureSize(Clone.StructureSize)
 	, SystemMemorySize(Clone.SystemMemorySize)
 	, VideoMemorySize(Clone.VideoMemorySize)
+	, InclSystemMemorySize(Clone.InclSystemMemorySize)
+	, InclVideoMemorySize(Clone.InclVideoMemorySize)
 	, InternalFlags(Clone.InternalFlags)
 {
 }
@@ -92,36 +109,52 @@ const FText FObjectNode::GetTooltipText() const
 
 const FSlateBrush* FObjectNode::GetIcon() const
 {
-	if (IsClass())
-	{
-		return FInsightsStyle::GetBrush("Icons.UClass.TreeItem");
-	}
-	else if (IsPackage())
-	{
-		return FInsightsStyle::GetBrush("Icons.UPackage.TreeItem");
-	}
-	else
+	if (InternalFlags != EInternalFlags::None)
 	{
-		return FInsightsStyle::GetBrush("Icons.UObject.TreeItem");
+		if (IsClass())
+		{
+			return FInsightsStyle::GetBrush("Icons.UClass.TreeItem");
+		}
+		else if (IsStruct())
+		{
+			return FInsightsStyle::GetBrush("Icons.UStruct.TreeItem");
+		}
+		else if (IsField())
+		{
+			return FInsightsStyle::GetBrush("Icons.UField.TreeItem");
+		}
+		else if (IsPackage())
+		{
+			return FInsightsStyle::GetBrush("Icons.UPackage.TreeItem");
+		}
 	}
+	return FInsightsStyle::GetBrush("Icons.UObject.TreeItem");
 }
 
 ////////////////////////////////////////////////////////////////////////////////////////////////////
 
 FLinearColor FObjectNode::GetIconColor() const
 {
-	if (IsClass())
+	if (InternalFlags != EInternalFlags::None)
 	{
-		return FLinearColor(0.5f, 1.0f, 1.0f, 1.0f);
-	}
-	else if (IsPackage())
-	{
-		return FLinearColor(1.0f, 1.0f, 0.5f, 1.0f);
-	}
-	else
-	{
-		return GetDefaultColor(IsGroup());
+		if (IsClass())
+		{
+			return FLinearColor(0.5f, 1.0f, 1.0f, 1.0f);
+		}
+		else if (IsStruct())
+		{
+			return FLinearColor(0.6f, 1.0f, 1.0f, 1.0f);
+		}
+		else if (IsField())
+		{
+			return FLinearColor(0.7f, 1.0f, 1.0f, 1.0f);
+		}
+		else if (IsPackage())
+		{
+			return FLinearColor(1.0f, 1.0f, 0.5f, 1.0f);
+		}
 	}
+	return GetDefaultColor(IsGroup());
 }
 
 ////////////////////////////////////////////////////////////////////////////////////////////////////
@@ -144,6 +177,18 @@ FLinearColor FObjectNode::GetColor() const
 
 ////////////////////////////////////////////////////////////////////////////////////////////////////
 
+TSharedPtr<FObjectNode> FObjectNode::GetPackage() const
+{
+	TSharedPtr<FObjectNode> Package = GetOuter();
+	if (Package.IsValid())
+	{
+		return Package->IsPackage() ? Package : Package->GetPackage();
+	}
+	return nullptr;
+}
+
+////////////////////////////////////////////////////////////////////////////////////////////////////
+
 } // namespace UE::Insights::ObjectProfiler
 
 #undef LOCTEXT_NAMESPACE
diff --git a/Engine/Source/Developer/TraceInsights/Private/Insights/ObjectProfiler/ViewModels/ObjectNode.h b/Engine/Source/Developer/TraceInsights/Private/Insights/ObjectProfiler/ViewModels/ObjectNode.h
--- a/Engine/Source/Developer/TraceInsights/Private/Insights/ObjectProfiler/ViewModels/ObjectNode.h
+++ b/Engine/Source/Developer/TraceInsights/Private/Insights/ObjectProfiler/ViewModels/ObjectNode.h
@@ -39,9 +39,11 @@ private:
 	enum class EInternalFlags
 	{
 		None = 0,
-		IsClass = 1 << 0,
-		IsPackage = 1 << 1,
-		IsMissing = 1 << 2, // missing intermediate node in the Outer hierarchy
+		IsField = 1 << 0,
+		IsStruct = 1 << 1,
+		IsClass = 1 << 2,
+		IsPackage = 1 << 3,
+		IsMissing = 1 << 4, // missing intermediate node in the Outer hierarchy
 	};
 
 public:
@@ -66,11 +68,11 @@ public:
 
 	uint32 GetClassId() const { return ClassId; }
 	uint32 GetOuterId() const { return OuterId; }
-	uint32 GetPackageId() const { return PackageId; }
 
 	TSharedPtr<FObjectNode> GetClass() const { return Class.Pin(); }
 	TSharedPtr<FObjectNode> GetOuter() const { return Outer.Pin(); }
-	TSharedPtr<FObjectNode> GetPackage() const { return Package.Pin(); }
+
+	TSharedPtr<FObjectNode> GetPackage() const;
 
 	const TCHAR* GetObjectName() const { return ObjectName; }
 	const TCHAR* GetObjectPath() const { return *ObjectPath; }
@@ -90,11 +92,15 @@ public:
 
 	uint32 GetNumReferences() const { return 0; } // TODO
 
+	bool IsField() const { return (uint32(InternalFlags) & uint32(EInternalFlags::IsField)) != 0; }
+	bool IsStruct() const { return (uint32(InternalFlags) & uint32(EInternalFlags::IsStruct)) != 0; }
 	bool IsClass() const { return (uint32(InternalFlags) & uint32(EInternalFlags::IsClass)) != 0; }
 	bool IsPackage() const { return (uint32(InternalFlags) & uint32(EInternalFlags::IsPackage)) != 0; }
 	bool IsMissing() const { return (uint32(InternalFlags) & uint32(EInternalFlags::IsMissing)) != 0; }
 
 private:
+	void SetIsField() { InternalFlags = EInternalFlags(uint32(InternalFlags) | uint32(EInternalFlags::IsField)); }
+	void SetIsStruct() { InternalFlags = EInternalFlags(uint32(InternalFlags) | uint32(EInternalFlags::IsStruct)); }
 	void SetIsClass() { InternalFlags = EInternalFlags(uint32(InternalFlags) | uint32(EInternalFlags::IsClass)); }
 	void SetIsPackage() { InternalFlags = EInternalFlags(uint32(InternalFlags) | uint32(EInternalFlags::IsPackage)); }
 	void SetIsMissing() { InternalFlags = EInternalFlags(uint32(InternalFlags) | uint32(EInternalFlags::IsMissing)); }
@@ -102,18 +108,15 @@ private:
 
 	void SetClass(TWeakPtr<FObjectNode> InClass) { Class = InClass; }
 	void SetOuter(TWeakPtr<FObjectNode> InOuter) { Outer = InOuter; }
-	void SetPackage(TWeakPtr<FObjectNode> InPackage) { Package = InPackage; }
 
 private:
 	uint32 ObjectId = uint32(-1);
 
 	uint32 ClassId = uint32(-1);
 	uint32 OuterId = uint32(-1);
-	uint32 PackageId = uint32(-1);
 
 	TWeakPtr<FObjectNode> Class;
 	TWeakPtr<FObjectNode> Outer;
-	TWeakPtr<FObjectNode> Package;
 
 	const TCHAR* ObjectName = nullptr;
 	const TCHAR* VersePath = nullptr;
diff --git a/Engine/Source/Developer/TraceInsights/Private/Insights/ObjectProfiler/ViewModels/ObjectTable.cpp b/Engine/Source/Developer/TraceInsights/Private/Insights/ObjectProfiler/ViewModels/ObjectTable.cpp
--- a/Engine/Source/Developer/TraceInsights/Private/Insights/ObjectProfiler/ViewModels/ObjectTable.cpp
+++ b/Engine/Source/Developer/TraceInsights/Private/Insights/ObjectProfiler/ViewModels/ObjectTable.cpp
@@ -33,8 +33,11 @@ const FName FObjectTableColumns::ObjectPathColumnId("ObjectPath");
 const FName FObjectTableColumns::VersePathColumnId("VersePath");
 const FName FObjectTableColumns::StructureSizeColumnId("StructureSize");
 const FName FObjectTableColumns::SystemMemorySizeColumnId("SystemMemorySize");
+const FName FObjectTableColumns::InclSystemMemorySizeColumnId("InclSystemMemorySize");
 const FName FObjectTableColumns::VideoMemorySizeColumnId("VideoMemorySize");
+const FName FObjectTableColumns::InclVideoMemorySizeColumnId("InclVideoMemorySize");
 const FName FObjectTableColumns::EstimatedSizeColumnId("EstimatedSize");
+const FName FObjectTableColumns::InclEstimatedSizeColumnId("InclEstimatedSize");
 const FName FObjectTableColumns::ImpactColumnId("Impact");
 const FName FObjectTableColumns::ReferencesColumnId("References");
 
@@ -549,10 +552,10 @@ void FObjectTable::AddDefaultColumns()
 				if (Node.Is<FObjectNode>())
 				{
 					const FObjectNode& ObjectNode = Node.As<FObjectNode>();
-					const uint32 Id = ObjectNode.GetPackageId();
-					if (Id != uint32(-1))
+					TSharedPtr<FObjectNode> Package = ObjectNode.GetPackage();
+					if (Package.IsValid() && Package->GetObjectId() != uint32(-1))
 					{
-						return FTableCellValue(int64(Id));
+						return FTableCellValue(int64(Package->GetObjectId()));
 					}
 				}
 				return TOptional<FTableCellValue>();
@@ -831,6 +834,56 @@ void FObjectTable::AddDefaultColumns()
 		AddColumn(ColumnRef);
 	}
 	//////////////////////////////////////////////////
+	// Inclusive System Memory Size Column
+	{
+		TSharedRef<FTableColumn> ColumnRef = MakeShared<FTableColumn>(FObjectTableColumns::InclSystemMemorySizeColumnId);
+		FTableColumn& Column = *ColumnRef;
+
+		Column.SetIndex(ColumnIndex++);
+
+		Column.SetShortName(LOCTEXT("InclSystemMemorySizeColumnName", "Incl. Mem Size"));
+		Column.SetTitleName(LOCTEXT("InclSystemMemorySizeColumnTitle", "Inclusive System Memory Size"));
+		Column.SetDescription(LOCTEXT("InclSystemMemorySizeColumnDesc", "The total estimated system memory size (bytes) of the UObject instance (including references)"));
+
+		Column.SetFlags(ETableColumnFlags::CanBeHidden |
+						ETableColumnFlags::CanBeFiltered);
+
+		Column.SetHorizontalAlignment(HAlign_Right);
+		Column.SetInitialWidth(100.0f);
+
+		Column.SetDataType(ETableCellDataType::Int64);
+
+		class FValueGetter : public FAggregatedValueGetter
+		{
+		public:
+			virtual const TOptional<FTableCellValue> GetNodeValue(const FBaseTreeNode& Node) const override
+			{
+				if (Node.Is<FObjectNode>())
+				{
+					const FObjectNode& ObjectNode = Node.As<FObjectNode>();
+					return FTableCellValue(int64(ObjectNode.GetInclusiveSystemMemorySize()));
+				}
+				return TOptional<FTableCellValue>();
+			}
+		};
+		TSharedRef<ITableCellValueGetter> Getter = MakeShared<FValueGetter>();
+		Column.SetValueGetter(Getter);
+
+		TSharedRef<ITableCellValueFormatter> Formatter = MakeShared<FInt64ValueFormatterAsMemory>();
+		Column.SetValueFormatter(Formatter);
+
+		TSharedRef<ITableCellValueSorter> Sorter = MakeShared<FSorterByInt64Value>(ColumnRef);
+		Column.SetValueSorter(Sorter);
+		Column.SetInitialSortMode(EColumnSortMode::Descending);
+
+		//TSharedRef<IFilterValueConverter> Converter = MakeShared<FMemoryFilterValueConverter>();
+		//Column.SetValueConverter(Converter);
+
+		Column.SetAggregation(ETableColumnAggregation::None);
+
+		AddColumn(ColumnRef);
+	}
+	//////////////////////////////////////////////////
 	// Video Memory Size Column
 	{
 		TSharedRef<FTableColumn> ColumnRef = MakeShared<FTableColumn>(FObjectTableColumns::VideoMemorySizeColumnId);
@@ -882,6 +935,56 @@ void FObjectTable::AddDefaultColumns()
 		AddColumn(ColumnRef);
 	}
 	//////////////////////////////////////////////////
+	// Inclusive Video Memory Size Column
+	{
+		TSharedRef<FTableColumn> ColumnRef = MakeShared<FTableColumn>(FObjectTableColumns::InclVideoMemorySizeColumnId);
+		FTableColumn& Column = *ColumnRef;
+
+		Column.SetIndex(ColumnIndex++);
+
+		Column.SetShortName(LOCTEXT("InclVideoMemorySizeColumnName", "Incl. Video Size"));
+		Column.SetTitleName(LOCTEXT("InclVideoMemorySizeColumnTitle", "Inclusive Video Memory Size"));
+		Column.SetDescription(LOCTEXT("InclVideoMemorySizeColumnDesc", "The total estimated video memory size (bytes) of the UObject instance (including references)"));
+
+		Column.SetFlags(ETableColumnFlags::CanBeHidden |
+						ETableColumnFlags::CanBeFiltered);
+
+		Column.SetHorizontalAlignment(HAlign_Right);
+		Column.SetInitialWidth(100.0f);
+
+		Column.SetDataType(ETableCellDataType::Int64);
+
+		class FValueGetter : public FAggregatedValueGetter
+		{
+		public:
+			virtual const TOptional<FTableCellValue> GetNodeValue(const FBaseTreeNode& Node) const override
+			{
+				if (Node.Is<FObjectNode>())
+				{
+					const FObjectNode& ObjectNode = Node.As<FObjectNode>();
+					return FTableCellValue(int64(ObjectNode.GetInclusiveVideoMemorySize()));
+				}
+				return TOptional<FTableCellValue>();
+			}
+		};
+		TSharedRef<ITableCellValueGetter> Getter = MakeShared<FValueGetter>();
+		Column.SetValueGetter(Getter);
+
+		TSharedRef<ITableCellValueFormatter> Formatter = MakeShared<FInt64ValueFormatterAsMemory>();
+		Column.SetValueFormatter(Formatter);
+
+		TSharedRef<ITableCellValueSorter> Sorter = MakeShared<FSorterByInt64Value>(ColumnRef);
+		Column.SetValueSorter(Sorter);
+		Column.SetInitialSortMode(EColumnSortMode::Descending);
+
+		//TSharedRef<IFilterValueConverter> Converter = MakeShared<FMemoryFilterValueConverter>();
+		//Column.SetValueConverter(Converter);
+
+		Column.SetAggregation(ETableColumnAggregation::None);
+
+		AddColumn(ColumnRef);
+	}
+	//////////////////////////////////////////////////
 	// Estimated Size Column
 	{
 		TSharedRef<FTableColumn> ColumnRef = MakeShared<FTableColumn>(FObjectTableColumns::EstimatedSizeColumnId);
@@ -969,6 +1072,93 @@ void FObjectTable::AddDefaultColumns()
 		AddColumn(ColumnRef);
 	}
 	//////////////////////////////////////////////////
+	// Inclusive Estimated Size Column
+	{
+		TSharedRef<FTableColumn> ColumnRef = MakeShared<FTableColumn>(FObjectTableColumns::InclEstimatedSizeColumnId);
+		FTableColumn& Column = *ColumnRef;
+
+		Column.SetIndex(ColumnIndex++);
+
+		Column.SetShortName(LOCTEXT("InclEstimatedSizeColumnName", "Incl. Estimated Size"));
+		Column.SetTitleName(LOCTEXT("InclEstimatedSizeColumnTitle", "Inclusive Estimated Size (MiB)"));
+		Column.SetDescription(LOCTEXT("InclEstimatedSizeColumnDesc", "The total estimated memory size (including references), in MiB"));
+
+		Column.SetFlags(ETableColumnFlags::CanBeHidden |
+						ETableColumnFlags::CanBeFiltered |
+						ETableColumnFlags::IsDynamic);
+
+		Column.SetHorizontalAlignment(HAlign_Right);
+		Column.SetInitialWidth(60.0f);
+
+		Column.SetDataType(ETableCellDataType::Int64);
+
+		class FValueGetter : public FAggregatedValueGetter
+		{
+		public:
+			virtual const TOptional<FTableCellValue> GetNodeValue(const FBaseTreeNode& Node) const override
+			{
+				if (Node.Is<FObjectNode>())
+				{
+					const FObjectNode& ObjectNode = Node.As<FObjectNode>();
+					return FTableCellValue(static_cast<int64>(ObjectNode.GetInclusiveEstimatedSize()));
+				}
+				return TOptional<FTableCellValue>();
+			}
+		};
+		TSharedRef<ITableCellValueGetter> Getter = MakeShared<FValueGetter>();
+		Column.SetValueGetter(Getter);
+
+		class FEstimatedSizeFormatter : public FTableCellValueFormatter
+		{
+		public:
+			virtual FText FormatValue(const TOptional<FTableCellValue>& InValue) const override
+			{
+				if (InValue.IsSet())
+				{
+					const int64 Size = InValue.GetValue().Int64;
+					if (Size > 0)
+					{
+						int64 SizeMB = Size / (1024 * 1024);
+						if (SizeMB == 0)
+						{
+							return LOCTEXT("LessThan1MB", "< 1 MiB");
+						}
+						return FText::Format(LOCTEXT("EstimatedSizeFmt", "{0} MiB"), FText::AsNumber(SizeMB));
+					}
+				}
+				return FText::GetEmpty();
+			}
+			virtual FText FormatValueForTooltip(const TOptional<FTableCellValue>& InValue) const override
+			{
+				if (InValue.IsSet())
+				{
+					const int64 Size = InValue.GetValue().Int64;
+					if (Size == 0)
+					{
+						return FText::GetEmpty(); // LOCTEXT("ZeroMB", "0");
+					}
+					FNumberFormattingOptions Options;
+					Options.MaximumFractionalDigits = 2;
+					return FText::Format(LOCTEXT("EstimatedSizeTooltipFmt", "{0} ({1} bytes)"),
+						FText::AsMemory(Size, &Options),
+						FText::AsNumber(Size));
+				}
+				return FText::GetEmpty();
+			}
+		};
+
+		TSharedRef<ITableCellValueFormatter> Formatter = MakeShared<FEstimatedSizeFormatter>();
+		Column.SetValueFormatter(Formatter);
+
+		TSharedRef<ITableCellValueSorter> Sorter = MakeShared<FSorterByInt64Value>(ColumnRef);
+		Column.SetValueSorter(Sorter);
+		Column.SetInitialSortMode(EColumnSortMode::Descending);
+
+		Column.SetAggregation(ETableColumnAggregation::None);
+
+		AddColumn(ColumnRef);
+	}
+	//////////////////////////////////////////////////
 	// Impact Column
 	{
 		TSharedRef<FTableColumn> ColumnRef = MakeShared<FTableColumn>(FObjectTableColumns::ImpactColumnId);
diff --git a/Engine/Source/Developer/TraceInsights/Private/Insights/ObjectProfiler/ViewModels/ObjectTable.h b/Engine/Source/Developer/TraceInsights/Private/Insights/ObjectProfiler/ViewModels/ObjectTable.h
--- a/Engine/Source/Developer/TraceInsights/Private/Insights/ObjectProfiler/ViewModels/ObjectTable.h
+++ b/Engine/Source/Developer/TraceInsights/Private/Insights/ObjectProfiler/ViewModels/ObjectTable.h
@@ -30,8 +30,11 @@ struct FObjectTableColumns
 	static const FName VersePathColumnId;
 	static const FName StructureSizeColumnId;
 	static const FName SystemMemorySizeColumnId;
+	static const FName InclSystemMemorySizeColumnId;
 	static const FName VideoMemorySizeColumnId;
+	static const FName InclVideoMemorySizeColumnId;
 	static const FName EstimatedSizeColumnId;
+	static const FName InclEstimatedSizeColumnId;
 	static const FName ImpactColumnId;
 	static const FName ReferencesColumnId;
 };
diff --git a/Engine/Source/Developer/TraceInsights/Private/Insights/ObjectProfiler/Widgets/SObjectTableTreeView.cpp b/Engine/Source/Developer/TraceInsights/Private/Insights/ObjectProfiler/Widgets/SObjectTableTreeView.cpp
--- a/Engine/Source/Developer/TraceInsights/Private/Insights/ObjectProfiler/Widgets/SObjectTableTreeView.cpp
+++ b/Engine/Source/Developer/TraceInsights/Private/Insights/ObjectProfiler/Widgets/SObjectTableTreeView.cpp
@@ -377,9 +377,9 @@ void SObjectTableTreeView::Tick(const FGeometry& AllottedGeometry, const double
 		RebuildTree(false);
 
 		const int32 NumTags = GetTableRowNodes().Num();
-		// 1000 tags --> check each 200ms
-		// 10000 tags --> check each 1.1s
-		// 100000 tags --> check each 10.1s
+		// 1000 objects --> check each 200ms
+		// 10000 objects --> check each 1.1s
+		// 100000 objects --> check each 10.1s
 		const double WaitTimeSec = 0.1 + (double)NumTags / 10000.0;
 		const uint64 WaitTime = static_cast<uint64>(WaitTimeSec / FPlatformTime::GetSecondsPerCycle64());
 		NextTimestamp = Time + WaitTime;
@@ -401,6 +401,8 @@ void SObjectTableTreeView::RebuildTree(bool bResync)
 	if (bResync)
 	{
 		ObjectNodes.Empty();
+		Classes.Empty();
+		Packages.Empty();
 		TableRowNodes.Empty();
 		bListHasChanged = true;
 	}
@@ -418,6 +420,8 @@ void SObjectTableTreeView::RebuildTree(bool bResync)
 				if (Snapshot)
 				{
 					ObjectNodes.Empty();
+					Classes.Empty();
+					Packages.Empty();
 					TableRowNodes.Empty();
 					bListHasChanged = true;
 
@@ -427,9 +431,6 @@ void SObjectTableTreeView::RebuildTree(bool bResync)
 					ObjectNodes.AddDefaulted(ObjectCount); // index == object id
 					uint32 ValidObjectCount = 0;
 
-					uint32 ClassId = uint32(-1);
-					uint32 PackageId = uint32(-1);
-
 					TSharedPtr<FObjectTable> ObjectTable = StaticCastSharedPtr<FObjectTable>(GetTable());
 					for (uint32 ObjectId = 0; ObjectId < ObjectCount; ++ObjectId)
 					{
@@ -440,15 +441,14 @@ void SObjectTableTreeView::RebuildTree(bool bResync)
 							ObjectNodes[ObjectId] = ObjectNode;
 							++ValidObjectCount;
 
-							if (FCString::Strcmp(Object->Name, TEXT("Class")) == 0)
+							if (ObjectNode->IsClass())
 							{
-								ClassId = Object->Id;
-								ObjectNode->SetIsClass();
+								Classes.Add(ObjectNode);
 							}
-							else if (FCString::Strcmp(Object->Name, TEXT("Package")) == 0)
+
+							if (ObjectNode->IsPackage())
 							{
-								PackageId = Object->Id;
-								ObjectNode->SetIsPackage();
+								Packages.Add(ObjectNode);
 							}
 						}
 					}
@@ -459,11 +459,9 @@ void SObjectTableTreeView::RebuildTree(bool bResync)
 					{
 						if (NodePtr.IsValid())
 						{
-							// Resolve pointers to class, outer and package objects.
+							// Resolve pointers to class and outer objects.
 							NodePtr->SetClass(GetObjectNode(NodePtr->GetClassId()));
 							NodePtr->SetOuter(GetObjectNode(NodePtr->GetOuterId()));
-							NodePtr->SetPackage(GetObjectNode(NodePtr->GetPackageId()));
-
 							TableRowNodes.Add(NodePtr);
 						}
 					}
@@ -512,7 +510,7 @@ void SObjectTableTreeView::RebuildTree(bool bResync)
 	if (TotalTime > 0.01)
 	{
 		const double SyncTime = SyncStopwatch.GetAccumulatedTime();
-		UE_LOG(LogObjectProfiler, Log, TEXT("[UObj] Tree view rebuilt in %.4fs (sync: %.4fs + update: %.4fs) --> %d objects"),
+		UE_LOGF(LogObjectProfiler, Log, "[Obj] Tree view rebuilt in %.4fs (sync: %.4fs + update: %.4fs) --> %d objects",
 			TotalTime, SyncTime, TotalTime - SyncTime, TableRowNodes.Num());
 	}
 }
diff --git a/Engine/Source/Developer/TraceInsights/Private/Insights/ObjectProfiler/Widgets/SObjectTableTreeView.h b/Engine/Source/Developer/TraceInsights/Private/Insights/ObjectProfiler/Widgets/SObjectTableTreeView.h
--- a/Engine/Source/Developer/TraceInsights/Private/Insights/ObjectProfiler/Widgets/SObjectTableTreeView.h
+++ b/Engine/Source/Developer/TraceInsights/Private/Insights/ObjectProfiler/Widgets/SObjectTableTreeView.h
@@ -89,6 +89,8 @@ private:
 	uint32 PendingSnapshotId = 0;
 
 	TArray<TSharedPtr<FObjectNode>> ObjectNodes;
+	TArray<TSharedPtr<FObjectNode>> Classes;
+	TArray<TSharedPtr<FObjectNode>> Packages;
 
 	//////////////////////////////////////////////////
 	// Filters
diff --git a/Engine/Source/Developer/TraceServices/Private/Analyzers/ObjectTraceAnalysis.cpp b/Engine/Source/Developer/TraceServices/Private/Analyzers/ObjectTraceAnalysis.cpp
--- a/Engine/Source/Developer/TraceServices/Private/Analyzers/ObjectTraceAnalysis.cpp
+++ b/Engine/Source/Developer/TraceServices/Private/Analyzers/ObjectTraceAnalysis.cpp
@@ -14,6 +14,10 @@
 #include "TraceServices/Model/Strings.h"
 #include "TraceServices/Model/ObjectProvider.h"
 
+#define UE_OBJECT_ANALYZER_EVENT(_EventName) const ANSICHAR* EventName = _EventName;
+#define UE_OBJECT_ANALYZER_LOG_API_L0(Format, ...) UE_LOGF(LogTraceServices, Log, "[Obj] Tid=%u %s " Format, Context.ThreadInfo.GetId(), EventName, ##__VA_ARGS__);
+#define UE_OBJECT_ANALYZER_LOG_API_L1(Format, ...) //UE_LOGF(LogTraceServices, Log, "[Obj] Tid=%u %s " Format, Context.ThreadInfo.GetId(), EventName, ##__VA_ARGS__);
+
 namespace TraceServices
 {
 
@@ -27,16 +31,43 @@ void FObjectAnalyzer::OnAnalysisBegin(const FOnAnalysisContext& Context)
 {
 	auto& Builder = Context.InterfaceBuilder;
 
-	Builder.RouteEvent(RouteId_BeginSnapshot, "CoreUObject", "BeginSnapshot");
-	Builder.RouteEvent(RouteId_EndSnapshot,   "CoreUObject", "EndSnapshot");
-	Builder.RouteEvent(RouteId_ObjectSpec,    "CoreUObject", "ObjectSpec");
-	Builder.RouteEvent(RouteId_ObjectRef,     "CoreUObject", "ObjectRef");
+	Builder.RouteEvent(RouteId_BeginSnapshot,         "CoreUObject", "BeginSnapshot");
+	Builder.RouteEvent(RouteId_EndSnapshot,           "CoreUObject", "EndSnapshot");
+	Builder.RouteEvent(RouteId_ObjectSpec,            "CoreUObject", "ObjectSpec");
+	Builder.RouteEvent(RouteId_VerseExtraSpec,        "CoreUObject", "VerseExtraSpec");
+	Builder.RouteEvent(RouteId_ResourceSizeExtraSpec, "CoreUObject", "ResourceSizeExtraSpec");
+	Builder.RouteEvent(RouteId_FieldExtraSpec,        "CoreUObject", "FieldExtraSpec");
+	Builder.RouteEvent(RouteId_StructExtraSpec,       "CoreUObject", "StructExtraSpec");
+	Builder.RouteEvent(RouteId_ClassExtraSpec,        "CoreUObject", "ClassExtraSpec");
+	Builder.RouteEvent(RouteId_PackageExtraSpec,      "CoreUObject", "PackageExtraSpec");
+	Builder.RouteEvent(RouteId_ObjectRef,             "CoreUObject", "ObjectRef");
 }
 
 void FObjectAnalyzer::OnAnalysisEnd()
 {
 }
 
+bool FObjectAnalyzer::CheckValidSnapshotForEvent(const ANSICHAR* EventName)
+{
+	if (!CurrentSnapshot)
+	{
+		if (++NumWarnings < MaxWarningMessages)
+		{
+			UE_LOGF(LogTraceServices, Warning, "[Obj] Ignoring %s event received outside a Begin-End Snapshot scope.", EventName);
+		}
+		return false;
+	}
+	return true;
+}
+
+void FObjectAnalyzer::LogErrorForEventWithInvalidObjectId(const ANSICHAR* EventName, uint32 Id)
+{
+	if (++NumErrors < MaxErrorMessages)
+	{
+		UE_LOGF(LogTraceServices, Error, "[Obj] %s: Failed to edit object with Id=%u", EventName, Id);
+	}
+}
+
 static const TCHAR* ReadString(const UE::Trace::IAnalyzer::FEventData& EventData, const IDefinitionProvider* DefinitionProvider, const ANSICHAR* FieldName)
 {
 	const int32 FieldIndex = EventData.GetTypeInfo().GetFieldIndex(FieldName);
@@ -63,21 +94,24 @@ bool FObjectAnalyzer::OnEvent(uint16 RouteId, EStyle Style, const FOnEventContex
 
 	case RouteId_BeginSnapshot:
 	{
+		UE_OBJECT_ANALYZER_EVENT("BeginSnapshot");
+
 		if (CurrentSnapshot)
 		{
-			UE_LOG(LogTraceServices, Warning, TEXT("[UObj] Ignoring duplicate BeginSnapshot event."));
+			++NumWarnings;
+			UE_LOGF(LogTraceServices, Warning, "[Obj] Ignoring duplicate BeginSnapshot event.");
 			break;
 		}
 
-		const uint32 ThreadId = Context.ThreadInfo.GetId();
 		const uint64 Cycle = EventData.GetValue<uint64>("Cycle");
 		const double Time = Context.EventTime.AsSeconds(Cycle);
 
-		UE_LOG(LogTraceServices, Log, TEXT("[UObj] Tid=%u BeginSnapshot Time=%f"), ThreadId, Time);
+		UE_OBJECT_ANALYZER_LOG_API_L0("Time=%f", Time);
 
 		if (Time < SnapshotLastTime)
 		{
-			UE_LOG(LogTraceServices, Warning, TEXT("[UObj] Invalid timestamp for BeginSnapshot event (%f < %f; delta = %f)!"), Time, SnapshotLastTime, SnapshotLastTime - Time);
+			++NumWarnings;
+			UE_LOGF(LogTraceServices, Warning, "[Obj] Invalid timestamp for BeginSnapshot event (%f < %f; delta=%f)!", Time, SnapshotLastTime, SnapshotLastTime - Time);
 		}
 		SnapshotLastTime = Time;
 
@@ -86,26 +120,38 @@ bool FObjectAnalyzer::OnEvent(uint16 RouteId, EStyle Style, const FOnEventContex
 			CurrentSnapshot = EditableProvider.BeginSnapshot(Time);
 		}
 		check(CurrentSnapshot);
+
+		const uint32 SizeOfUObject = EventData.GetValue<uint32>("SizeOfUObject");
+		const uint32 SizeOfUField = EventData.GetValue<uint32>("SizeOfUField");
+		const uint32 SizeOfUStruct = EventData.GetValue<uint32>("SizeOfUStruct");
+		const uint32 SizeOfUClass = EventData.GetValue<uint32>("SizeOfUClass");
+		const uint32 SizeOfUPackage = EventData.GetValue<uint32>("SizeOfUPackage");
+		UE_LOGF(LogTraceServices, Log, "[Obj] sizeof(UObject)=%u, sizeof(UField)=%u, sizeof(UStruct)=%u, sizeof(UClass)=%u, sizeof(UPackage)=%u",
+			SizeOfUObject, SizeOfUField, SizeOfUStruct, SizeOfUClass, SizeOfUPackage);
+
 		break;
 	}
 
 	case RouteId_EndSnapshot:
 	{
+		UE_OBJECT_ANALYZER_EVENT("EndSnapshot");
+
 		if (!CurrentSnapshot)
 		{
-			UE_LOG(LogTraceServices, Warning, TEXT("[UObj] Ignoring EndSnapshot event received without a matching BeginSnapshot event."));
+			++NumWarnings;
+			UE_LOGF(LogTraceServices, Warning, "[Obj] Ignoring EndSnapshot event received without a matching BeginSnapshot event.");
 			break;
 		}
 
-		const uint32 ThreadId = Context.ThreadInfo.GetId();
 		const uint64 Cycle = EventData.GetValue<uint64>("Cycle");
 		const double Time = Context.EventTime.AsSeconds(Cycle);
 
-		UE_LOG(LogTraceServices, Log, TEXT("[UObj] Tid=%u EndSnapshot Time=%f (Duration=%f)"), ThreadId, Time, Time - SnapshotLastTime);
+		UE_OBJECT_ANALYZER_LOG_API_L0("Time=%f (Duration=%f)", Time, Time - SnapshotLastTime);
 
 		if (Time < SnapshotLastTime)
 		{
-			UE_LOG(LogTraceServices, Warning, TEXT("[UObj] Invalid timestamp for EndSnapshot event (%f < %f; delta = %f)!"), Time, SnapshotLastTime, SnapshotLastTime - Time);
+			++NumWarnings;
+			UE_LOGF(LogTraceServices, Warning, "[Obj] Invalid timestamp for EndSnapshot event (%f < %f; delta=%f)!", Time, SnapshotLastTime, SnapshotLastTime - Time);
 		}
 		SnapshotLastTime = Time;
 
@@ -115,14 +161,15 @@ bool FObjectAnalyzer::OnEvent(uint16 RouteId, EStyle Style, const FOnEventContex
 			if (Snapshot != nullptr)
 			{
 				check(Snapshot == CurrentSnapshot);
-				UE_LOG(LogTraceServices, Log, TEXT("[UObj] Completed snapshot (%u objects, %u refs)."), Snapshot->GetNumObjects(), Snapshot->GetNumReferences());
+				UE_LOGF(LogTraceServices, Log, "[Obj] Completed snapshot (%u objects, %u refs).", Snapshot->GetNumObjects(), Snapshot->GetNumReferences());
 				
 				// debug
 				//Snapshot->SaveAs(*FString::Printf(TEXT("D:/uobjects_%u.csv"), Snapshot->GetId()));
 			}
 			else
 			{
-				UE_LOG(LogTraceServices, Warning, TEXT("[UObj] Completed empty snapshot!"));
+				++NumWarnings;
+				UE_LOGF(LogTraceServices, Warning, "[Obj] Completed empty snapshot!");
 			}
 		}
 		CurrentSnapshot = nullptr;
@@ -131,24 +178,17 @@ bool FObjectAnalyzer::OnEvent(uint16 RouteId, EStyle Style, const FOnEventContex
 
 	case RouteId_ObjectSpec:
 	{
-		if (!CurrentSnapshot)
+		UE_OBJECT_ANALYZER_EVENT("ObjectSpec");
+		if (!CheckValidSnapshotForEvent(EventName))
 		{
-			UE_LOG(LogTraceServices, Warning, TEXT("[UObj] Ignoring ObjectSpec event received outside a Begin-End Snapshot scope."));
 			break;
 		}
 
-		const uint32 ThreadId = Context.ThreadInfo.GetId();
-
 		uint32 Id = EventData.GetValue<uint32>("Id");
 		uint32 ClassId = EventData.GetValue<uint32>("ClassId");
 		uint32 OuterId = EventData.GetValue<uint32>("OuterId");
-		uint32 PackageId = EventData.GetValue<uint32>("PackageId");
 		uint32 Flags = EventData.GetValue<uint32>("Flags");
 
-		int32 StructureSize = EventData.GetValue<int32>("StructureSize");
-		uint64 SystemMemoryBytes = EventData.GetValue<uint64>("SystemMemoryBytes");
-		uint64 VideoMemoryBytes = EventData.GetValue<uint64>("VideoMemoryBytes");
-
 		//////////////////////////////////////////////////
 		// Read Object's Name
 
@@ -188,6 +228,16 @@ bool FObjectAnalyzer::OnEvent(uint16 RouteId, EStyle Style, const FOnEventContex
 		}
 		check(PersistentName != nullptr);
 
+		UE_OBJECT_ANALYZER_LOG_API_L1("Id=%u, ClassId=%u, OuterId=%u, Flags=%u, Name=\"%ls\"", Id, ClassId, OuterId, Flags, PersistentName);
+
+		//////////////////////////////////////////////////
+		// Temporary backward compatibility
+
+		uint32 PackageId = EventData.GetValue<uint32>("PackageId");
+		int32 StructureSize = EventData.GetValue<int32>("StructureSize");
+		uint64 SystemMemoryBytes = EventData.GetValue<uint64>("SystemMemoryBytes");
+		uint64 VideoMemoryBytes = EventData.GetValue<uint64>("VideoMemoryBytes");
+
 		const TCHAR* PersistentVersePath = nullptr;
 		{
 			FStringView VersePath;
@@ -198,48 +248,243 @@ bool FObjectAnalyzer::OnEvent(uint16 RouteId, EStyle Style, const FOnEventContex
 				PersistentVersePath = Session.StoreString(VersePath);
 			}
 		}
-		if (!PersistentVersePath)
-		{
-			FAnalysisSessionEditScope _(Session);
-			PersistentVersePath = Session.StoreString(TEXTVIEW(""));
-		}
 
 		//////////////////////////////////////////////////
 
 		FProviderEditScopeLock _(EditableProvider);
-		FObjectInfo* Object = EditableProvider.AddObject(Id, ClassId, OuterId, PackageId, PersistentName, PersistentVersePath, Flags);
+		FObjectInfo* Object = EditableProvider.AddObject(Id, ClassId, OuterId, Flags, PersistentName);
 		if (Object)
 		{
+			// Temporary backward compatibility
+			Object->PackageId = uint64(PackageId);
 			Object->StructureSize = StructureSize;
 			Object->SystemMemoryBytes = SystemMemoryBytes;
 			Object->VideoMemoryBytes = VideoMemoryBytes;
+			Object->VersePath = PersistentVersePath ? PersistentVersePath : TEXT("");
+
+			if (StructureSize != 0)
+			{
+				Object->FlagsEx |= EObjectInfoFlags::IsStruct;
+			}
 		}
 		else
 		{
-			UE_LOG(LogTraceServices, Error, TEXT("[UObj] Failed to add object with Id=%u (Name=\"%s\", VersePath=\"%s\")"), Id, PersistentName, PersistentVersePath);
+			if (++NumErrors < MaxErrorMessages)
+			{
+				UE_LOGF(LogTraceServices, Error, "[Obj] Failed to add object with Id=%u (Name=\"%ls\")!", Id, PersistentName);
+			}
+		}
+		break;
+	}
+
+	case RouteId_VerseExtraSpec:
+	{
+		UE_OBJECT_ANALYZER_EVENT("VerseExtraSpec");
+		if (!CheckValidSnapshotForEvent(EventName))
+		{
+			break;
+		}
+
+		uint32 Id = EventData.GetValue<uint32>("Id");
+
+		const TCHAR* PersistentVersePath = nullptr;
+		{
+			FStringView VersePath;
+			if (EventData.GetString("Path", VersePath) && !VersePath.IsEmpty())
+			{
+				FAnalysisSessionEditScope _(Session);
+				PersistentVersePath = Session.StoreString(VersePath);
+			}
+		}
+
+		UE_OBJECT_ANALYZER_LOG_API_L1("Id=%u, Path=\"%ls\"", Id, PersistentVersePath ? PersistentVersePath : TEXT(""));
+
+		FProviderEditScopeLock _(EditableProvider);
+		FObjectInfo* Object = EditableProvider.GetEditableObject(Id);
+		if (Object)
+		{
+			Object->VersePath = PersistentVersePath;
+		}
+		else
+		{
+			LogErrorForEventWithInvalidObjectId(EventName, Id);
+		}
+		break;
+	}
+
+	case RouteId_ResourceSizeExtraSpec:
+	{
+		UE_OBJECT_ANALYZER_EVENT("ResourceSizeExtraSpec");
+		if (!CheckValidSnapshotForEvent(EventName))
+		{
+			break;
+		}
+
+		uint32 Id = EventData.GetValue<uint32>("Id");
+		uint64 SystemMemoryBytes = EventData.GetValue<uint64>("SystemMemoryBytes");
+		uint64 VideoMemoryBytes = EventData.GetValue<uint64>("VideoMemoryBytes");
+		uint64 TotalSystemMemoryBytes = EventData.GetValue<uint64>("TotalSystemMemoryBytes");
+		uint64 TotalVideoMemoryBytes = EventData.GetValue<uint64>("TotalVideoMemoryBytes");
+
+		UE_OBJECT_ANALYZER_LOG_API_L1("Id=%u, SystemMemoryBytes=%llu (Total=%llu), VideoMemoryBytes=%llu (Total=%llu)",
+			Id, SystemMemoryBytes, TotalSystemMemoryBytes, VideoMemoryBytes, TotalVideoMemoryBytes);
+
+		FProviderEditScopeLock _(EditableProvider);
+		FObjectInfo* Object = EditableProvider.GetEditableObject(Id);
+		if (Object)
+		{
+			Object->SystemMemoryBytes = SystemMemoryBytes;
+			Object->VideoMemoryBytes = VideoMemoryBytes;
+			Object->TotalSystemMemoryBytes = TotalSystemMemoryBytes;
+			Object->TotalVideoMemoryBytes = TotalVideoMemoryBytes;
+		}
+		else
+		{
+			LogErrorForEventWithInvalidObjectId(EventName, Id);
+		}
+		break;
+	}
+
+	case RouteId_FieldExtraSpec:
+	{
+		UE_OBJECT_ANALYZER_EVENT("FieldExtraSpec");
+		if (!CheckValidSnapshotForEvent(EventName))
+		{
+			break;
+		}
+
+		uint32 Id = EventData.GetValue<uint32>("Id");
+
+		UE_OBJECT_ANALYZER_LOG_API_L1("Id=%u", Id);
+
+		FProviderEditScopeLock _(EditableProvider);
+		FObjectInfo* Object = EditableProvider.GetEditableObject(Id);
+		if (Object)
+		{
+			Object->FlagsEx |= EObjectInfoFlags::IsField;
+		}
+		else
+		{
+			LogErrorForEventWithInvalidObjectId(EventName, Id);
+		}
+		break;
+	}
+
+	case RouteId_StructExtraSpec:
+	{
+		UE_OBJECT_ANALYZER_EVENT("StructExtraSpec");
+		if (!CheckValidSnapshotForEvent(EventName))
+		{
+			break;
+		}
+
+		uint32 Id = EventData.GetValue<uint32>("Id");
+		uint32 SuperId = EventData.GetValue<uint32>("SuperId");
+		uint32 InheritanceSuperId = EventData.GetValue<uint32>("InheritanceSuperId");
+		int32 StructureSize = EventData.GetValue<int32>("StructureSize"); 
+
+		UE_OBJECT_ANALYZER_LOG_API_L1("Id=%u, SuperId=%u, InheritanceSuperId=%u, StructureSize=%d", Id, SuperId, InheritanceSuperId, StructureSize);
+
+		FProviderEditScopeLock _(EditableProvider);
+		FObjectInfo* Object = EditableProvider.GetEditableObject(Id);
+		if (Object)
+		{
+			Object->FlagsEx |= EObjectInfoFlags::IsStruct;
+			Object->SuperId = SuperId;
+			Object->InheritanceSuperId = InheritanceSuperId;
+			Object->StructureSize = StructureSize;
+		}
+		else
+		{
+			LogErrorForEventWithInvalidObjectId(EventName, Id);
+		}
+		break;
+	}
+
+	case RouteId_ClassExtraSpec:
+	{
+		UE_OBJECT_ANALYZER_EVENT("ClassExtraSpec");
+		if (!CheckValidSnapshotForEvent(EventName))
+		{
+			break;
+		}
+
+		uint32 Id = EventData.GetValue<uint32>("Id");
+
+		UE_OBJECT_ANALYZER_LOG_API_L1("Id=%u", Id);
+
+		FProviderEditScopeLock _(EditableProvider);
+		FObjectInfo* Object = EditableProvider.GetEditableObject(Id);
+		if (Object)
+		{
+			Object->FlagsEx |= EObjectInfoFlags::IsClass;
+		}
+		else
+		{
+			LogErrorForEventWithInvalidObjectId(EventName, Id);
+		}
+		break;
+	}
+
+	case RouteId_PackageExtraSpec:
+	{
+		UE_OBJECT_ANALYZER_EVENT("PackageExtraSpec");
+		if (!CheckValidSnapshotForEvent(EventName))
+		{
+			break;
+		}
+
+		uint32 Id = EventData.GetValue<uint32>("Id");
+		uint64 PackageId = EventData.GetValue<uint64>("PackageId");
+
+		const TCHAR* PersistentPath = nullptr;
+		{
+			FStringView Path;
+			if (EventData.GetString("LocalFullPath", Path) && !Path.IsEmpty())
+			{
+				FAnalysisSessionEditScope _(Session);
+				PersistentPath = Session.StoreString(Path);
+			}
+		}
+
+		UE_OBJECT_ANALYZER_LOG_API_L1("Id=%u, PackageId=%llu, LocalFullPath=\"%ls\"", Id, PackageId, PersistentPath ? PersistentPath : TEXT(""));
+
+		FProviderEditScopeLock _(EditableProvider);
+		FObjectInfo* Object = EditableProvider.GetEditableObject(Id);
+		if (Object)
+		{
+			Object->FlagsEx |= EObjectInfoFlags::IsPackage;
+			Object->PackageId = PackageId;
+			Object->PackagePath = PersistentPath;
+		}
+		else
+		{
+			LogErrorForEventWithInvalidObjectId(EventName, Id);
 		}
 		break;
 	}
 
 	case RouteId_ObjectRef:
 	{
-		if (!CurrentSnapshot)
+		UE_OBJECT_ANALYZER_EVENT("ObjectRef");
+		if (!CheckValidSnapshotForEvent(EventName))
 		{
-			UE_LOG(LogTraceServices, Warning, TEXT("[UObj] Ignoring ObjectRef event received outside a Begin-End Snapshot scope."));
 			break;
 		}
 
-		const uint32 ThreadId = Context.ThreadInfo.GetId();
 		uint32 ReferencerId = EventData.GetValue<uint32>("ReferencerId");
 		uint32 ObjectId = EventData.GetValue<uint32>("ObjectId");
 
-		//UE_LOG(LogTraceServices, VeryVerbose, TEXT("[UObj] Tid=%u ObjectReference ReferencerSerial=%d ObjectSerial=%d"), ThreadId, ReferencerSerial, ObjectSerial);
+		UE_OBJECT_ANALYZER_LOG_API_L1("ReferencerId=%u, ObjectId=%u", ReferencerId, ObjectId);
 
 		FProviderEditScopeLock _(EditableProvider);
 		FObjectReferenceInfo* Reference = EditableProvider.AddObjectReference(ReferencerId, ObjectId);
 		if (!Reference)
 		{
-			UE_LOG(LogTraceServices, Error, TEXT("[UObj] Failed to add reference (ReferencerId=%u, ObjectId=%u)"), ReferencerId, ObjectId);
+			if (++NumErrors < MaxErrorMessages)
+			{
+				UE_LOGF(LogTraceServices, Error, "[Obj] Failed to add reference (ReferencerId=%u, ObjectId=%u)!", ReferencerId, ObjectId);
+			}
 		}
 		break;
 	}
@@ -250,3 +495,7 @@ bool FObjectAnalyzer::OnEvent(uint16 RouteId, EStyle Style, const FOnEventContex
 }
 
 } // namespace TraceServices
+
+#undef UE_OBJECT_ANALYZER_LOG_API_L1
+#undef UE_OBJECT_ANALYZER_LOG_API_L0
+#undef UE_OBJECT_ANALYZER_EVENT
diff --git a/Engine/Source/Developer/TraceServices/Private/Analyzers/ObjectTraceAnalysis.h b/Engine/Source/Developer/TraceServices/Private/Analyzers/ObjectTraceAnalysis.h
--- a/Engine/Source/Developer/TraceServices/Private/Analyzers/ObjectTraceAnalysis.h
+++ b/Engine/Source/Developer/TraceServices/Private/Analyzers/ObjectTraceAnalysis.h
@@ -21,6 +21,12 @@ private:
 		RouteId_BeginSnapshot,
 		RouteId_EndSnapshot,
 		RouteId_ObjectSpec,
+		RouteId_VerseExtraSpec,
+		RouteId_ResourceSizeExtraSpec,
+		RouteId_FieldExtraSpec,
+		RouteId_StructExtraSpec,
+		RouteId_ClassExtraSpec,
+		RouteId_PackageExtraSpec,
 		RouteId_ObjectRef,
 	};
 
@@ -30,12 +36,20 @@ public:
 	virtual void OnAnalysisEnd() override;
 	virtual bool OnEvent(uint16 RouteId, EStyle Style, const FOnEventContext& Context) override;
 
+private:
+	bool CheckValidSnapshotForEvent(const ANSICHAR* EventName);
+	void LogErrorForEventWithInvalidObjectId(const ANSICHAR* EventName, uint32 Id);
+
 private:
 	IAnalysisSession& Session;
 	IEditableObjectProvider& EditableProvider;
 
 	IObjectEditableSnapshot* CurrentSnapshot = nullptr;
 	double SnapshotLastTime = 0.0;
+	uint32 NumWarnings = 0;
+	uint32 NumErrors = 0;
+	static constexpr uint32 MaxWarningMessages = 100;
+	static constexpr uint32 MaxErrorMessages = 100;
 };
 
 } // namespace TraceServices
diff --git a/Engine/Source/Developer/TraceServices/Private/Model/ObjectProvider.cpp b/Engine/Source/Developer/TraceServices/Private/Model/ObjectProvider.cpp
--- a/Engine/Source/Developer/TraceServices/Private/Model/ObjectProvider.cpp
+++ b/Engine/Source/Developer/TraceServices/Private/Model/ObjectProvider.cpp
@@ -88,7 +88,7 @@ void FObjectSnapshot::SaveAs(const TCHAR* CsvFileName) const
 	IFileHandle* File = PlatformFile.OpenWrite(CsvFileName, false, false);
 	if (File)
 	{
-		FUtf8StringView Header = UTF8TEXTVIEW("Id,ClassId,OuterId,PackageId,Flags,StructureSize,SystemMemoryBytes,VideoMemoryBytes,Name,VersePath\n");
+		FUtf8StringView Header = UTF8TEXTVIEW("Id,ClassId,OuterId,Flags,StructureSize,SystemMemoryBytes,VideoMemoryBytes,TotalSystemMemoryBytes,TotalVideoMemoryBytes,Name,VersePath\n");
 		File->Write((const uint8*)Header.GetData(), Header.NumBytes());
 		for (const FObjectInfo& Obj : Objects)
 		{
@@ -96,15 +96,16 @@ void FObjectSnapshot::SaveAs(const TCHAR* CsvFileName) const
 			{
 				continue;
 			}
-			FUtf8String Str = FUtf8String::Printf(UTF8TEXT("%d,%d,%d,%d,0x%08X,%d,%llu,%llu,%s,%s\n"),
+			FUtf8String Str = FUtf8String::Printf(UTF8TEXT("%d,%d,%d,0x%08X,%d,%llu,%llu,%llu,%llu,%s,%s\n"),
 				int32(Obj.Id),
 				int32(Obj.ClassId),
 				int32(Obj.OuterId),
-				int32(Obj.PackageId),
 				Obj.Flags, //TCHAR_TO_UTF8(*LexToString((EObjectFlags)Flags)),
 				Obj.StructureSize,
 				Obj.SystemMemoryBytes,
 				Obj.VideoMemoryBytes,
+				Obj.TotalSystemMemoryBytes,
+				Obj.TotalVideoMemoryBytes,
 				TCHAR_TO_UTF8(Obj.Name),
 				TCHAR_TO_UTF8(Obj.VersePath));
 			const auto& Data = Str.GetCharArray();
@@ -300,7 +301,7 @@ IObjectEditableSnapshot* FObjectProvider::EndSnapshot(double Time)
 
 ////////////////////////////////////////////////////////////////////////////////////////////////////
 
-FObjectInfo* FObjectProvider::AddObject(uint32 Id, uint32 ClassId, uint32 OuterId, uint32 PackageId, const TCHAR* PersistentName, const TCHAR* PersistentVersePath, uint32 Flags)
+FObjectInfo* FObjectProvider::AddObject(uint32 Id, uint32 ClassId, uint32 OuterId, uint32 Flags, const TCHAR* PersistentName)
 {
 	EditAccessCheck();
 
@@ -318,9 +319,7 @@ FObjectInfo* FObjectProvider::AddObject(uint32 Id, uint32 ClassId, uint32 OuterI
 	Object.Id = Id;
 	Object.ClassId = ClassId;
 	Object.OuterId = OuterId;
-	Object.PackageId = PackageId;
 	Object.Name = PersistentName;
-	Object.VersePath = PersistentVersePath;
 	Object.Flags = Flags;
 
 	return &Object;
@@ -328,6 +327,25 @@ FObjectInfo* FObjectProvider::AddObject(uint32 Id, uint32 ClassId, uint32 OuterI
 
 ////////////////////////////////////////////////////////////////////////////////////////////////////
 
+FObjectInfo* FObjectProvider::GetEditableObject(uint32 Id)
+{
+	EditAccessCheck();
+
+	if (!ensure(CurrentSnapshot))
+	{
+		return nullptr;
+	}
+
+	if (!FObjectSnapshot::IsValidObjectId(Id))
+	{
+		return nullptr;
+	}
+
+	return CurrentSnapshot->GetEditableObject(Id);
+}
+
+////////////////////////////////////////////////////////////////////////////////////////////////////
+
 FObjectReferenceInfo* FObjectProvider::AddObjectReference(uint32 ReferencerId, uint32 ObjectId)
 {
 	EditAccessCheck();
diff --git a/Engine/Source/Developer/TraceServices/Private/Model/ObjectProviderPrivate.h b/Engine/Source/Developer/TraceServices/Private/Model/ObjectProviderPrivate.h
--- a/Engine/Source/Developer/TraceServices/Private/Model/ObjectProviderPrivate.h
+++ b/Engine/Source/Developer/TraceServices/Private/Model/ObjectProviderPrivate.h
@@ -47,6 +47,10 @@ public:
 	{
 		return ObjectId < uint32(Objects.Num()) && Objects[ObjectId].Id == ObjectId ? &Objects[ObjectId] : nullptr;
 	}
+	virtual FObjectInfo* GetEditableObject(uint32 ObjectId) override
+	{
+		return ObjectId < uint32(Objects.Num()) && Objects[ObjectId].Id == ObjectId ? &Objects[ObjectId] : nullptr;
+	}
 
 	virtual TConstArrayView<FObjectInfo> GetObjectsArray() const override { return Objects; }
 	virtual TConstArrayView<FObjectReferenceInfo> GetReferencesArray() const override { return References; }
@@ -105,7 +109,8 @@ public:
 	virtual IObjectEditableSnapshot* GetCurrentSnapshot() override;
 	virtual IObjectEditableSnapshot* BeginSnapshot(double Time) override;
 	virtual IObjectEditableSnapshot* EndSnapshot(double Time) override;
-	virtual FObjectInfo* AddObject(uint32 Id, uint32 ClassId, uint32 OuterId, uint32 PackageId, const TCHAR* PersistentName, const TCHAR* PersistentVersePath, uint32 Flags) override;
+	virtual FObjectInfo* AddObject(uint32 Id, uint32 ClassId, uint32 OuterId, uint32 Flags, const TCHAR* PersistentName) override;
+	virtual FObjectInfo* GetEditableObject(uint32 Id) override;
 	virtual FObjectReferenceInfo* AddObjectReference(uint32 ReferencerId, uint32 ObjectId) override;
 	virtual const IObjectProvider* GetReadProvider() const override { return this; }
 
diff --git a/Engine/Source/Developer/TraceServices/Public/TraceServices/Model/ObjectProvider.h b/Engine/Source/Developer/TraceServices/Public/TraceServices/Model/ObjectProvider.h
--- a/Engine/Source/Developer/TraceServices/Public/TraceServices/Model/ObjectProvider.h
+++ b/Engine/Source/Developer/TraceServices/Public/TraceServices/Model/ObjectProvider.h
@@ -13,19 +13,38 @@
 namespace TraceServices
 {
 
+enum class EObjectInfoFlags : uint32
+{
+	None = 0,
+	IsField  = 1 << 0,
+	IsStruct = 1 << 1,
+	IsClass = 1 << 2,
+	IsPackage = 1 << 3,
+};
+ENUM_CLASS_FLAGS(EObjectInfoFlags);
+
 struct FObjectInfo
 {
 	uint32 Id = uint32(-1);
 	uint32 ClassId = uint32(-1);
 	uint32 OuterId = uint32(-1);
-	uint32 PackageId = uint32(-1);
 	const TCHAR* Name = nullptr;
-	const TCHAR* VersePath = nullptr;
 	uint32 Flags = 0;
-	int32 StructureSize = 0;
+	EObjectInfoFlags FlagsEx = EObjectInfoFlags::None;
+
+	const TCHAR* VersePath = nullptr;
+
 	uint64 SystemMemoryBytes = 0;
 	uint64 VideoMemoryBytes = 0;
-	uint64 _Reserved = 0;
+	uint64 TotalSystemMemoryBytes = 0;
+	uint64 TotalVideoMemoryBytes = 0;
+
+	uint32 SuperId = uint32(-1);
+	uint32 InheritanceSuperId = uint32(-1);
+	int32 StructureSize = 0;
+
+	uint64 PackageId = 0;
+	const TCHAR* PackagePath = nullptr;
 };
 
 struct FObjectReferenceInfo
@@ -53,6 +72,8 @@ class IObjectEditableSnapshot : public IObjectSnapshot
 {
 public:
 	virtual ~IObjectEditableSnapshot() = default;
+
+	virtual FObjectInfo* GetEditableObject(uint32 ObjectId) = 0;
 };
 
 class IObjectProvider : public IProvider
@@ -134,14 +155,14 @@ public:
 	 * It is only valid to call this between BeginSnapshot() and EndSnapshot() calls.
 	 * The resulting pointer is valid only under the current edit lock.
 	 */
-	virtual FObjectInfo* AddObject(
-		uint32 ObjectId,
-		uint32 ClassId,
-		uint32 OuterId,
-		uint32 PackageId,
-		const TCHAR* PersistentName,
-		const TCHAR* PersistentVersePath,
-		uint32 Flags) = 0;
+	virtual FObjectInfo* AddObject(uint32 ObjectId, uint32 ClassId, uint32 OuterId, uint32 Flags, const TCHAR* PersistentName) = 0;
+
+	/**
+	 * Gets the object identified by id.
+	 * It is only valid to call this between BeginSnapshot() and EndSnapshot() calls.
+	 * The resulting pointer is valid only under the current edit lock.
+	 */
+	virtual FObjectInfo* GetEditableObject(uint32 Id) = 0;
 
 	/**
 	 * Adds a dependency between objects, to the current snapshot.
