[Insights]
* ObjectProfiler: Added debug validation for a snapshot.
* ObjectProfiler: Added analysis for the Function extra spec events.
* ObjectProfiler: Added debug validation for a snapshot.

#rnx
#Synced-CL 51880104

diff --git a/Engine/Source/Developer/TraceInsights/Private/Insights/InsightsStyle.cpp b/Engine/Source/Developer/TraceInsights/Private/Insights/InsightsStyle.cpp
--- a/Engine/Source/Developer/TraceInsights/Private/Insights/InsightsStyle.cpp
+++ b/Engine/Source/Developer/TraceInsights/Private/Insights/InsightsStyle.cpp
@@ -300,12 +300,22 @@ void FInsightsStyle::FStyle::Initialize()
 	Set("Icons.AssetMemTag.TreeItem", new IMAGE_BRUSH_SVG("MemTag_Asset_12", Icon12x12));
 	Set("Icons.ClassMemTag.TreeItem", new IMAGE_BRUSH_SVG("MemTag_Class_12", Icon12x12));
 
-	Set("Icons.UObject.TreeItem", new IMAGE_BRUSH_SVG("UObject_12", Icon12x12));
-	Set("Icons.UField.TreeItem", new IMAGE_BRUSH_SVG("UField_12", Icon12x12)); // UField -> UObject
-	Set("Icons.UStruct.TreeItem", new IMAGE_BRUSH_SVG("UStruct_12", Icon12x12)); // UStruct -> UField
-	Set("Icons.UClass.TreeItem", new IMAGE_BRUSH_SVG("UClass_12", Icon12x12)); // UClass -> UStruct
-	Set("Icons.UPackage.TreeItem", new IMAGE_BRUSH_SVG("UPackage_12", Icon12x12)); // UPackage -> UObject
-	Set("Icons.AActor.TreeItem", new IMAGE_BRUSH_SVG("AActor_12", Icon12x12)); // AActor -> UObject
+	// Core UObject Classes
+	Set("Icons.UObject.TreeItem",       new IMAGE_BRUSH_SVG("UObject_12", Icon12x12));
+	Set("Icons.UField.TreeItem",        new IMAGE_BRUSH_SVG("UObject_12", Icon12x12)); // UField -> UObject
+	Set("Icons.UStruct.TreeItem",       new EDITOR_IMAGE_BRUSH_SVG("Starship/Common/Struct", Icon12x12)); // UStruct -> UField
+	Set("Icons.UClass.TreeItem",        new EDITOR_IMAGE_BRUSH_SVG("Starship/Common/Class", Icon12x12)); // UClass -> UStruct
+	Set("Icons.UFunction.TreeItem",     new EDITOR_IMAGE_BRUSH_SVG("Starship/GraphEditors/FunctionInputParameter", Icon12x12)); // UFunction -> UStruct
+	Set("Icons.UPackage.TreeItem",      new EDITOR_IMAGE_BRUSH_SVG("Starship/Common/ProjectPackage", Icon12x12)); // UPackage -> UObject
+
+	// Engine UObject Classes
+	Set("Icons.UTexture.TreeItem",      new IMAGE_BRUSH_SVG("UObject_12", Icon12x12)); // UTexture -> UStreamableRenderAsset -> UObject
+	Set("Icons.UMaterial.TreeItem",     new IMAGE_BRUSH_SVG("UObject_12", Icon12x12)); // UMaterial -> UMaterialInterface -> UObject
+	Set("Icons.UStaticMesh.TreeItem",   new EDITOR_IMAGE_BRUSH_SVG("Starship/Common/StaticMesh", Icon12x12)); // UStaticMesh -> UStreamableRenderAsset -> UObject
+	Set("Icons.USkeletalMesh.TreeItem", new EDITOR_IMAGE_BRUSH_SVG("Starship/Common/SkeletalMesh", Icon12x12)); // USkeletalMesh -> USkinnedAsset -> UStreamableRenderAsset -> UObject
+	Set("Icons.ULevel.TreeItem",        new EDITOR_IMAGE_BRUSH_SVG("Starship/Common/Levels", Icon12x12)); // ULevel -> UObject
+	Set("Icons.UWorld.TreeItem",        new EDITOR_IMAGE_BRUSH_SVG("Starship/Common/World", Icon12x12)); // UWorld -> UObject
+	Set("Icons.AActor.TreeItem",        new IMAGE_BRUSH_SVG("UObject_12", Icon12x12)); // AActor -> UObject
 
 	Set("Icons.HasGraph.TreeItem", new IMAGE_BRUSH_SVG("RoundedBullet", Icon12x12));
 
diff --git a/Engine/Source/Developer/TraceServices/Private/Analyzers/ObjectTraceAnalysis.cpp b/Engine/Source/Developer/TraceServices/Private/Analyzers/ObjectTraceAnalysis.cpp
--- a/Engine/Source/Developer/TraceServices/Private/Analyzers/ObjectTraceAnalysis.cpp
+++ b/Engine/Source/Developer/TraceServices/Private/Analyzers/ObjectTraceAnalysis.cpp
@@ -31,16 +31,19 @@ void FObjectAnalyzer::OnAnalysisBegin(const FOnAnalysisContext& Context)
 {
 	auto& Builder = Context.InterfaceBuilder;
 
-	Builder.RouteEvent(RouteId_BeginSnapshot,         "CoreUObject", "BeginSnapshot");
-	Builder.RouteEvent(RouteId_EndSnapshot,           "CoreUObject", "EndSnapshot");
-	Builder.RouteEvent(RouteId_ObjectSpec,            "CoreUObject", "ObjectSpec");
-	Builder.RouteEvent(RouteId_VerseExtraSpec,        "CoreUObject", "VerseExtraSpec");
-	Builder.RouteEvent(RouteId_ResourceSizeExtraSpec, "CoreUObject", "ResourceSizeExtraSpec");
-	Builder.RouteEvent(RouteId_FieldExtraSpec,        "CoreUObject", "FieldExtraSpec");
-	Builder.RouteEvent(RouteId_StructExtraSpec,       "CoreUObject", "StructExtraSpec");
-	Builder.RouteEvent(RouteId_ClassExtraSpec,        "CoreUObject", "ClassExtraSpec");
-	Builder.RouteEvent(RouteId_PackageExtraSpec,      "CoreUObject", "PackageExtraSpec");
-	Builder.RouteEvent(RouteId_ObjectRef,             "CoreUObject", "ObjectRef");
+#define ROUTE_EVENT(Event) Builder.RouteEvent(RouteId_##Event, "CoreUObject", #Event)
+	ROUTE_EVENT(BeginSnapshot);
+	ROUTE_EVENT(EndSnapshot);
+	ROUTE_EVENT(ObjectSpec);
+	ROUTE_EVENT(VerseExtraSpec);
+	ROUTE_EVENT(ResourceSizeExtraSpec);
+	ROUTE_EVENT(FieldExtraSpec);
+	ROUTE_EVENT(StructExtraSpec);
+	ROUTE_EVENT(ClassExtraSpec);
+	ROUTE_EVENT(FunctionExtraSpec);
+	ROUTE_EVENT(PackageExtraSpec);
+	ROUTE_EVENT(ObjectRef);
+#undef ROUTE_EVENT
 }
 
 void FObjectAnalyzer::OnAnalysisEnd()
@@ -121,13 +124,19 @@ bool FObjectAnalyzer::OnEvent(uint16 RouteId, EStyle Style, const FOnEventContex
 		}
 		check(CurrentSnapshot);
 
-		const uint32 SizeOfUObject = EventData.GetValue<uint32>("SizeOfUObject");
-		const uint32 SizeOfUField = EventData.GetValue<uint32>("SizeOfUField");
-		const uint32 SizeOfUStruct = EventData.GetValue<uint32>("SizeOfUStruct");
-		const uint32 SizeOfUClass = EventData.GetValue<uint32>("SizeOfUClass");
-		const uint32 SizeOfUPackage = EventData.GetValue<uint32>("SizeOfUPackage");
-		UE_LOGF(LogTraceServices, Log, "[Obj] sizeof(UObject)=%u, sizeof(UField)=%u, sizeof(UStruct)=%u, sizeof(UClass)=%u, sizeof(UPackage)=%u",
-			SizeOfUObject, SizeOfUField, SizeOfUStruct, SizeOfUClass, SizeOfUPackage);
+		const ANSICHAR* KnownClasses[] =
+		{
+			"UObject", "UField", "UStruct", "UClass", "UFunction", "UPackage",
+		};
+		constexpr int32 NumKnownClasses = UE_ARRAY_COUNT(KnownClasses);
+		for (int32 EventIndex = 0; EventIndex < NumKnownClasses; ++EventIndex)
+		{
+			ANSICHAR FieldName[32];
+			FCStringAnsi::Strcpy(FieldName, "SizeOf");
+			FCStringAnsi::Strcat(FieldName, KnownClasses[EventIndex]);
+			const uint32 SizeOf = EventData.GetValue<uint32>(FieldName);
+			UE_LOGF(LogTraceServices, Log, "[Obj] sizeof(%s)=%u", KnownClasses[EventIndex], SizeOf);
+		}
 
 		break;
 	}
@@ -162,7 +171,7 @@ bool FObjectAnalyzer::OnEvent(uint16 RouteId, EStyle Style, const FOnEventContex
 			{
 				check(Snapshot == CurrentSnapshot);
 				UE_LOGF(LogTraceServices, Log, "[Obj] Completed snapshot (%u objects, %u refs).", Snapshot->GetNumObjects(), Snapshot->GetNumReferences());
-				
+
 				// debug
 				//Snapshot->SaveAs(*FString::Printf(TEXT("D:/uobjects_%u.csv"), Snapshot->GetId()));
 			}
@@ -381,7 +390,7 @@ bool FObjectAnalyzer::OnEvent(uint16 RouteId, EStyle Style, const FOnEventContex
 		uint32 Id = EventData.GetValue<uint32>("Id");
 		uint32 SuperId = EventData.GetValue<uint32>("SuperId");
 		uint32 InheritanceSuperId = EventData.GetValue<uint32>("InheritanceSuperId");
-		int32 StructureSize = EventData.GetValue<int32>("StructureSize"); 
+		int32 StructureSize = EventData.GetValue<int32>("StructureSize");
 
 		UE_OBJECT_ANALYZER_LOG_API_L1("Id=%u, SuperId=%u, InheritanceSuperId=%u, StructureSize=%d", Id, SuperId, InheritanceSuperId, StructureSize);
 
@@ -426,6 +435,38 @@ bool FObjectAnalyzer::OnEvent(uint16 RouteId, EStyle Style, const FOnEventContex
 		break;
 	}
 
+	case RouteId_FunctionExtraSpec:
+	{
+		UE_OBJECT_ANALYZER_EVENT("FunctionExtraSpec");
+		if (!CheckValidSnapshotForEvent(EventName))
+		{
+			break;
+		}
+
+		uint32 Id = EventData.GetValue<uint32>("Id");
+		uint32 FunctionFlags = EventData.GetValue<uint32>("Flags");
+		uint8 FunctionNumParms = EventData.GetValue<uint8>("NumParms");
+		uint16 FunctionParmsSize = EventData.GetValue<uint16>("ParmsSize");
+
+		UE_OBJECT_ANALYZER_LOG_API_L1("Id=%u, NumParms=%u, ParmsSize=%u, Flags=%s",
+			Id, uint32(FunctionNumParms), uint32(FunctionParmsSize), *FunctionFlagsToUtf8String(FunctionFlags));
+
+		FProviderEditScopeLock _(EditableProvider);
+		FObjectInfo* Object = EditableProvider.GetEditableObject(Id);
+		if (Object)
+		{
+			Object->FlagsEx |= EObjectInfoFlags::IsFunction;
+			Object->FunctionFlags = FunctionFlags;
+			Object->FunctionNumParms = FunctionNumParms;
+			Object->FunctionParmsSize = FunctionParmsSize;
+		}
+		else
+		{
+			LogErrorForEventWithInvalidObjectId(EventName, Id);
+		}
+		break;
+	}
+
 	case RouteId_PackageExtraSpec:
 	{
 		UE_OBJECT_ANALYZER_EVENT("PackageExtraSpec");
diff --git a/Engine/Source/Developer/TraceServices/Private/Analyzers/ObjectTraceAnalysis.h b/Engine/Source/Developer/TraceServices/Private/Analyzers/ObjectTraceAnalysis.h
--- a/Engine/Source/Developer/TraceServices/Private/Analyzers/ObjectTraceAnalysis.h
+++ b/Engine/Source/Developer/TraceServices/Private/Analyzers/ObjectTraceAnalysis.h
@@ -26,6 +26,7 @@ private:
 		RouteId_FieldExtraSpec,
 		RouteId_StructExtraSpec,
 		RouteId_ClassExtraSpec,
+		RouteId_FunctionExtraSpec,
 		RouteId_PackageExtraSpec,
 		RouteId_ObjectRef,
 	};
diff --git a/Engine/Source/Developer/TraceServices/Private/Model/ObjectProvider.cpp b/Engine/Source/Developer/TraceServices/Private/Model/ObjectProvider.cpp
--- a/Engine/Source/Developer/TraceServices/Private/Model/ObjectProvider.cpp
+++ b/Engine/Source/Developer/TraceServices/Private/Model/ObjectProvider.cpp
@@ -9,6 +9,7 @@
 #include "GenericPlatform/GenericPlatformFile.h"
 
 // TraceServices
+#include "Common/Utils.h"
 #include "TraceServices/Model/AnalysisSession.h"
 #include "TraceServices/Containers/Allocators.h"
 
@@ -82,6 +83,211 @@ const FObjectInfo* FObjectSnapshot::FindObject(FStringView Name) const
 
 ////////////////////////////////////////////////////////////////////////////////////////////////////
 
+void FObjectSnapshot::Validate()
+{
+	UE_LOGF(LogTraceServices, Log, "[Obj] Validating objects...");
+
+	struct FObjectValidationStats
+	{
+		const TCHAR* BaseClassName;
+		uint32 BaseClassId;
+		EObjectInfoFlags SingleFlag;
+		EObjectInfoFlags ExpectedFlags;
+		uint32 NumClasses;
+		uint32 NumInstances;
+	};
+
+	constexpr EObjectInfoFlags ExpectedClassFlags = EObjectInfoFlags::IsClass | EObjectInfoFlags::IsStruct | EObjectInfoFlags::IsField;
+
+	FObjectValidationStats ValidationStats[] =
+	{
+#define DECLARE_VALIDATION_STATS(ClassName, ExpectedFlags) \
+		{ TEXT(#ClassName), InvalidObjectId, EObjectInfoFlags::Is##ClassName, ExpectedFlags, 0, 0 },
+		DECLARE_VALIDATION_STATS(Class, ExpectedClassFlags)
+		DECLARE_VALIDATION_STATS(Function, EObjectInfoFlags::IsFunction | EObjectInfoFlags::IsStruct | EObjectInfoFlags::IsField)
+		DECLARE_VALIDATION_STATS(Struct, EObjectInfoFlags::IsStruct | EObjectInfoFlags::IsField)
+		DECLARE_VALIDATION_STATS(Field, EObjectInfoFlags::IsField)
+		DECLARE_VALIDATION_STATS(Package, EObjectInfoFlags::IsPackage)
+#undef DECLARE_VALIDATION_STATS
+	};
+	constexpr int32 NumValidationStats = UE_ARRAY_COUNT(ValidationStats);
+
+	// Searches the UClass object and counts the number of valid objects.
+	// Also validates object flags for known instances.
+	uint32 TotalObjectCount = 0;
+	for (FObjectInfo& Obj : Objects)
+	{
+		if (Obj.Id == InvalidObjectId)
+		{
+			continue;
+		}
+		if (Obj.ClassId >= uint32(Objects.Num()))
+		{
+			UE_LOGF(LogTraceServices, Error, "[Obj] Object %u (%ls) with flags %u has invalid class id %u!",
+				Obj.Id, Obj.Name, Obj.FlagsEx, Obj.ClassId);
+			return;
+		}
+		FObjectInfo& ClassObj = Objects[Obj.ClassId];
+		if (ClassObj.Id != Obj.ClassId)
+		{
+			UE_LOGF(LogTraceServices, Error, "[Obj] Object %u (%ls) with flags %u has invalid class object at id %u!",
+				Obj.Id, Obj.Name, Obj.FlagsEx, Obj.ClassId);
+			return;
+		}
+		++TotalObjectCount;
+		for (int32 Index = 0; Index < NumValidationStats; ++Index)
+		{
+			FObjectValidationStats& Stats = ValidationStats[Index];
+			if (EnumHasAnyFlags(Obj.FlagsEx, Stats.SingleFlag))
+			{
+				++Stats.NumInstances;
+				ValidateExpectedFlags(Obj, Stats.ExpectedFlags);
+			}
+			if (Obj.SuperId != Obj.InheritanceSuperId &&
+				Obj.InheritanceSuperId != InvalidObjectId)
+			{
+				UE_LOGF(LogTraceServices, Log, "[Obj] Object %u (%ls) with class id %u (%ls) and flags %u has super id %d != inheritance super id %d!",
+					Obj.Id, Obj.Name, ClassObj.Id, ClassObj.Name, Obj.FlagsEx, int32(Obj.SuperId), int32(Obj.InheritanceSuperId));
+			}
+			if (Obj.SuperId != InvalidObjectId && Obj.SuperId >= uint32(Objects.Num()))
+			{
+				UE_LOGF(LogTraceServices, Warning, "[Obj] Object %u (%ls) with class id %u (%ls) and flags %u has invalid super id %u!",
+					Obj.Id, Obj.Name, ClassObj.Id, ClassObj.Name, Obj.FlagsEx, Obj.SuperId);
+				Obj.SuperId = InvalidObjectId;
+			}
+			if ((Stats.BaseClassId == InvalidObjectId) &&
+				(FCString::Strcmp(Obj.Name, Stats.BaseClassName) == 0))
+			{
+				Stats.BaseClassId = Obj.Id;
+				if (!EnumHasAllFlags(Obj.FlagsEx, ExpectedClassFlags))
+				{
+					UE_LOGF(LogTraceServices, Error, "[Obj] Object %u (%ls) with class id %u (%ls) is expected to be a class, but has invalid flags (current: %u, expected: %u)!",
+						Obj.Id, Obj.Name, ClassObj.Id, ClassObj.Name, Obj.FlagsEx, ExpectedClassFlags);
+				}
+			}
+		}
+	}
+	UE_LOGF(LogTraceServices, Log, "[Obj] Found %d objects.", TotalObjectCount);
+
+	const uint32 ClassClassId = ValidationStats[0].BaseClassId;
+	if (ClassClassId == InvalidObjectId)
+	{
+		UE_LOGF(LogTraceServices, Error, "[Obj] No UClass object!");
+		return;
+	}
+
+	// Gather classes.
+	uint32 DefaultClassCount = 0;
+	uint32 OtherClassCount = 0;
+	TArray<uint32> Classes;
+	for (FObjectInfo& Obj : Objects)
+	{
+		if (Obj.Id == InvalidObjectId)
+		{
+			continue;
+		}
+		FObjectInfo& ClassObj = Objects[Obj.ClassId];
+		if (Obj.ClassId == ClassClassId)
+		{
+			Classes.Add(Obj.Id);
+			ValidateExpectedFlags(Obj, ExpectedClassFlags);
+		}
+		else if (EnumHasAnyFlags(Obj.FlagsEx, EObjectInfoFlags::IsClass))
+		{
+			if (Obj.SuperId == InvalidObjectId)
+			{
+				++DefaultClassCount;
+				UE_LOGF(LogTraceServices, Log, "[Obj] Object %u (%ls) with class id %u (%ls), super id %d and flags %u has IsClass flag! Default__*Class?",
+					Obj.Id, Obj.Name, ClassObj.Id, ClassObj.Name, int32(Obj.SuperId), Obj.FlagsEx);
+			}
+			else
+			{
+				++OtherClassCount;
+				UE_LOGF(LogTraceServices, Verbose, "[Obj] Object %u (%ls) with class id %u (%ls), super id %d and flags %u has IsClass flag!",
+					Obj.Id, Obj.Name, ClassObj.Id, ClassObj.Name, int32(Obj.SuperId), Obj.FlagsEx);
+			}
+			Classes.Add(Obj.Id);
+			ValidateExpectedFlags(Obj, ExpectedClassFlags);
+		}
+	}
+	const uint32 BaseClassCount = uint32(Classes.Num()) - DefaultClassCount - OtherClassCount;
+	if (uint32(Classes.Num()) != ValidationStats[0].NumInstances)
+	{
+		UE_LOGF(LogTraceServices, Error, "[Obj] Found %u classes (base: %u + default: %u + other: %u) vs. %u class instances!",
+			uint32(Classes.Num()), BaseClassCount, DefaultClassCount, OtherClassCount, ValidationStats[0].NumInstances);
+	}
+	else
+	{
+		UE_LOGF(LogTraceServices, Log, "[Obj] Found %u classes (base: %u + default: %u + other: %u)!",
+			uint32(Classes.Num()), BaseClassCount, DefaultClassCount, OtherClassCount);
+	}
+
+	//////////////////////////////////////////////////
+
+	for (int32 Index = 0; Index < NumValidationStats; ++Index)
+	{
+		FObjectValidationStats& Stats = ValidationStats[Index];
+
+		if (Stats.BaseClassId != InvalidObjectId)
+		{
+			Stats.NumClasses = ValidateClassesRec(Classes, ClassClassId, Stats.BaseClassId, Stats.ExpectedFlags);
+		}
+		else
+		{
+			UE_LOGF(LogTraceServices, Error, "[Obj] No U%ls class!", Stats.BaseClassName);
+		}
+
+		UE_LOGF(LogTraceServices, Log, "[Obj] Class id %u (%ls) has %u derived classes and %u object instances.",
+			Stats.BaseClassId, Stats.BaseClassName, Stats.NumClasses, Stats.NumInstances);
+	}
+
+	//////////////////////////////////////////////////
+
+	UE_LOGF(LogTraceServices, Log, "[Obj] Validation completed.");
+}
+
+////////////////////////////////////////////////////////////////////////////////////////////////////
+
+void FObjectSnapshot::ValidateExpectedFlags(FObjectInfo& Obj, EObjectInfoFlags ExpectedFlags)
+{
+	if (!EnumHasAllFlags(Obj.FlagsEx, ExpectedFlags))
+	{
+		UE_LOGF(LogTraceServices, Warning, "[Obj] Object %u \"%ls\" (class id %u, super id %d) does not have all expected flags (current: %u, expected: %u)!",
+			Obj.Id, Obj.Name,
+			Obj.ClassId, int32(Obj.SuperId),
+			Obj.FlagsEx, ExpectedFlags);
+	}
+}
+
+////////////////////////////////////////////////////////////////////////////////////////////////////
+
+uint32 FObjectSnapshot::ValidateClassesRec(const TArray<uint32>& Classes, uint32 ClassClassId, uint32 BaseClassId, EObjectInfoFlags ExpectedClassFlags)
+{
+	uint32 Count = 0;
+	for (uint32 CurrentClassId : Classes)
+	{
+		FObjectInfo& Obj = Objects[CurrentClassId];
+		if (Obj.SuperId == BaseClassId)
+		{
+			FObjectInfo& ClassObj = Objects[Obj.ClassId];
+			UE_LOGF(LogTraceServices, Log, "[Obj] Object %u (%ls) with class id %u (%ls), super id %d, inheritance super id %d and flags %u.",
+				Obj.Id, Obj.Name, ClassObj.Id, ClassObj.Name, int32(Obj.SuperId), int32(Obj.InheritanceSuperId), Obj.FlagsEx);
+
+			if (Obj.ClassId != ClassClassId)
+			{
+				UE_LOGF(LogTraceServices, Warning, "[Obj] Object %u (%ls) is not class! It is %u (%ls)!",
+					Obj.Id, Obj.Name, Obj.ClassId, Objects[Obj.ClassId].Name);
+			}
+
+			++Count;
+			Count += ValidateClassesRec(Classes, ClassClassId, Obj.Id, ExpectedClassFlags);
+		}
+	}
+	return Count;
+}
+
+////////////////////////////////////////////////////////////////////////////////////////////////////
+
 void FObjectSnapshot::SaveAs(const TCHAR* CsvFileName) const
 {
 	IPlatformFile& PlatformFile = FPlatformFileManager::Get().GetPlatformFile();
@@ -292,6 +498,7 @@ IObjectEditableSnapshot* FObjectProvider::EndSnapshot(double Time)
 
 	CurrentSnapshot->EndTime = FMath::Max(Time, CurrentSnapshot->StartTime);
 	CurrentSnapshot->Id = uint32(Snapshots.Num());
+	CurrentSnapshot->Validate();
 	Snapshots.Add(CurrentSnapshot);
 
 	FObjectSnapshot* Snapshot = CurrentSnapshot;
@@ -368,7 +575,74 @@ FObjectReferenceInfo* FObjectProvider::AddObjectReference(uint32 ReferencerId, u
 	return &ReferenceInfo;
 }
 
+////////////////////////////////////////////////////////////////////////////////////////////////////
+
+FUtf8String FunctionFlagsToUtf8String(uint32 Flags)
+{
+	if (Flags == 0)
+	{
+		return FUtf8String();
+	}
+
+	const ANSICHAR* FlagNames[] =
+	{
+		"Final",
+		"RequiredAPI",
+		"BlueprintAuthorityOnly",
+		"BlueprintCosmetic",
+		"_UnusedBit4",
+		"_UnusedBit5",
+		"Net",
+		"NetReliable",
+		"NetRequest",
+		"Exec",
+		"Native",
+		"Event",
+		"NetResponse",
+		"Static",
+		"NetMulticast",
+		"UbergraphFunction",
+		"MulticastDelegate",
+		"Public",
+		"Private",
+		"Protected",
+		"Delegate",
+		"NetServer",
+		"HasOutParms",
+		"HasDefaults",
+		"NetClient",
+		"DLLImport",
+		"BlueprintCallable",
+		"BlueprintEvent",
+		"BlueprintPure",
+		"EditorOnly",
+		"Const",
+		"NetValidate"
+	};
+	static_assert(UE_ARRAY_COUNT(FlagNames) == 32, "");
+
+	bool bFirst = true;
+	TUtf8StringBuilder<256> Str;
+	for (int32 BitIndex = 0; BitIndex < 32; ++BitIndex)
+	{
+		if ((Flags & (1u << BitIndex)) != 0)
+		{
+			if (bFirst)
+			{
+				bFirst = false;
+			}
+			else
+			{
+				Str.Append(" | ");
+			}
+			Str.Append(FlagNames[BitIndex]);
+		}
+	}
+	return FUtf8String(Str.ToView());
+}
+
 ////////////////////////////////////////////////////////////////////////////////////////////////////
 ////////////////////////////////////////////////////////////////////////////////////////////////////
 
 } // namespace TraceServices
+
diff --git a/Engine/Source/Developer/TraceServices/Private/Model/ObjectProviderPrivate.h b/Engine/Source/Developer/TraceServices/Private/Model/ObjectProviderPrivate.h
--- a/Engine/Source/Developer/TraceServices/Private/Model/ObjectProviderPrivate.h
+++ b/Engine/Source/Developer/TraceServices/Private/Model/ObjectProviderPrivate.h
@@ -64,8 +64,14 @@ public:
 	void EnumerateObjects(TFunctionRef<bool(const FObjectInfo&)> Callback) const;
 	const FObjectInfo* FindObject(FStringView Name) const;
 
+	void Validate();
+
 	virtual void SaveAs(const TCHAR* CsvFileName) const;
 
+private:
+	void ValidateExpectedFlags(FObjectInfo& Obj, EObjectInfoFlags ExpectedFlags);
+	uint32 ValidateClassesRec(const TArray<uint32>& Classes, uint32 ClassClassId, uint32 BaseClassId, EObjectInfoFlags ExpectedClassFlags);
+
 private:
 	uint32 Id; // snapshot's unique id
 	double StartTime = 0.0;
diff --git a/Engine/Source/Developer/TraceServices/Public/TraceServices/Model/ObjectProvider.h b/Engine/Source/Developer/TraceServices/Public/TraceServices/Model/ObjectProvider.h
--- a/Engine/Source/Developer/TraceServices/Public/TraceServices/Model/ObjectProvider.h
+++ b/Engine/Source/Developer/TraceServices/Public/TraceServices/Model/ObjectProvider.h
@@ -13,13 +13,16 @@
 namespace TraceServices
 {
 
+FUtf8String FunctionFlagsToUtf8String(uint32 Flags);
+
 enum class EObjectInfoFlags : uint32
 {
 	None = 0,
 	IsField  = 1 << 0,
 	IsStruct = 1 << 1,
 	IsClass = 1 << 2,
-	IsPackage = 1 << 3,
+	IsFunction = 1 << 3,
+	IsPackage = 1 << 4,
 };
 ENUM_CLASS_FLAGS(EObjectInfoFlags);
 
@@ -43,6 +46,10 @@ struct FObjectInfo
 	uint32 InheritanceSuperId = uint32(-1);
 	int32 StructureSize = 0;
 
+	uint32 FunctionFlags = 0;
+	uint8 FunctionNumParms = 0;
+	uint16 FunctionParmsSize = 0;
+
 	uint64 PackageId = 0;
 	const TCHAR* PackagePath = nullptr;
 };
