Fix Linux Unreal Insights memory trace symbolication

#rnx
#Synced-CL 52132593

diff --git a/Engine/Source/Developer/TraceServices/Private/Model/SymslibResolver.cpp b/Engine/Source/Developer/TraceServices/Private/Model/SymslibResolver.cpp
--- a/Engine/Source/Developer/TraceServices/Private/Model/SymslibResolver.cpp
+++ b/Engine/Source/Developer/TraceServices/Private/Model/SymslibResolver.cpp
@@ -377,8 +377,22 @@ namespace
 		SYMS_DbgAccel* DbgAccel = syms_dbg_accel_from_file(Arena, File.GetData(), Accel);
 		if (!syms_accel_is_good(DbgAccel))
 		{
-			UE_LOG(LogSymslib, Warning, TEXT("Cannot parse debug symbols file '%s'"), *DebugPath);
-			return false;
+			// syms_dbg_accel_from_file only handles PDB. For ELF .debug files (separate DWARF),
+			// we need to go through the bin path: parse as binary, then extract debug info from it.
+			if (syms_file_is_bin(Accel))
+			{
+				SYMS_BinAccel* DbgBin = syms_bin_accel_from_file(Arena, File.GetData(), Accel);
+				if (syms_accel_is_good(DbgBin) && syms_bin_is_dbg(DbgBin))
+				{
+					DbgAccel = syms_dbg_accel_from_bin(Arena, File.GetData(), DbgBin);
+				}
+			}
+
+			if (!syms_accel_is_good(DbgAccel))
+			{
+				UE_LOG(LogSymslib, Warning, TEXT("Cannot parse debug symbols file '%s'"), *DebugPath);
+				return false;
+			}
 		}
 
 		File.FilePath = DebugPath;
diff --git a/Engine/Source/Runtime/Core/Private/ProfilingDebugging/Unix/UnixModuleDiagnostics.cpp b/Engine/Source/Runtime/Core/Private/ProfilingDebugging/Unix/UnixModuleDiagnostics.cpp
--- a/Engine/Source/Runtime/Core/Private/ProfilingDebugging/Unix/UnixModuleDiagnostics.cpp
+++ b/Engine/Source/Runtime/Core/Private/ProfilingDebugging/Unix/UnixModuleDiagnostics.cpp
@@ -62,13 +62,12 @@ void Modules_Initialize()
 			}
 		}
 
+		// Record the full path so things can be found later
 		FString ImageName = FString(Info->dlpi_name);
 		if (ImageName.IsEmpty())
 		{
-			constexpr bool bRemoveExtension = false;
-			ImageName = FPlatformProcess::ExecutableName(bRemoveExtension);
+			ImageName = FPlatformProcess::ExecutablePath();
 		}
-		ImageName = FPaths::GetCleanFilename(ImageName);
 
 		UE_TRACE_LOG(Diagnostics, ModuleLoad, ModuleChannel, sizeof(TCHAR) * ImageName.Len() + BuildIdSize)
 			<< ModuleLoad.Name(*ImageName, ImageName.Len())
