[Insights]
* TraceAnalysis: Added FEventFieldInfo::GetRefUid().
* TraceAnalysis: Added debug log limit for logs emitted when UE_TRACE_ANALYSIS_DEBUG is enabled.
* TraceTrimmer: Fixed the scoped events to also be filtered out when the payload event is filtered out.
* TraceTrimmer: Fixed filtering by name to also filter the NewEvent events and the Important events.
* TraceTrimmer: Significantly improved performance when filtering by name.
* TraceTrimmer: Added support for "DefinitionId" fields and partial support for reference fields.

#Synced-CL 52347234

diff --git a/Engine/Source/Developer/TraceAnalysis/Private/Analysis/Engine.cpp b/Engine/Source/Developer/TraceAnalysis/Private/Analysis/Engine.cpp
--- a/Engine/Source/Developer/TraceAnalysis/Private/Analysis/Engine.cpp
+++ b/Engine/Source/Developer/TraceAnalysis/Private/Analysis/Engine.cpp
@@ -736,6 +736,13 @@ uint8 IAnalyzer::FEventFieldInfo::GetSize() const
 	return uint8(Inner->Size);
 }
 
+////////////////////////////////////////////////////////////////////////////////
+uint32 IAnalyzer::FEventFieldInfo::GetRefUid() const
+{
+	const auto* Inner = (const FDispatch::FField*)this;
+	return Inner->RefUid;
+}
+
 
 
 // {{{1 array-reader -----------------------------------------------------------
@@ -2063,6 +2070,7 @@ public:
 	void				LeaveScopeB(uint64 BaseRelativeTimestamp);
 	void				OnVersion(uint32 TransportVersion, uint32 ProtocolVersion) { AnalyzerHub.OnVersion(TransportVersion, ProtocolVersion); }
 #if UE_TRACE_ANALYSIS_DEBUG
+	bool				ShouldWriteDebugLog();
 	void				DebugLogEvent(const FTypeRegistry::FTypeInfo* TypeInfo, uint32 FixedSize, uint32 AuxSize, uint32 Serial);
 	void				DebugLogNewEvent(uint32 Uid, const FTypeRegistry::FTypeInfo* TypeInfo, uint32 EventSize);
 	void				DebugLogEnterScopeEvent(uint32 Uid, uint32 EventSize);
@@ -2271,6 +2279,14 @@ void FAnalysisBridge::DispatchLeaveScope()
 }
 
 #if UE_TRACE_ANALYSIS_DEBUG
+////////////////////////////////////////////////////////////////////////////////
+bool FAnalysisBridge::ShouldWriteDebugLog()
+{
+	// Example for limiting debug logs to first 100K events...
+	//return (State.TotalEventCount < 100000);
+	return true;
+}
+
 ////////////////////////////////////////////////////////////////////////////////
 void FAnalysisBridge::DebugLogEvent(const FTypeRegistry::FTypeInfo* TypeInfo, uint32 FixedSize, uint32 AuxSize, uint32 Serial)
 {
@@ -2279,23 +2295,27 @@ void FAnalysisBridge::DebugLogEvent(const FTypeRegistry::FTypeInfo* TypeInfo, ui
 	State.TotalEventSize += EventSize;
 
 #if UE_TRACE_ANALYSIS_DEBUG_LEVEL >= 2
-	UE_TRACE_ANALYSIS_DEBUG_BeginStringBuilder();
-	UE_TRACE_ANALYSIS_DEBUG_Appendf("[EVENT %llu]", State.TotalEventCount);
-
-	if (GetActiveThreadId() != ~0u)
+	const bool bDebugLog = ShouldWriteDebugLog();
+	if (bDebugLog)
 	{
-		UE_TRACE_ANALYSIS_DEBUG_Appendf(" Tid=%u", GetActiveThreadId());
-	}
+		UE_TRACE_ANALYSIS_DEBUG_BeginStringBuilder();
+		UE_TRACE_ANALYSIS_DEBUG_Appendf("[EVENT %llu]", State.TotalEventCount);
+
+		if (GetActiveThreadId() != ~0u)
+		{
+			UE_TRACE_ANALYSIS_DEBUG_Appendf(" Tid=%u", GetActiveThreadId());
+		}
 
-	UE_TRACE_ANALYSIS_DEBUG_Appendf(" Uid=%u", uint32(TypeInfo->Uid));
+		UE_TRACE_ANALYSIS_DEBUG_Appendf(" Uid=%u", uint32(TypeInfo->Uid));
 
-	const IAnalyzer::FEventTypeInfo& EventTypeInfo = *(IAnalyzer::FEventTypeInfo*)TypeInfo;
-	UE_TRACE_ANALYSIS_DEBUG_Appendf(" (%s.%s)", EventTypeInfo.GetLoggerName(), EventTypeInfo.GetName());
+		const IAnalyzer::FEventTypeInfo& EventTypeInfo = *(IAnalyzer::FEventTypeInfo*)TypeInfo;
+		UE_TRACE_ANALYSIS_DEBUG_Appendf(" (%s.%s)", EventTypeInfo.GetLoggerName(), EventTypeInfo.GetName());
 
-	UE_TRACE_ANALYSIS_DEBUG_Appendf(" Size=%u", EventSize);
-	if (AuxSize != 0)
-	{
-		UE_TRACE_ANALYSIS_DEBUG_Appendf(" (%u + %u)", FixedSize, AuxSize);
+		UE_TRACE_ANALYSIS_DEBUG_Appendf(" Size=%u", EventSize);
+		if (AuxSize != 0)
+		{
+			UE_TRACE_ANALYSIS_DEBUG_Appendf(" (%u + %u)", FixedSize, AuxSize);
+		}
 	}
 #endif // UE_TRACE_ANALYSIS_DEBUG_LEVEL
 
@@ -2304,7 +2324,10 @@ void FAnalysisBridge::DebugLogEvent(const FTypeRegistry::FTypeInfo* TypeInfo, ui
 		++State.SyncEventCount;
 		State.SyncEventSize += EventSize;
 #if UE_TRACE_ANALYSIS_DEBUG_LEVEL >= 2
-		UE_TRACE_ANALYSIS_DEBUG_Appendf(" Sync=%u", Serial);
+		if (bDebugLog)
+		{
+			UE_TRACE_ANALYSIS_DEBUG_Appendf(" Sync=%u", Serial);
+		}
 #endif // UE_TRACE_ANALYSIS_DEBUG_LEVEL
 	}
 	else if (TypeInfo->IsImportant())
@@ -2312,7 +2335,10 @@ void FAnalysisBridge::DebugLogEvent(const FTypeRegistry::FTypeInfo* TypeInfo, ui
 		++State.ImportantNoSyncEventCount;
 		State.ImportantNoSyncEventSize += EventSize;
 #if UE_TRACE_ANALYSIS_DEBUG_LEVEL >= 2
-		UE_TRACE_ANALYSIS_DEBUG_Append(" Important");
+		if (bDebugLog)
+		{
+			UE_TRACE_ANALYSIS_DEBUG_Append(" Important");
+		}
 #endif // UE_TRACE_ANALYSIS_DEBUG_LEVEL
 	}
 	else
@@ -2322,7 +2348,10 @@ void FAnalysisBridge::DebugLogEvent(const FTypeRegistry::FTypeInfo* TypeInfo, ui
 	}
 
 #if UE_TRACE_ANALYSIS_DEBUG_LEVEL >= 2
-	UE_TRACE_ANALYSIS_DEBUG_EndStringBuilder();
+	if (bDebugLog)
+	{
+		UE_TRACE_ANALYSIS_DEBUG_EndStringBuilder();
+	}
 #endif // UE_TRACE_ANALYSIS_DEBUG_LEVEL
 }
 
@@ -2334,8 +2363,12 @@ void FAnalysisBridge::DebugLogNewEvent(uint32 Uid, const FTypeRegistry::FTypeInf
 	State.TotalEventSize += EventSize;
 	State.NewEventSize += EventSize;
 #if UE_TRACE_ANALYSIS_DEBUG_LEVEL >= 2
-	const uint32 Tid = ETransportTid::Importants; // GetActiveThreadId()
-	UE_TRACE_ANALYSIS_DEBUG_LOG("[EVENT %llu] Tid=%u Uid=%u (NewEvent) Size=%u", State.TotalEventCount, Tid, Uid, EventSize);
+	const bool bDebugLog = ShouldWriteDebugLog();
+	if (bDebugLog)
+	{
+		const uint32 Tid = ETransportTid::Importants; // GetActiveThreadId()
+		UE_TRACE_ANALYSIS_DEBUG_LOG("[EVENT %llu] Tid=%u Uid=%u (NewEvent) Size=%u", State.TotalEventCount, Tid, Uid, EventSize);
+	}
 #endif // UE_TRACE_ANALYSIS_DEBUG_LEVEL
 }
 
@@ -2347,8 +2380,12 @@ void FAnalysisBridge::DebugLogEnterScopeEvent(uint32 Uid, uint32 EventSize)
 	State.TotalEventSize += EventSize;
 	State.EnterScopeEventSize += EventSize;
 #if UE_TRACE_ANALYSIS_DEBUG_LEVEL >= 2
-	UE_TRACE_ANALYSIS_DEBUG_LOG("[EVENT %llu] Tid=%u Uid=%u (EnterScope) Size=%u",
-		State.TotalEventCount, GetActiveThreadId(), Uid, EventSize);
+	const bool bDebugLog = ShouldWriteDebugLog();
+	if (bDebugLog)
+	{
+		UE_TRACE_ANALYSIS_DEBUG_LOG("[EVENT %llu] Tid=%u Uid=%u (EnterScope) Size=%u",
+			State.TotalEventCount, GetActiveThreadId(), Uid, EventSize);
+	}
 #endif // UE_TRACE_ANALYSIS_DEBUG_LEVEL
 }
 
@@ -2360,11 +2397,15 @@ void FAnalysisBridge::DebugLogEnterScopeEvent(uint32 Uid, uint64 RelativeTimesta
 	State.TotalEventSize += EventSize;
 	State.EnterScopeTEventSize += EventSize;
 #if UE_TRACE_ANALYSIS_DEBUG_LEVEL >= 2
-	check(ThreadInfo);
-	UE_TRACE_ANALYSIS_DEBUG_LOG("[EVENT %llu] Tid=%u Uid=%u (EnterScope_T) Timestamp=(+%llu)=%llu Size=%u",
-		State.TotalEventCount, GetActiveThreadId(), Uid,
-		RelativeTimestamp, State.Timing.BaseTimestamp + ThreadInfo->PrevTimestamp + RelativeTimestamp,
-		EventSize);
+	const bool bDebugLog = ShouldWriteDebugLog();
+	if (bDebugLog)
+	{
+		check(ThreadInfo);
+		UE_TRACE_ANALYSIS_DEBUG_LOG("[EVENT %llu] Tid=%u Uid=%u (EnterScope_T) Timestamp=(+%llu)=%llu Size=%u",
+			State.TotalEventCount, GetActiveThreadId(), Uid,
+			RelativeTimestamp, State.Timing.BaseTimestamp + ThreadInfo->PrevTimestamp + RelativeTimestamp,
+			EventSize);
+	}
 #endif // UE_TRACE_ANALYSIS_DEBUG_LEVEL
 }
 
@@ -2376,10 +2417,14 @@ void FAnalysisBridge::DebugLogEnterScopeAEvent(uint32 Uid, uint64 AbsoluteTimest
 	State.TotalEventSize += EventSize;
 	State.EnterScopeTEventSize += EventSize;
 #if UE_TRACE_ANALYSIS_DEBUG_LEVEL >= 2
-	UE_TRACE_ANALYSIS_DEBUG_LOG("[EVENT %llu] Tid=%u Uid=%u (EnterScope_TA) Timestamp=%llu Size=%u",
-		State.TotalEventCount, GetActiveThreadId(), Uid,
-		AbsoluteTimestamp,
-		EventSize);
+	const bool bDebugLog = ShouldWriteDebugLog();
+	if (bDebugLog)
+	{
+		UE_TRACE_ANALYSIS_DEBUG_LOG("[EVENT %llu] Tid=%u Uid=%u (EnterScope_TA) Timestamp=%llu Size=%u",
+			State.TotalEventCount, GetActiveThreadId(), Uid,
+			AbsoluteTimestamp,
+			EventSize);
+	}
 #endif // UE_TRACE_ANALYSIS_DEBUG_LEVEL
 }
 
@@ -2391,10 +2436,14 @@ void FAnalysisBridge::DebugLogEnterScopeBEvent(uint32 Uid, uint64 BaseAbsoluteTi
 	State.TotalEventSize += EventSize;
 	State.EnterScopeTEventSize += EventSize;
 #if UE_TRACE_ANALYSIS_DEBUG_LEVEL >= 2
-	UE_TRACE_ANALYSIS_DEBUG_LOG("[EVENT %llu] Tid=%u Uid=%u (EnterScope_TB) Timestamp=(+%llu)=%llu Size=%u",
-		State.TotalEventCount, GetActiveThreadId(), Uid,
-		BaseAbsoluteTimestamp, State.Timing.BaseTimestamp + BaseAbsoluteTimestamp,
-		EventSize);
+	const bool bDebugLog = ShouldWriteDebugLog();
+	if (bDebugLog)
+	{
+		UE_TRACE_ANALYSIS_DEBUG_LOG("[EVENT %llu] Tid=%u Uid=%u (EnterScope_TB) Timestamp=(+%llu)=%llu Size=%u",
+			State.TotalEventCount, GetActiveThreadId(), Uid,
+			BaseAbsoluteTimestamp, State.Timing.BaseTimestamp + BaseAbsoluteTimestamp,
+			EventSize);
+	}
 #endif // UE_TRACE_ANALYSIS_DEBUG_LEVEL
 }
 
@@ -2406,8 +2455,12 @@ void FAnalysisBridge::DebugLogLeaveScopeEvent(uint32 Uid, uint32 EventSize)
 	State.TotalEventSize += EventSize;
 	State.LeaveScopeEventSize += EventSize;
 #if UE_TRACE_ANALYSIS_DEBUG_LEVEL >= 2
-	UE_TRACE_ANALYSIS_DEBUG_LOG("[EVENT %llu] Tid=%u Uid=%u (LeaveScope) Size=%u",
-		State.TotalEventCount, GetActiveThreadId(), Uid, EventSize);
+	const bool bDebugLog = ShouldWriteDebugLog();
+	if (bDebugLog)
+	{
+		UE_TRACE_ANALYSIS_DEBUG_LOG("[EVENT %llu] Tid=%u Uid=%u (LeaveScope) Size=%u",
+			State.TotalEventCount, GetActiveThreadId(), Uid, EventSize);
+	}
 #endif // UE_TRACE_ANALYSIS_DEBUG_LEVEL
 }
 
@@ -2419,11 +2472,15 @@ void FAnalysisBridge::DebugLogLeaveScopeEvent(uint32 Uid, uint64 RelativeTimesta
 	State.TotalEventSize += EventSize;
 	State.LeaveScopeTEventSize += EventSize;
 #if UE_TRACE_ANALYSIS_DEBUG_LEVEL >= 2
-	check(ThreadInfo);
-	UE_TRACE_ANALYSIS_DEBUG_LOG("[EVENT %llu] Tid=%u Uid=%u (LeaveScope_T) Timestamp=(+%llu)=%llu Size=%u",
-		State.TotalEventCount, GetActiveThreadId(), Uid,
-		RelativeTimestamp, State.Timing.BaseTimestamp + ThreadInfo->PrevTimestamp + RelativeTimestamp,
-		EventSize);
+	const bool bDebugLog = ShouldWriteDebugLog();
+	if (bDebugLog)
+	{
+		check(ThreadInfo);
+		UE_TRACE_ANALYSIS_DEBUG_LOG("[EVENT %llu] Tid=%u Uid=%u (LeaveScope_T) Timestamp=(+%llu)=%llu Size=%u",
+			State.TotalEventCount, GetActiveThreadId(), Uid,
+			RelativeTimestamp, State.Timing.BaseTimestamp + ThreadInfo->PrevTimestamp + RelativeTimestamp,
+			EventSize);
+	}
 #endif // UE_TRACE_ANALYSIS_DEBUG_LEVEL
 }
 
@@ -2435,10 +2492,14 @@ void FAnalysisBridge::DebugLogLeaveScopeAEvent(uint32 Uid, uint64 AbsoluteTimest
 	State.TotalEventSize += EventSize;
 	State.LeaveScopeTEventSize += EventSize;
 #if UE_TRACE_ANALYSIS_DEBUG_LEVEL >= 2
-	UE_TRACE_ANALYSIS_DEBUG_LOG("[EVENT %llu] Tid=%u Uid=%u (LeaveScope_TA) Timestamp=%llu Size=%u",
-		State.TotalEventCount, GetActiveThreadId(), Uid,
-		AbsoluteTimestamp,
-		EventSize);
+	const bool bDebugLog = ShouldWriteDebugLog();
+	if (bDebugLog)
+	{
+		UE_TRACE_ANALYSIS_DEBUG_LOG("[EVENT %llu] Tid=%u Uid=%u (LeaveScope_TA) Timestamp=%llu Size=%u",
+			State.TotalEventCount, GetActiveThreadId(), Uid,
+			AbsoluteTimestamp,
+			EventSize);
+	}
 #endif // UE_TRACE_ANALYSIS_DEBUG_LEVEL
 }
 
@@ -2450,10 +2511,14 @@ void FAnalysisBridge::DebugLogLeaveScopeBEvent(uint32 Uid, uint64 BaseAbsoluteTi
 	State.TotalEventSize += EventSize;
 	State.LeaveScopeTEventSize += EventSize;
 #if UE_TRACE_ANALYSIS_DEBUG_LEVEL >= 2
-	UE_TRACE_ANALYSIS_DEBUG_LOG("[EVENT %llu] Tid=%u Uid=%u (LeaveScope_TB) Timestamp=(+%llu)=%llu Size=%u",
-		State.TotalEventCount, GetActiveThreadId(), Uid,
-		BaseAbsoluteTimestamp, State.Timing.BaseTimestamp + BaseAbsoluteTimestamp,
-		EventSize);
+	const bool bDebugLog = ShouldWriteDebugLog();
+	if (bDebugLog)
+	{
+		UE_TRACE_ANALYSIS_DEBUG_LOG("[EVENT %llu] Tid=%u Uid=%u (LeaveScope_TB) Timestamp=(+%llu)=%llu Size=%u",
+			State.TotalEventCount, GetActiveThreadId(), Uid,
+			BaseAbsoluteTimestamp, State.Timing.BaseTimestamp + BaseAbsoluteTimestamp,
+			EventSize);
+	}
 #endif // UE_TRACE_ANALYSIS_DEBUG_LEVEL
 }
 #endif // UE_TRACE_ANALYSIS_DEBUG
diff --git a/Engine/Source/Developer/TraceAnalysis/Public/Trace/Analyzer.h b/Engine/Source/Developer/TraceAnalysis/Public/Trace/Analyzer.h
--- a/Engine/Source/Developer/TraceAnalysis/Public/Trace/Analyzer.h
+++ b/Engine/Source/Developer/TraceAnalysis/Public/Trace/Analyzer.h
@@ -80,6 +80,9 @@ public:
 
 		/** Gets the size in bytes for the type of this field. */
 		UE_API uint8 GetTypeSize() const;
+
+		/** Gets the Reference Uid for this field. */
+		UE_API uint32 GetRefUid() const;
 	};
 
 	struct FEventFieldHandle
diff --git a/Engine/Source/Developer/TraceServices/Private/Analyzers/TrimAnalyzer.cpp b/Engine/Source/Developer/TraceServices/Private/Analyzers/TrimAnalyzer.cpp
--- a/Engine/Source/Developer/TraceServices/Private/Analyzers/TrimAnalyzer.cpp
+++ b/Engine/Source/Developer/TraceServices/Private/Analyzers/TrimAnalyzer.cpp
@@ -63,12 +63,16 @@ FTrimAnalyzer::FTrimAnalyzer(FTrimParameters& InParameters, UE::Trace::FTraceWri
 	, TraceWriter(InTraceWriter)
 {
 	// Empty IncludeFilter is equivalent with "*.*" (include all).
-	if (!Parameters.Include.Equals("*.*"))
+	if (!Parameters.Include.IsEmpty() &&
+		!Parameters.Include.Equals("*.*"))
 	{
-		IncludeFilter.Init(Parameters.Include);
+		InitNameFilter(Parameters.Include, IncludeFilter);
 	}
 
-	ExcludeFilter.Init(Parameters.Exclude);
+	if (!Parameters.Exclude.IsEmpty())
+	{
+		InitNameFilter(Parameters.Exclude, ExcludeFilter);
+	}
 }
 
 ////////////////////////////////////////////////////////////////////////////////////////////////////
@@ -79,10 +83,10 @@ FTrimAnalyzer::~FTrimAnalyzer()
 
 ////////////////////////////////////////////////////////////////////////////////////////////////////
 
-void FTrimAnalyzer::FNameFilter::Init(FAnsiString Filter)
+void FTrimAnalyzer::InitNameFilter(FAnsiString InFilter, FNameFilter& OutNameFilter)
 {
 	TArray<FAnsiString> FilterPatterns;
-	Filter.ParseIntoArray(FilterPatterns, ",", true);
+	InFilter.ParseIntoArray(FilterPatterns, ",", true);
 
 	for (FAnsiString& Pattern : FilterPatterns)
 	{
@@ -91,7 +95,11 @@ void FTrimAnalyzer::FNameFilter::Init(FAnsiString Filter)
 
 		if (Names.Num() == 2)
 		{
-			Patterns.Add({ Names[0], Names[1] });
+			OutNameFilter.Patterns.Add({ Names[0], Names[1] });
+		}
+		else
+		{
+			LogTraceTrimmerWarning("Ignoring malformed event name pattern (%s) in filter (%s)!", ANSI_TO_TCHAR(*Pattern), ANSI_TO_TCHAR(*InFilter));
 		}
 	}
 }
@@ -423,6 +431,61 @@ void FTrimAnalyzer::ProcessKnownEvent(uint32 EventId, const ANSICHAR* LoggerName
 
 ////////////////////////////////////////////////////////////////////////////////////////////////////
 
+bool FTrimAnalyzer::FilterByEventName(const ANSICHAR* LoggerName, const ANSICHAR* EventName)
+{
+	// Include only the specified loggers and events.
+	if (!IncludeFilter.IsEmpty())
+	{
+		bool bIsIncluded = false;
+
+		FAnsiString LoggerNameStr(LoggerName);
+		FAnsiString EventNameStr(EventName);
+
+		for (const FNameFilterPattern& FilterPattern : IncludeFilter.Patterns)
+		{
+			if (LoggerNameStr.MatchesWildcard(FilterPattern.LoggerName) &&
+				EventNameStr.MatchesWildcard(FilterPattern.EventName))
+			{
+				bIsIncluded = true;
+				break;
+			}
+		}
+
+		if (!bIsIncluded)
+		{
+			return false;
+		}
+	}
+
+	// Exclude the specified loggers and events.
+	if (!ExcludeFilter.IsEmpty())
+	{
+		bool bIsExcluded = false;
+
+		FAnsiString LoggerNameStr(LoggerName);
+		FAnsiString EventNameStr(EventName);
+
+		for (const FNameFilterPattern& FilterPattern : ExcludeFilter.Patterns)
+		{
+			if (LoggerNameStr.MatchesWildcard(FilterPattern.LoggerName) &&
+				EventNameStr.MatchesWildcard(FilterPattern.EventName))
+			{
+				bIsExcluded = true;
+				break;
+			}
+		}
+
+		if (bIsExcluded)
+		{
+			return false;
+		}
+	}
+
+	return true;
+}
+
+////////////////////////////////////////////////////////////////////////////////////////////////////
+
 bool FTrimAnalyzer::OnNewEvent(uint16 RouteId, const FEventTypeInfo& TypeInfo)
 {
 	const uint32 EventId = TypeInfo.GetId();
@@ -433,14 +496,52 @@ bool FTrimAnalyzer::OnNewEvent(uint16 RouteId, const FEventTypeInfo& TypeInfo)
 	const ANSICHAR* EventName = TypeInfo.GetName();
 	check(EventName != nullptr);
 
+	//////////////////////////////////////////////////
+
+	FEventInfo* EventInfo = EventInfos.Find(EventId);
+
+	// Events are not expected to be declared multiple times.
+	if (EventInfo)
+	{
+		if (EventInfo->TypeInfo->GetId() != EventId)
+		{
+			LogTraceTrimmerWarning("Event (%s.%s) with id %u is already declared with id %u!",
+				ANSI_TO_TCHAR(LoggerName), ANSI_TO_TCHAR(EventName), EventId, EventInfo->TypeInfo->GetId());
+
+			EventInfo = nullptr;
+		}
+		else
+		{
+			LogTraceTrimmerWarning("Event (%s.%s) is already declared!", ANSI_TO_TCHAR(LoggerName), ANSI_TO_TCHAR(EventName));
+			return true;
+		}
+	}
+
+	// Register the new event.
+	check(!EventInfo);
+	EventInfo = &EventInfos.Add(EventId, { &TypeInfo, InvalidWriterEventId, false });
+
+	const UE::Trace::FTraceWriterEventInfo* WriterEventInfo = TraceWriter.FindEvent(LoggerName, EventName);
+	if (WriterEventInfo)
+	{
+		// The event was already declared by trace writer.
+		EventInfo->WriterEventId = WriterEventInfo->GetId();
+	}
+
+	// Pre-filter event by name.
+	check(EventInfo);
+	EventInfo->bIsFilteredOut = !FilterByEventName(LoggerName, EventName);
+
+	//////////////////////////////////////////////////
+
 	ProcessKnownEvent(EventId, LoggerName, EventName);
 
+	//////////////////////////////////////////////////
+	// Known $Trace.* events
+
 	static constexpr FAnsiStringView TraceLoggerName = "$Trace";
 	static constexpr FAnsiStringView TraceNewTraceEventName = "NewTrace";
-	static constexpr FAnsiStringView TraceThreadInfoName = "ThreadInfo";
-
-	FEventInfo* EventInfo = EventInfos.Find(EventId);
-	const UE::Trace::FTraceWriterEventInfo* WriterEventInfo = TraceWriter.FindEvent(LoggerName, EventName);
+	static constexpr FAnsiStringView TraceThreadInfoEventName = "ThreadInfo";
 
 	// Is it a $Trace.* event?
 	if (TraceLoggerName.Equals(LoggerName))
@@ -450,7 +551,7 @@ bool FTrimAnalyzer::OnNewEvent(uint16 RouteId, const FEventTypeInfo& TypeInfo)
 			// $Trace.NewTrace
 			TraceNewTraceEventId = EventId;
 		}
-		else if (TraceThreadInfoName.Equals(EventName))
+		else if (TraceThreadInfoEventName.Equals(EventName))
 		{
 			// $Trace.ThreadInfo
 			TraceThreadInfoEventId = EventId;
@@ -462,30 +563,22 @@ bool FTrimAnalyzer::OnNewEvent(uint16 RouteId, const FEventTypeInfo& TypeInfo)
 			LogTraceTrimmerWarning("Unknown default event (%s.%s)!", ANSI_TO_TCHAR(LoggerName), ANSI_TO_TCHAR(EventName));
 		}
 	}
-	else
-	{
-		// Normal events are not expected to be declared multiple times.
-		if (EventInfo)
-		{
-			if (EventInfo->TypeInfo->GetId() != EventId)
-			{
-				LogTraceTrimmerWarning("Event (%s.%s) with id %u is already declared with id %u!",
-					ANSI_TO_TCHAR(LoggerName), ANSI_TO_TCHAR(EventName), EventId, EventInfo->TypeInfo->GetId());
 
-				EventInfo = nullptr;
-			}
-			else
-			{
-				LogTraceTrimmerWarning("Event (%s.%s) is already declared!", ANSI_TO_TCHAR(LoggerName), ANSI_TO_TCHAR(EventName));
-			}
-		}
+	//////////////////////////////////////////////////
+	// Filtering
+
+	if (EventInfo->bIsFilteredOut)
+	{
+		// The event is filtered out by name.
+		return true;
 	}
 
-	uint32 WriterEventId = 0;
+	//////////////////////////////////////////////////
+	// Declare event.
 
 	if (WriterEventInfo)
 	{
-		WriterEventId = WriterEventInfo->GetId();
+		LogTraceTrimmerMessage("Event %u (%s.%s) was already declared (writer id %u).", EventId, ANSI_TO_TCHAR(LoggerName), ANSI_TO_TCHAR(EventName), WriterEventInfo->GetId());
 	}
 	else
 	{
@@ -495,14 +588,36 @@ bool FTrimAnalyzer::OnNewEvent(uint16 RouteId, const FEventTypeInfo& TypeInfo)
 		{
 			const FEventFieldInfo& FieldInfo = *(TypeInfo.GetFieldInfo(FieldIndex));
 			UE::Trace::ETraceWriterFieldType WriterFieldType = UE::Trace::FTraceWriter::ConvertFieldType(FieldInfo);
-			Builder.Field(FieldInfo.GetName(), WriterFieldType);
+			if (!EnumHasAnyFlags(WriterFieldType, UE::Trace::ETraceWriterFieldType::ReferenceFlag | UE::Trace::ETraceWriterFieldType::DefinitionIdFlag))
+			{
+				// Normal field
+				Builder.Field(FieldInfo.GetName(), WriterFieldType);
+			}
+			else if (EnumHasAnyFlags(WriterFieldType, UE::Trace::ETraceWriterFieldType::ReferenceFlag))
+			{
+				// Reference field
+				const uint32 RefUid = FieldInfo.GetRefUid();
+				FEventInfo* RefEventInfo = EventInfos.Find(RefUid);
+				uint16 WriterRefUid;
+				if (RefEventInfo && RefEventInfo->WriterEventId != InvalidEventId)
+				{
+					WriterRefUid = uint16(RefEventInfo->WriterEventId);
+				}
+				else
+				{
+					WriterRefUid = 0;
+					LogTraceTrimmerWarning("Invalid RefUid (%u) for event (%s.%s) field \"%s\"!",
+						RefUid, ANSI_TO_TCHAR(LoggerName), ANSI_TO_TCHAR(EventName), ANSI_TO_TCHAR(FieldInfo.GetName()));
+				}
+				Builder.ReferenceField(FieldInfo.GetName(), WriterFieldType, WriterRefUid);
+			}
+			else // if (EnumHasAnyFlags(WriterFieldType, UE::Trace::ETraceWriterFieldType::DefinitionIdFlag))
+			{
+				// DefinitionId field
+				Builder.DefinitionIdField(WriterFieldType);
+			}
 		}
-		WriterEventId = Builder.End();
-	}
-
-	if (!EventInfo)
-	{
-		EventInfos.Add(EventId, { &TypeInfo, WriterEventId });
+		EventInfo->WriterEventId = Builder.End();
 	}
 
 	return true;
@@ -528,8 +643,27 @@ bool FTrimAnalyzer::OnEvent(uint16 RouteId, EStyle Style, const FOnEventContext&
 	const ANSICHAR* EventName = TypeInfo.GetName();
 	check(EventName != nullptr);
 
+	//////////////////////////////////////////////////
+
 	bool bShouldWriteEvent = true;
 
+	FEventInfo* EventInfo = EventInfos.Find(EventId);
+	if (!EventInfo)
+	{
+		if (LoggerName[0] == 0) // empty dispatch, expected on LeaveScope events
+		{
+			EventInfo = &EmptyEventInfo;
+		}
+		else
+		{
+			LogTraceTrimmerWarning("Event %u (%s.%s) is not registered!", EventId, ANSI_TO_TCHAR(LoggerName), ANSI_TO_TCHAR(EventName));
+			const bool bIsFilteredOut = !FilterByEventName(LoggerName, EventName);
+			EventInfo = &EventInfos.Add(EventId, { &TypeInfo, InvalidWriterEventId, bIsFilteredOut });
+			bShouldWriteEvent = false;
+		}
+	}
+	check(EventInfo);
+
 	//////////////////////////////////////////////////
 	// Processing known events
 
@@ -613,6 +747,33 @@ bool FTrimAnalyzer::OnEvent(uint16 RouteId, EStyle Style, const FOnEventContext&
 		}
 	}
 
+	//////////////////////////////////////////////////
+	// Filters
+
+	if (EventInfo->bIsFilteredOut)
+	{
+		// Event is filtered out by name.
+		bShouldWriteEvent = false;
+	}
+
+	if (bShouldWriteEvent &&
+		CurrentThreadId != 1 &&
+		(CurrentThread.CurrentTime < Parameters.StartTime ||
+		 CurrentThread.CurrentTime > Parameters.EndTime))
+	{
+		// Non-important event is filtered out by time range.
+		bShouldWriteEvent = false;
+	}
+
+	if (bShouldWriteEvent)
+	{
+		if (EventInfo->WriterEventId == InvalidWriterEventId)
+		{
+			LogTraceTrimmerWarning("Event %u (%s.%s) is not declared!", EventId, ANSI_TO_TCHAR(LoggerName), ANSI_TO_TCHAR(EventName));
+			bShouldWriteEvent = false;
+		}
+	}
+
 	//////////////////////////////////////////////////
 
 	if (RouteId == 0) // Normal events
@@ -627,33 +788,39 @@ bool FTrimAnalyzer::OnEvent(uint16 RouteId, EStyle Style, const FOnEventContext&
 	{
 		if (Style == EStyle::EnterScope)
 		{
-			SetCurrentThread(CurrentThreadId);
-
-			uint64 Timestamp = Context.EventTime.GetTimestamp();
-			if (Timestamp != 0)
+			if (bShouldWriteEvent)
 			{
-				TraceWriter.WriteStampedEnterScopeEvent(SessionStartCycle + Timestamp);
-			}
-			else
-			{
-				TraceWriter.WriteEnterScopeEvent();
+				SetCurrentThread(CurrentThreadId);
+
+				uint64 Timestamp = Context.EventTime.GetTimestamp();
+				if (Timestamp != 0)
+				{
+					TraceWriter.WriteStampedEnterScopeEvent(SessionStartCycle + Timestamp);
+				}
+				else
+				{
+					TraceWriter.WriteEnterScopeEvent();
+				}
 			}
 		}
 		else if (Style == EStyle::LeaveScope)
 		{
-			SetCurrentThread(CurrentThreadId);
-
-			uint64 Timestamp = Context.EventTime.GetTimestamp();
-			if (Timestamp != 0)
-			{
-				TraceWriter.WriteStampedLeaveScopeEvent(SessionStartCycle + Timestamp);
-			}
-			else
+			if (bShouldWriteEvent)
 			{
-				TraceWriter.WriteLeaveScopeEvent();
-			}
+				SetCurrentThread(CurrentThreadId);
 
-			bShouldWriteEvent = false;
+				uint64 Timestamp = Context.EventTime.GetTimestamp();
+				if (Timestamp != 0)
+				{
+					TraceWriter.WriteStampedLeaveScopeEvent(SessionStartCycle + Timestamp);
+				}
+				else
+				{
+					TraceWriter.WriteLeaveScopeEvent();
+				}
+
+				bShouldWriteEvent = false;
+			}
 		}
 		else // unknown Style
 		{
@@ -668,21 +835,41 @@ bool FTrimAnalyzer::OnEvent(uint16 RouteId, EStyle Style, const FOnEventContext&
 	}
 
 	//////////////////////////////////////////////////
+	// Write event
 
-	FEventInfo* EventInfo = EventInfos.Find(EventId);
-	if (EventInfo == nullptr)
+	if (bShouldWriteEvent)
 	{
-		if (LoggerName[0] == 0) // empty dispatch, expected on LeaveScope events
-		{
-			EventInfo = &EmptyEventInfo;
-			bShouldWriteEvent = false;
-		}
-		else
+		SetCurrentThread(CurrentThreadId);
+
+		check(EventInfo);
+		check(EventInfo->WriterEventId != InvalidWriterEventId);
+		UE::Trace::FTraceWriterEventBuilder& EventBuilder = TraceWriter.WriteEvent(EventInfo->WriterEventId);
+
+		const uint32 FieldCount = TypeInfo.GetFieldCount();
+		for (uint32 FieldIndex = 0; FieldIndex < FieldCount; ++FieldIndex)
 		{
-			LogTraceTrimmerWarning("Event %u (%s.%s) is not registered!", EventId, ANSI_TO_TCHAR(LoggerName), ANSI_TO_TCHAR(EventName));
-			EventInfo = &EventInfos.Add(EventId, { &TypeInfo, 0 });
-			bShouldWriteEvent = false;
+			const FEventFieldInfo& FieldInfo = *(TypeInfo.GetFieldInfo(FieldIndex));
+			if (FieldInfo.IsArray())
+			{
+				WriteArrayField(FieldIndex, FieldInfo, EventData, EventBuilder);
+			}
+			else
+			{
+				WriteScalarField(FieldIndex, FieldInfo, EventData, EventBuilder);
+
+				// Detect timestamp fields...
+				if (FieldInfo.GetType() == FEventFieldInfo::EType::Integer &&
+					!FieldInfo.IsSigned() &&
+					FieldInfo.GetTypeSize() == 8 &&
+					FCStringAnsi::Strcmp(FieldInfo.GetName(), "Cycle") == 0)
+				{
+					const uint64 Cycle = EventData.GetValue<uint64>(FieldInfo.GetName());
+					UpdateThreadTime(CurrentThread, Context, Cycle);
+				}
+			}
 		}
+
+		EventBuilder.End();
 	}
 
 	//////////////////////////////////////////////////
@@ -720,106 +907,6 @@ bool FTrimAnalyzer::OnEvent(uint16 RouteId, EStyle Style, const FOnEventContext&
 		}
 	}
 
-	//////////////////////////////////////////////////
-	// Filters
-
-	// Only filter the non-important events.
-	if (bShouldWriteEvent && CurrentThreadId != 1)
-	{
-		// Filter by time range.
-		if (CurrentThread.CurrentTime < Parameters.StartTime ||
-			CurrentThread.CurrentTime > Parameters.EndTime)
-		{
-			bShouldWriteEvent = false;
-		}
-
-		// Include only the specified loggers and events.
-		if (bShouldWriteEvent && !IncludeFilter.IsEmpty())
-		{
-			bool bIsIncluded = false;
-
-			FAnsiString LoggerNameStr(LoggerName);
-			FAnsiString EventNameStr(EventName);
-
-			for (const FNameFilterPattern& FilterPattern : IncludeFilter.Patterns)
-			{
-				if (LoggerNameStr.MatchesWildcard(FilterPattern.LoggerName) &&
-					EventNameStr.MatchesWildcard(FilterPattern.EventName))
-				{
-					bIsIncluded = true;
-					break;
-				}
-			}
-
-			if (!bIsIncluded)
-			{
-				bShouldWriteEvent = false;
-			}
-		}
-
-		// Exclude the specified loggers and events.
-		if (bShouldWriteEvent && !ExcludeFilter.IsEmpty())
-		{
-			bool bIsExcluded = false;
-
-			FAnsiString LoggerNameStr(LoggerName);
-			FAnsiString EventNameStr(EventName);
-
-			for (const FNameFilterPattern& FilterPattern : ExcludeFilter.Patterns)
-			{
-				if (LoggerNameStr.MatchesWildcard(FilterPattern.LoggerName) &&
-					EventNameStr.MatchesWildcard(FilterPattern.EventName))
-				{
-					bIsExcluded = true;
-					break;
-				}
-			}
-
-			if (bIsExcluded)
-			{
-				bShouldWriteEvent = false;
-			}
-		}
-	}
-
-	//////////////////////////////////////////////////
-	// Write event
-
-	if (bShouldWriteEvent)
-	{
-		SetCurrentThread(CurrentThreadId);
-
-		UE::Trace::FTraceWriterEventBuilder& EventBuilder = TraceWriter.WriteEvent(EventInfo->WriterEventId);
-
-		const uint32 FieldCount = (Style == EStyle::LeaveScope) ? 0 : TypeInfo.GetFieldCount();
-		for (uint32 FieldIndex = 0; FieldIndex < FieldCount; ++FieldIndex)
-		{
-			const FEventFieldInfo& FieldInfo = *(TypeInfo.GetFieldInfo(FieldIndex));
-			if (FieldInfo.IsArray())
-			{
-				WriteArrayField(FieldIndex, FieldInfo, EventData, EventBuilder);
-			}
-			else
-			{
-				WriteScalarField(FieldIndex, FieldInfo, EventData, EventBuilder);
-
-				// Detect timestamp fields...
-				if (FieldInfo.GetType() == FEventFieldInfo::EType::Integer &&
-					!FieldInfo.IsSigned() &&
-					FieldInfo.GetTypeSize() == 8 &&
-					FCStringAnsi::Strcmp(FieldInfo.GetName(), "Cycle") == 0)
-				{
-					const uint64 Cycle = EventData.GetValue<uint64>(FieldInfo.GetName());
-					UpdateThreadTime(CurrentThread, Context, Cycle);
-				}
-			}
-		}
-
-		EventBuilder.End();
-	}
-
-	//////////////////////////////////////////////////
-
 	return !bShouldStopAnalysis;
 }
 
@@ -866,6 +953,7 @@ void FTrimAnalyzer::WriteArrayField(
 					default:
 					{
 						// error
+						LogTraceTrimmerWarning("Invalid array field type: unsupported signed integer size (%u)!", uint32(FieldInfo.GetTypeSize()));
 					}
 				} // switch (FieldInfo.GetTypeSize())
 			}
@@ -900,6 +988,7 @@ void FTrimAnalyzer::WriteArrayField(
 					default:
 					{
 						// error
+						LogTraceTrimmerWarning("Invalid array field type: unsupported unsigned integer size (%u)!", uint32(FieldInfo.GetTypeSize()));
 					}
 				} // switch (FieldInfo.GetTypeSize())
 			}
@@ -925,6 +1014,7 @@ void FTrimAnalyzer::WriteArrayField(
 				default:
 				{
 					// error
+					LogTraceTrimmerWarning("Invalid array field type: unsupported float size (%u)!", uint32(FieldInfo.GetTypeSize()));
 				}
 			} // switch (FieldInfo.GetTypeSize())
 			break;
@@ -949,6 +1039,7 @@ void FTrimAnalyzer::WriteArrayField(
 		default:
 		{
 			// error
+			LogTraceTrimmerWarning("Invalid array field type (%u)!", uint32(FieldInfo.GetType()));
 		}
 	}
 }
@@ -996,6 +1087,7 @@ void FTrimAnalyzer::WriteScalarField(
 					default:
 					{
 						// error
+						LogTraceTrimmerWarning("Invalid scalar field type: unsupported signed integer size (%u)!", uint32(FieldInfo.GetTypeSize()));
 					}
 				} // switch (FieldInfo.GetTypeSize())
 			}
@@ -1030,6 +1122,7 @@ void FTrimAnalyzer::WriteScalarField(
 					default:
 					{
 						// error
+						LogTraceTrimmerWarning("Invalid scalar field type: unsupported unsigned integer size (%u)!", uint32(FieldInfo.GetTypeSize()));
 					}
 				} // switch (FieldInfo.GetTypeSize())
 			}
@@ -1055,6 +1148,7 @@ void FTrimAnalyzer::WriteScalarField(
 				default:
 				{
 					// error
+					LogTraceTrimmerWarning("Invalid scalar field type: unsupported float size (%u)!", uint32(FieldInfo.GetTypeSize()));
 				}
 			} // switch (FieldInfo.GetTypeSize())
 			break;
@@ -1064,7 +1158,7 @@ void FTrimAnalyzer::WriteScalarField(
 		{
 			check(FieldInfo.GetSize() == 1);
 			UE::Trace::FEventRef8 RefValue = EventData.GetReferenceValue<uint8>(FieldIndex);
-			//TODO: EventBuilder.Field(FieldIndex, RefValue);
+			EventBuilder.Field(FieldIndex, RefValue);
 			break;
 		}
 
@@ -1072,7 +1166,7 @@ void FTrimAnalyzer::WriteScalarField(
 		{
 			check(FieldInfo.GetSize() == 2);
 			UE::Trace::FEventRef16 RefValue = EventData.GetReferenceValue<uint16>(FieldIndex);
-			//TODO: EventBuilder.Field(FieldIndex, RefValue);
+			EventBuilder.Field(FieldIndex, RefValue);
 			break;
 		}
 
@@ -1080,7 +1174,7 @@ void FTrimAnalyzer::WriteScalarField(
 		{
 			check(FieldInfo.GetSize() == 4);
 			UE::Trace::FEventRef32 RefValue = EventData.GetReferenceValue<uint32>(FieldIndex);
-			//TODO: EventBuilder.Field(FieldIndex, RefValue);
+			EventBuilder.Field(FieldIndex, RefValue);
 			break;
 		}
 
@@ -1088,13 +1182,14 @@ void FTrimAnalyzer::WriteScalarField(
 		{
 			check(FieldInfo.GetSize() == 8);
 			UE::Trace::FEventRef64 RefValue = EventData.GetReferenceValue<uint64>(FieldIndex);
-			//TODO: EventBuilder.Field(FieldIndex, RefValue);
+			EventBuilder.Field(FieldIndex, RefValue);
 			break;
 		}
 
 		default:
 		{
 			// error
+			LogTraceTrimmerWarning("Invalid scalar field type (%u)!", uint32(FieldInfo.GetType()));
 		}
 	}
 }
diff --git a/Engine/Source/Developer/TraceServices/Private/Analyzers/TrimAnalyzer.h b/Engine/Source/Developer/TraceServices/Private/Analyzers/TrimAnalyzer.h
--- a/Engine/Source/Developer/TraceServices/Private/Analyzers/TrimAnalyzer.h
+++ b/Engine/Source/Developer/TraceServices/Private/Analyzers/TrimAnalyzer.h
@@ -28,7 +28,8 @@ private:
 	struct FEventInfo
 	{
 		const FEventTypeInfo* TypeInfo = nullptr;
-		uint32 WriterEventId = 0;
+		uint32 WriterEventId = ~0u;
+		bool bIsFilteredOut = false;
 	};
 
 	struct FNameFilterPattern
@@ -39,7 +40,6 @@ private:
 
 	struct FNameFilter
 	{
-		void Init(FAnsiString Filter);
 		bool IsEmpty() const { return Patterns.IsEmpty(); }
 
 		TArray<FNameFilterPattern> Patterns;
@@ -165,6 +165,9 @@ public:
 	uint32 GetNumWarnings() const { return NumWarnings; }
 
 private:
+	void InitNameFilter(FAnsiString InFilter, FNameFilter& OutNameFilter);
+	bool FilterByEventName(const ANSICHAR* LoggerName, const ANSICHAR* EventName);
+
 	virtual void OnAnalysisBegin(const FOnAnalysisContext& Context) override;
 	virtual void OnAnalysisEnd() override;
 	virtual void OnVersion(uint32 InTransportVersion, uint32 InProtocolVersion) /* override */;
@@ -213,9 +216,11 @@ private:
 	FNameFilter IncludeFilter;
 	FNameFilter ExcludeFilter;
 
-	static constexpr uint32 InvalidEventId = (uint32)-1;
+	static constexpr uint32 InvalidEventId = ~0u;
 	static constexpr uint32 MaxKnownEventId = 0xFFFF;
 
+	static constexpr uint32 InvalidWriterEventId = ~0u;
+
 	uint32 TraceNewTraceEventId = InvalidEventId;
 	bool bIsTraceNewTraceEventDetected = false;
 
