[TraceTrimmer]
* Extracted FKnownEventProcessor -- pulls the EKnownEvent dispatch table, KnownEventMap, and CPU Profiler event batch processing out of FTrimAnalyzer into a standalone reusable class.
* Optimized thread sync -- replaced the full-scan TMap with a SortedThreads array so UpdateThreadsOnSyncEvent can break early once threads are past EndTime.
* Added a possible optimization for fast accessing thread infos.
* Fixed the TrimAsync public API -- wraps FAnalysisProcessor behind an ITrimAnalysisProcessor interface.
* Added periodic progress logging (at 10 seconds intervals) -- shows number of read/written events, session time and trim process time.
* Cleaned up the logging macros (wprintf/UE_LOG+ANSI_TO_TCHAR --> printf/UE_LOGF).

#Synced-CL 52393202

diff --git a/Engine/Source/Developer/TraceServices/Private/Analyzers/KnownEventProcessor.cpp b/Engine/Source/Developer/TraceServices/Private/Analyzers/KnownEventProcessor.cpp
new file mode 100644
--- /dev/null
+++ b/Engine/Source/Developer/TraceServices/Private/Analyzers/KnownEventProcessor.cpp
@@ -0,0 +1,507 @@
+// Copyright Epic Games, Inc. All Rights Reserved.
+
+#include "KnownEventProcessor.h"
+
+#include "Misc/CString.h"
+
+// TraceServices
+#include "Common/Utils.h"
+
+////////////////////////////////////////////////////////////////////////////////////////////////////
+
+namespace TraceServices
+{
+
+////////////////////////////////////////////////////////////////////////////////////////////////////
+// FKnownEventProcessor
+////////////////////////////////////////////////////////////////////////////////////////////////////
+
+FKnownEventProcessor::FKnownEventProcessor()
+{
+	KnownEvents.AddDefaulted((int)EKnownEvent::Count);
+	KnownEventMap.AddDefaulted(MaxKnownEventId + 1);
+
+	Reset();
+}
+
+////////////////////////////////////////////////////////////////////////////////////////////////////
+
+FKnownEventProcessor::~FKnownEventProcessor()
+{
+}
+
+////////////////////////////////////////////////////////////////////////////////////////////////////
+
+void FKnownEventProcessor::Reset()
+{
+	KnownEvents[(int)EKnownEvent::Invalid].Id = InvalidEventId;
+	KnownEvents[(int)EKnownEvent::Trace_NewEvent].Id = 0;
+
+	for (uint32 EventId = 0; EventId <= MaxKnownEventId; ++EventId)
+	{
+		KnownEventMap[EventId] = EKnownEvent::Invalid;
+	}
+	KnownEventMap[0] = EKnownEvent::Trace_NewEvent;
+}
+
+////////////////////////////////////////////////////////////////////////////////////////////////////
+
+void FKnownEventProcessor::RegisterCallbacks()
+{
+	using FOnEventContext = UE::Trace::IAnalyzer::FOnEventContext;
+	using FEventData = UE::Trace::IAnalyzer::FEventData;
+
+	//////////////////////////////////////////////////
+	// Misc.*
+
+	KnownEvents[(int)EKnownEvent::Misc_RegionBeginWithId].Callback =
+		[this](const FOnEventContext& Context)
+		{
+			const uint64 Cycle = Context.EventData.GetValue<uint64>("CycleAndId");
+			UpdateThreadTime(Context, Cycle);
+		};
+
+	//////////////////////////////////////////////////
+	// CpuProfiler.*
+
+	KnownEvents[(int)EKnownEvent::CpuProfiler_EventBatchV3].Callback =
+		[this](const FOnEventContext& Context)
+		{
+			const FEventData& EventData = Context.EventData;
+			TArrayView<const uint8> Data = EventData.GetArrayView<uint8>("Data");
+			if (Data.GetData() != nullptr)
+			{
+				ProcessCpuEventBatchV2V3(Context, Data.GetData(), Data.Num(), 3);
+			}
+		};
+
+	KnownEvents[(int)EKnownEvent::CpuProfiler_EventBatchV2].Callback =
+		[this](const FOnEventContext& Context)
+		{
+			const FEventData& EventData = Context.EventData;
+			TArrayView<const uint8> Data = EventData.GetArrayView<uint8>("Data");
+			if (Data.GetData() != nullptr)
+			{
+				ProcessCpuEventBatchV2V3(Context, Data.GetData(), Data.Num(), 2);
+			}
+		};
+
+	KnownEvents[(int)EKnownEvent::CpuProfiler_EventBatch].Callback =
+	KnownEvents[(int)EKnownEvent::CpuProfiler_EndCapture].Callback =
+		[this](const FOnEventContext& Context)
+		{
+			const FEventData& EventData = Context.EventData;
+			TArrayView<const uint8> Data = EventData.GetArrayView<uint8>("Data");
+			if (Data.GetData() != nullptr)
+			{
+				ProcessCpuEventBatchV1(Context, Data.GetData(), Data.Num());
+			}
+		};
+
+	//////////////////////////////////////////////////
+	// GpuProfiler.*
+
+	KnownEvents[(int)EKnownEvent::GpuProfiler_EventBeginBreadcrumb].Callback =
+	KnownEvents[(int)EKnownEvent::GpuProfiler_EventBeginWork].Callback =
+		[this](const FOnEventContext& Context)
+		{
+			const uint64 Cycle = Context.EventData.GetValue<uint64>("GPUTimestampTOP");
+			UpdateThreadTime(Context, Cycle);
+		};
+	KnownEvents[(int)EKnownEvent::GpuProfiler_EventEndBreadcrumb].Callback =
+	KnownEvents[(int)EKnownEvent::GpuProfiler_EventEndWork].Callback =
+		[this](const FOnEventContext& Context)
+		{
+			const uint64 Cycle = Context.EventData.GetValue<uint64>("GPUTimestampBOP");
+			UpdateThreadTime(Context, Cycle);
+		};
+	KnownEvents[(int)EKnownEvent::GpuProfiler_EventWait].Callback =
+		[this](const FOnEventContext& Context)
+		{
+			const uint64 Cycle = Context.EventData.GetValue<uint64>("EndTime");
+			UpdateThreadTime(Context, Cycle);
+		};
+	KnownEvents[(int)EKnownEvent::GpuProfiler_SignalFence].Callback =
+	KnownEvents[(int)EKnownEvent::GpuProfiler_WaitFence].Callback =
+		[this](const FOnEventContext& Context)
+		{
+			const uint64 Cycle = Context.EventData.GetValue<uint64>("CPUTimestamp");
+			UpdateThreadTime(Context, Cycle);
+		};
+
+	//////////////////////////////////////////////////
+	// Known events with a "Cycle" field
+
+	KnownEvents[(int)EKnownEvent::Misc_Bookmark].Callback =
+	KnownEvents[(int)EKnownEvent::Misc_ScreenshotHeader].Callback =
+	KnownEvents[(int)EKnownEvent::Misc_RegionBegin].Callback =
+	KnownEvents[(int)EKnownEvent::Misc_RegionEnd].Callback =
+	KnownEvents[(int)EKnownEvent::Misc_RegionEndWithId].Callback =
+	KnownEvents[(int)EKnownEvent::Misc_BeginFrame].Callback =
+	KnownEvents[(int)EKnownEvent::Misc_EndFrame].Callback =
+	KnownEvents[(int)EKnownEvent::Logging_LogMessage].Callback =
+	KnownEvents[(int)EKnownEvent::PlatformFile_BeginOpen].Callback =
+	KnownEvents[(int)EKnownEvent::PlatformFile_EndOpen].Callback =
+	KnownEvents[(int)EKnownEvent::PlatformFile_BeginRead].Callback =
+	KnownEvents[(int)EKnownEvent::PlatformFile_EndRead].Callback =
+	KnownEvents[(int)EKnownEvent::PlatformFile_BeginClose].Callback =
+	KnownEvents[(int)EKnownEvent::PlatformFile_EndClose].Callback =
+	KnownEvents[(int)EKnownEvent::CpuProfiler_EndThread].Callback =
+	KnownEvents[(int)EKnownEvent::Memory_Marker].Callback =
+	KnownEvents[(int)EKnownEvent::LLM_TagValue].Callback =
+	KnownEvents[(int)EKnownEvent::SlateTrace_AddWidget].Callback =
+		[this](const FOnEventContext& Context)
+		{
+			const uint64 Cycle = Context.EventData.GetValue<uint64>("Cycle");
+			UpdateThreadTime(Context, Cycle);
+		};
+}
+
+////////////////////////////////////////////////////////////////////////////////////////////////////
+
+void FKnownEventProcessor::OnAnalysisBegin(const UE::Trace::IAnalyzer::FOnAnalysisContext& Context)
+{
+	Reset();
+	RegisterCallbacks();
+}
+
+///////////////////////////////////////////////////////////////////////////////////////////////////
+
+void FKnownEventProcessor::OnAnalysisEnd()
+{
+}
+
+///////////////////////////////////////////////////////////////////////////////////////////////////
+
+void FKnownEventProcessor::OnNewEvent(uint32 EventId, const ANSICHAR* LoggerName, const ANSICHAR* EventName)
+{
+	if (!ensure(EventId <= MaxKnownEventId))
+	{
+		return;
+	}
+
+	#define BEGIN_KNOWN_LOGGER(_Name) \
+		if (FCStringAnsi::Strcmp(LoggerName, _Name) == 0) \
+		{
+
+	#define KNOWN_EVENT(_Name,_Id) \
+			if (FCStringAnsi::Strcmp(EventName, _Name) == 0) \
+			{ \
+				KnownEvents[(int)EKnownEvent::_Id].Id = EventId; \
+				KnownEventMap[EventId] = EKnownEvent::_Id; \
+				return; \
+			}
+
+	#define END_KNOWN_LOGGER() \
+			return; \
+		}
+
+	BEGIN_KNOWN_LOGGER("$Trace")
+		KNOWN_EVENT("NewEvent", Trace_NewEvent)
+		//KNOWN_EVENT("NewTrace", Trace_NewTrace)
+		//KNOWN_EVENT("ThreadTiming", Trace_ThreadTiming)
+		//KNOWN_EVENT("ThreadInfo", Trace_ThreadInfo)
+		//KNOWN_EVENT("ThreadGroupBegin", Trace_ThreadGroupBegin)
+		//KNOWN_EVENT("ThreadGroupEnd", Trace_ThreadGroupEnd)
+		//KNOWN_EVENT("ChannelAnnounce", Trace_ChannelAnnounce_Old) // deprecated
+	END_KNOWN_LOGGER()
+
+	//BEGIN_KNOWN_LOGGER("Trace")
+		//KNOWN_EVENT("ChannelAnnounce", Trace_ChannelAnnounce)
+		//KNOWN_EVENT("ChannelAnnounce", Trace_ChannelToggle)
+	//END_KNOWN_LOGGER()
+
+	//BEGIN_KNOWN_LOGGER("Diagnostics")
+		//KNOWN_EVENT("Session", Diagnostics_Session) // deprecated
+		//KNOWN_EVENT("Session2", Diagnostics_Session2)
+	//END_KNOWN_LOGGER()
+
+	BEGIN_KNOWN_LOGGER("Misc")
+		//KNOWN_EVENT("CreateThread", Misc_CreateThread)
+		//KNOWN_EVENT("SetThreadGroup", Misc_SetThreadGroup)
+		//KNOWN_EVENT("BeginThreadGroupScope", Misc_BeginThreadGroupScope)
+		//KNOWN_EVENT("BookmarkSpec", Misc_BookmarkSpec)
+		KNOWN_EVENT("Bookmark", Misc_Bookmark)
+		KNOWN_EVENT("ScreenshotHeader", Misc_ScreenshotHeader)
+		KNOWN_EVENT("RegionBegin", Misc_RegionBegin)
+		KNOWN_EVENT("RegionBeginWithId", Misc_RegionBeginWithId)
+		KNOWN_EVENT("RegionEnd", Misc_RegionEnd)
+		KNOWN_EVENT("RegionEndWithId", Misc_RegionEndWithId)
+		//KNOWN_EVENT("BeginGameFrame", Misc_BeginGameFrame)
+		//KNOWN_EVENT("EndGameFrame", Misc_EndGameFrame)
+		//KNOWN_EVENT("BeginRenderFrame", Misc_BeginRenderFrame)
+		//KNOWN_EVENT("EndRenderFrame", Misc_EndRenderFrame)
+		KNOWN_EVENT("BeginFrame", Misc_BeginFrame)
+		KNOWN_EVENT("EndFrame", Misc_EndFrame)
+	END_KNOWN_LOGGER()
+
+	BEGIN_KNOWN_LOGGER("Logging")
+		//KNOWN_EVENT("LogCategory", Logging_LogCategory)
+		//KNOWN_EVENT("LogMessageSpec", Logging_LogMessageSpec)
+		KNOWN_EVENT("LogMessage", Logging_LogMessage)
+	END_KNOWN_LOGGER()
+
+	BEGIN_KNOWN_LOGGER("PlatformFile")
+		KNOWN_EVENT("BeginOpen", PlatformFile_BeginOpen)
+		KNOWN_EVENT("EndOpen", PlatformFile_EndOpen)
+		KNOWN_EVENT("BeginRead", PlatformFile_BeginRead)
+		KNOWN_EVENT("EndRead", PlatformFile_EndRead)
+		KNOWN_EVENT("BeginClose", PlatformFile_BeginClose)
+		KNOWN_EVENT("EndClose", PlatformFile_EndClose)
+	END_KNOWN_LOGGER()
+
+	//BEGIN_KNOWN_LOGGER("LoadTime")
+		//KNOWN_EVENT("ClassInfo", LoadTime_ClassInfo)
+		//KNOWN_EVENT("NewAsyncPackage", LoadTime_NewAsyncPackage)
+	//END_KNOWN_LOGGER()
+
+	BEGIN_KNOWN_LOGGER("GpuProfiler")
+		//KNOWN_EVENT("EventSpec", GpuProfiler_EventSpec) // deprecated
+		//KNOWN_EVENT("Frame", GpuProfiler_Frame) // deprecated
+		//KNOWN_EVENT("Frame2", GpuProfiler_Frame2) // deprecated
+		//KNOWN_EVENT("Init", GpuProfiler_Init)
+		//KNOWN_EVENT("QueueSpec", GpuProfiler_QueueSpec)
+		//KNOWN_EVENT("EventFrameBoundary", GpuProfiler_EventFrameBoundary)
+		//KNOWN_EVENT("EventBreadcrumbSpec", GpuProfiler_EventBreadcrumbSpec)
+		KNOWN_EVENT("EventBeginBreadcrumb", GpuProfiler_EventBeginBreadcrumb)
+		KNOWN_EVENT("EventEndBreadcrumb", GpuProfiler_EventEndBreadcrumb)
+		KNOWN_EVENT("EventBeginWork", GpuProfiler_EventBeginWork)
+		KNOWN_EVENT("EventEndWork", GpuProfiler_EventEndWork)
+		KNOWN_EVENT("EventWait", GpuProfiler_EventWait)
+		//KNOWN_EVENT("EventStats", GpuProfiler_EventStats)
+		KNOWN_EVENT("SignalFence", GpuProfiler_SignalFence)
+		KNOWN_EVENT("WaitFence", GpuProfiler_WaitFence)
+	END_KNOWN_LOGGER()
+
+	BEGIN_KNOWN_LOGGER("CpuProfiler")
+		//KNOWN_EVENT("EventSpec", CpuProfiler_EventSpec)
+		//KNOWN_EVENT("MetadataSpec", CpuProfiler_MetadataSpec)
+		KNOWN_EVENT("Metadata", CpuProfiler_Metadata)
+		KNOWN_EVENT("EventBatchV3", CpuProfiler_EventBatchV3)
+		KNOWN_EVENT("EventBatchV2", CpuProfiler_EventBatchV2) // deprecated
+		KNOWN_EVENT("EventBatch", CpuProfiler_EventBatch) // deprecated
+		KNOWN_EVENT("EndCapture", CpuProfiler_EndCapture) // deprecated
+		KNOWN_EVENT("EndThread", CpuProfiler_EndThread)
+	END_KNOWN_LOGGER()
+
+	//BEGIN_KNOWN_LOGGER("Stats")
+		//KNOWN_EVENT("Spec", Stats_Spec)
+	//END_KNOWN_LOGGER()
+
+	//BEGIN_KNOWN_LOGGER("Counters")
+		//KNOWN_EVENT("Spec", Counters_Spec)
+	//END_KNOWN_LOGGER()
+
+	//BEGIN_KNOWN_LOGGER("CsvProfiler")
+		//KNOWN_EVENT("RegisterCategory", CsvProfiler_RegisterCategory)
+		//KNOWN_EVENT("DefineInlineStat", CsvProfiler_DefineInlineStat)
+		//KNOWN_EVENT("DefineDeclaredStat", CsvProfiler_DefineDeclaredStat)
+		//KNOWN_EVENT("Metadata", CsvProfiler_Metadata)
+		//KNOWN_EVENT("BeginCapture", CsvProfiler_BeginCapture)
+	//END_KNOWN_LOGGER()
+
+	BEGIN_KNOWN_LOGGER("Memory")
+		KNOWN_EVENT("Marker", Memory_Marker)
+	END_KNOWN_LOGGER()
+
+	BEGIN_KNOWN_LOGGER("LLM")
+		KNOWN_EVENT("TagValue", LLM_TagValue)
+		//KNOWN_EVENT("TagsSpec", LLM_TagsSpec)
+		//KNOWN_EVENT("TrackerSpec", LLM_TrackerSpec)
+	END_KNOWN_LOGGER()
+
+	BEGIN_KNOWN_LOGGER("SlateTrace")
+		KNOWN_EVENT("AddWidget", SlateTrace_AddWidget)
+	END_KNOWN_LOGGER()
+
+	#undef BEGIN_KNOWN_LOGGER
+	#undef KNOWN_EVENT
+	#undef END_KNOWN_LOGGER
+}
+
+////////////////////////////////////////////////////////////////////////////////////////////////////
+
+void FKnownEventProcessor::OnEvent(uint32 EventId, const UE::Trace::IAnalyzer::FOnEventContext& Context)
+{
+	if (ensure(EventId <= MaxKnownEventId))
+	{
+		EKnownEvent KnownEvent = KnownEventMap[EventId];
+		check((int)KnownEvent < (int)EKnownEvent::Count);
+		const FKnownEvent& Event = KnownEvents[(int)KnownEvent];
+
+		if (Event.Callback)
+		{
+			Event.Callback(Context);
+		}
+	}
+}
+
+////////////////////////////////////////////////////////////////////////////////////////////////////
+
+void FKnownEventProcessor::ProcessCpuEventBatchV1(const UE::Trace::IAnalyzer::FOnEventContext& Context, const uint8* BufferPtr, uint32 BufferSize)
+{
+	// Trace analysis was changed to be able to provide a suitable id. Prior to
+	// this users of Trace would send along their own thread ids. For backwards
+	// compatibility we'll bias field thread ids to avoid collision with Trace's.
+	constexpr uint32 Bias = 0x70000000;
+	uint32 ThreadId = Context.EventData.GetValue<uint32>("ThreadId", 0);
+	ThreadId |= ThreadId ? Bias : Context.ThreadInfo.GetId();
+
+	FThreadInfo& Thread = GetOrAddThread(ThreadId);
+
+	Thread.Cycle = ProcessCpuEventBatchV1(Context, BufferPtr, BufferSize, Thread.Cycle);
+
+	OnUpdateThreadTime(ThreadId, Context, Thread.Cycle);
+
+	if (ThreadId != Context.ThreadInfo.GetId()) // ThreadId has bias (it's using the "ThreadId" field)
+	{
+		OnUpdateThreadTime(Context.ThreadInfo.GetId(), Context, Thread.Cycle);
+	}
+}
+
+////////////////////////////////////////////////////////////////////////////////////////////////////
+
+uint64 FKnownEventProcessor::ProcessCpuEventBatchV1(const UE::Trace::IAnalyzer::FOnEventContext& Context, const uint8* BufferPtr, uint32 BufferSize, uint64 ThreadCycle)
+{
+	const uint8* BufferEnd = BufferPtr + BufferSize;
+	while (BufferPtr < BufferEnd)
+	{
+		uint64 DecodedCycle = FTraceAnalyzerUtils::Decode7bit(BufferPtr);
+		ThreadCycle += (DecodedCycle >> 1);
+
+		if (DecodedCycle & 1ull)
+		{
+			uint32 SpecId = (uint32)FTraceAnalyzerUtils::Decode7bit(BufferPtr);
+		}
+	}
+	if (!ensure(BufferPtr == BufferEnd))
+	{
+		//error
+	}
+
+	return ThreadCycle;
+}
+
+////////////////////////////////////////////////////////////////////////////////////////////////////
+
+void FKnownEventProcessor::ProcessCpuEventBatchV2V3(const UE::Trace::IAnalyzer::FOnEventContext& Context, const uint8* BufferPtr, uint32 BufferSize, int32 Version)
+{
+	const uint32 ThreadId = Context.ThreadInfo.GetId();
+	FThreadInfo& Thread = GetOrAddThread(ThreadId);
+
+	Thread.Cycle = ProcessCpuEventBatchV2V3(Context, BufferPtr, BufferSize, Version, Thread.Cycle);
+
+	OnUpdateThreadTime(ThreadId, Context, Thread.Cycle);
+}
+	
+////////////////////////////////////////////////////////////////////////////////////////////////////
+
+uint64 FKnownEventProcessor::ProcessCpuEventBatchV2V3(const UE::Trace::IAnalyzer::FOnEventContext& Context, const uint8* BufferPtr, uint32 BufferSize, int32 Version, uint64 ThreadCycle)
+{
+	check(Context.EventTime.GetTimestamp() == 0);
+	const uint64 BaseCycle = Context.EventTime.AsCycle64();
+
+	const uint8* BufferEnd = BufferPtr + BufferSize;
+	while (BufferPtr < BufferEnd)
+	{
+		uint64 DecodedCycle = FTraceAnalyzerUtils::Decode7bit(BufferPtr);
+		uint64 ActualCycle = (DecodedCycle >> 2);
+
+		// ActualCycle larger or equal to ThreadCycle means we have a new base value.
+		if (ActualCycle < ThreadCycle)
+		{
+			ActualCycle += ThreadCycle;
+		}
+
+		// If we late connect we will be joining the cycle stream mid-flow and
+		// will have missed out on it's base timestamp. Reconstruct it here.
+		if (ActualCycle < BaseCycle)
+		{
+			ActualCycle += BaseCycle;
+		}
+
+		if (DecodedCycle & 2ull) // coroutine
+		{
+			if (DecodedCycle & 1ull) // Restore
+			{
+				uint64 CoroutineId = FTraceAnalyzerUtils::Decode7bit(BufferPtr);
+				uint32 TimerScopeDepth = (uint32)FTraceAnalyzerUtils::Decode7bit(BufferPtr);
+			}
+			else // Save
+			{
+				uint32 TimerScopeDepth = (uint32)FTraceAnalyzerUtils::Decode7bit(BufferPtr);
+			}
+		}
+		else // normal CPU events
+		{
+			if (DecodedCycle & 1ull)
+			{
+				uint32 SpecId = (uint32)FTraceAnalyzerUtils::Decode7bit(BufferPtr);
+#if 0
+				if (Version == 3)
+				{
+					if (SpecId & 1u) // The last bit is set if this is a metadata id.
+					{
+						uint32 MetadataId = SpecId >> 1;
+					}
+					else
+					{
+						SpecId = SpecId >> 1;
+					}
+				}
+				else // Version == 2
+				{
+					//SpecId
+				}
+#endif
+			}
+		}
+
+		check(ActualCycle > 0);
+		ThreadCycle = ActualCycle;
+	}
+	if (!ensure(BufferPtr == BufferEnd))
+	{
+		// error
+	}
+
+	return ThreadCycle;
+}
+
+////////////////////////////////////////////////////////////////////////////////////////////////////
+
+FKnownEventProcessor::FThreadInfo& FKnownEventProcessor::GetOrAddThread(uint32 ThreadId)
+{
+	FThreadInfo* FoundThread = ThreadMap.Find(ThreadId);
+	if (FoundThread)
+	{
+		return *FoundThread;
+	}
+	return ThreadMap.Add(ThreadId, { ThreadId });
+}
+
+////////////////////////////////////////////////////////////////////////////////////////////////////
+
+void FKnownEventProcessor::UpdateThreadTime(const UE::Trace::IAnalyzer::FOnEventContext& Context, uint64 Cycle)
+{
+	if (Cycle != 0)
+	{
+		const uint32 ThreadId = Context.ThreadInfo.GetId();
+		OnUpdateThreadTime(ThreadId, Context, Cycle);
+	}
+}
+
+////////////////////////////////////////////////////////////////////////////////////////////////////
+
+void FKnownEventProcessor::OnUpdateThreadTime(uint32 ThreadId, const UE::Trace::IAnalyzer::FOnEventContext& Context, uint64 Cycle)
+{
+	if (OnUpdateThreadTimeCallback)
+	{
+		OnUpdateThreadTimeCallback(ThreadId, Context, Cycle);
+	}
+}
+
+////////////////////////////////////////////////////////////////////////////////////////////////////
+////////////////////////////////////////////////////////////////////////////////////////////////////
+
+} // namespace TraceServices
diff --git a/Engine/Source/Developer/TraceServices/Private/Analyzers/KnownEventProcessor.h b/Engine/Source/Developer/TraceServices/Private/Analyzers/KnownEventProcessor.h
new file mode 100644
--- /dev/null
+++ b/Engine/Source/Developer/TraceServices/Private/Analyzers/KnownEventProcessor.h
@@ -0,0 +1,181 @@
+// Copyright Epic Games, Inc. All Rights Reserved.
+
+#pragma once
+
+#include "CoreTypes.h"
+
+#include "Containers/Array.h"
+#include "Containers/Map.h"
+#include "Templates/Function.h"
+
+// TraceAnalysis
+#include "Trace/Analyzer.h"
+
+namespace TraceServices
+{
+
+class FKnownEventProcessor
+{
+public:
+	typedef TFunction<void(uint32 ThreadId, const UE::Trace::IAnalyzer::FOnEventContext& Context, uint64 Cycle)> FOnUpdateThreadTimeFunc;
+
+private:
+	enum class EKnownEvent : uint32
+	{
+		Invalid,
+
+		Trace_NewEvent,
+		//Trace_NewTrace,
+		//Trace_ThreadTiming,
+		//Trace_ThreadInfo,
+		//Trace_ThreadGroupBegin,
+		//Trace_ThreadGroupEnd,
+		//Trace_ChannelAnnounce_Old, // "$Trace", deprecated in 4.26
+
+		//Trace_ChannelAnnounce, // "Trace", new in 4.26
+		//Trace_ChannelToggle,
+
+		//Diagnostics_Session,
+		//Diagnostics_Session2,
+
+		//Misc_CreateThread,
+		//Misc_SetThreadGroup,
+		//Misc_BeginThreadGroupScope,
+		//Misc_BookmarkSpec,
+		Misc_Bookmark,
+		Misc_ScreenshotHeader,
+		//Misc_ScreenshotChunk,
+		Misc_RegionBegin,
+		Misc_RegionBeginWithId,
+		Misc_RegionEnd,
+		Misc_RegionEndWithId,
+		//Misc_BeginGameFrame, // deprecated
+		//Misc_EndGameFrame, // deprecated
+		//Misc_BeginRenderFrame, // deprecated
+		//Misc_EndRenderFrame, // deprecated
+		Misc_BeginFrame,
+		Misc_EndFrame,
+
+		//Logging_LogCategory,
+		//Logging_LogMessageSpec,
+		Logging_LogMessage,
+
+		PlatformFile_BeginOpen,
+		PlatformFile_EndOpen,
+		PlatformFile_BeginRead,
+		PlatformFile_EndRead,
+		PlatformFile_BeginClose,
+		PlatformFile_EndClose,
+
+		//LoadTime_ClassInfo,
+		//LoadTime_NewAsyncPackage,
+
+		//GpuProfiler_EventSpec, // old GpuProfiler, deprecated
+		//GpuProfiler_Frame, // old GpuProfiler, deprecated
+		//GpuProfiler_Frame2, // old GpuProfiler, deprecated
+		//GpuProfiler_Init,
+		//GpuProfiler_QueueSpec,
+		//GpuProfiler_EventFrameBoundary,
+		//GpuProfiler_EventBreadcrumbSpec,
+		GpuProfiler_EventBeginBreadcrumb,
+		GpuProfiler_EventEndBreadcrumb,
+		GpuProfiler_EventBeginWork,
+		GpuProfiler_EventEndWork,
+		GpuProfiler_EventWait,
+		//GpuProfiler_EventStats,
+		GpuProfiler_SignalFence,
+		GpuProfiler_WaitFence,
+
+		//CpuProfiler_EventSpec,
+		//CpuProfiler_MetadataSpec, // added in UE 5.x
+		CpuProfiler_Metadata, // added in UE 5.x
+		CpuProfiler_EventBatchV3, // added in UE 5.6
+		CpuProfiler_EventBatchV2, // backward compatibility; added in UE 5.1, removed in 5.6
+		CpuProfiler_EventBatch, // backward compatibility; removed in UE 5.1
+		CpuProfiler_EndCapture, // backward compatibility; removed in UE 5.1
+		CpuProfiler_EndThread, // "Cycle" field added in UE 5.4
+
+		//Stats_Spec,
+
+		//Counters_Spec,
+
+		//CsvProfiler_RegisterCategory,
+		//CsvProfiler_DefineInlineStat,
+		//CsvProfiler_DefineDeclaredStat,
+		//CsvProfiler_Metadata,
+		//CsvProfiler_BeginCapture,
+
+		Memory_Marker,
+
+		LLM_TagValue,
+		//LLM_TagsSpec,
+		//LLM_TrackerSpec,
+
+		SlateTrace_AddWidget,
+
+		Count
+	};
+
+	typedef TFunction<void(const UE::Trace::IAnalyzer::FOnEventContext& Context)> KnownEventCallback;
+
+	struct FKnownEvent
+	{
+		uint32 Id = InvalidEventId;
+		KnownEventCallback Callback = nullptr;
+	};
+
+	struct FThreadInfo
+	{
+		uint32 Id = 0;
+		uint32 Reserved = 0; // for padding
+		uint64 Cycle = 0;
+	};
+
+public:
+	FKnownEventProcessor();
+	virtual ~FKnownEventProcessor();
+
+	void Reset();
+
+	void SetOnUpdateThreadTime(const FOnUpdateThreadTimeFunc& Callback)
+	{
+		OnUpdateThreadTimeCallback = Callback;
+	}
+
+	void OnAnalysisBegin(const UE::Trace::IAnalyzer::FOnAnalysisContext& Context);
+	void OnAnalysisEnd();
+	void OnNewEvent(uint32 EventId, const ANSICHAR* LoggerName, const ANSICHAR* EventName);
+	void OnEvent(uint32 EventId, const UE::Trace::IAnalyzer::FOnEventContext& Context);
+
+private:
+	bool IsKnownEvent(uint32 EventId, EKnownEvent KnownEvent) const
+	{
+		return EventId == KnownEvents[(int)KnownEvent].Id;
+	}
+
+	void RegisterCallbacks();
+
+	void ProcessCpuEventBatchV1(const UE::Trace::IAnalyzer::FOnEventContext& Context, const uint8* BufferPtr, uint32 BufferSize);
+	uint64 ProcessCpuEventBatchV1(const UE::Trace::IAnalyzer::FOnEventContext& Context, const uint8* BufferPtr, uint32 BufferSize, uint64 ThreadCycle);
+
+	void ProcessCpuEventBatchV2V3(const UE::Trace::IAnalyzer::FOnEventContext& Context, const uint8* BufferPtr, uint32 BufferSize, int32 Version);
+	uint64 ProcessCpuEventBatchV2V3(const UE::Trace::IAnalyzer::FOnEventContext& Context, const uint8* BufferPtr, uint32 BufferSize, int32 Version, uint64 ThreadCycle);
+
+	FThreadInfo& GetOrAddThread(uint32 ThreadId);
+
+	void UpdateThreadTime(const UE::Trace::IAnalyzer::FOnEventContext& Context, uint64 Cycle);
+	void OnUpdateThreadTime(uint32 ThreadId, const UE::Trace::IAnalyzer::FOnEventContext& Context, uint64 Cycle);
+
+private:
+	static constexpr uint32 InvalidEventId = ~0u;
+	static constexpr uint32 MaxKnownEventId = 0xFFFF;
+
+	TArray<FKnownEvent> KnownEvents; // [EKnownEvent::Count]
+	TArray<EKnownEvent> KnownEventMap; // index == EventId --> EKnownEvent; [MaxKnownEventId + 1]
+
+	TMap<uint32, FThreadInfo> ThreadMap; // ThreadId --> FThreadInfo
+
+	FOnUpdateThreadTimeFunc OnUpdateThreadTimeCallback;
+};
+
+} // namespace TraceServices
diff --git a/Engine/Source/Developer/TraceServices/Private/Analyzers/TrimAnalyzer.cpp b/Engine/Source/Developer/TraceServices/Private/Analyzers/TrimAnalyzer.cpp
--- a/Engine/Source/Developer/TraceServices/Private/Analyzers/TrimAnalyzer.cpp
+++ b/Engine/Source/Developer/TraceServices/Private/Analyzers/TrimAnalyzer.cpp
@@ -16,24 +16,26 @@
 #include "Common/Utils.h"
 #include "TraceServices/TraceTrimmer.h"
 
+#define TRIMMER_USE_THREAD_TABLE 0
+
 ////////////////////////////////////////////////////////////////////////////////////////////////////
 
 #define LogTraceTrimmerError(Format, ...) \
 	{ \
 		++NumErrors; \
-		wprintf(L"Error: " Format "\n", ##__VA_ARGS__); \
-		UE_LOG(LogTraceServices, Error, TEXT("[TraceTrimmer] " Format), ##__VA_ARGS__); \
+		printf("Error: " Format "\n", ##__VA_ARGS__); \
+		UE_LOGF(LogTraceServices, Error, "[TraceTrimmer] " Format, ##__VA_ARGS__); \
 	}
 #define LogTraceTrimmerWarning(Format, ...) \
 	{ \
 		++NumWarnings; \
-		wprintf(L"Warning: " Format "\n", ##__VA_ARGS__); \
-		UE_LOG(LogTraceServices, Warning, TEXT("[TraceTrimmer] " Format), ##__VA_ARGS__); \
+		printf("Warning: " Format "\n", ##__VA_ARGS__); \
+		UE_LOGF(LogTraceServices, Warning, "[TraceTrimmer] " Format, ##__VA_ARGS__); \
 	}
 #define LogTraceTrimmerMessage(Format, ...) \
 	{ \
-		wprintf(L"" Format "\n", ##__VA_ARGS__); \
-		UE_LOG(LogTraceServices, Log, TEXT("[TraceTrimmer] " Format), ##__VA_ARGS__); \
+		printf(Format "\n", ##__VA_ARGS__); \
+		UE_LOGF(LogTraceServices, Log, "[TraceTrimmer] " Format, ##__VA_ARGS__); \
 	}
 
 ////////////////////////////////////////////////////////////////////////////////////////////////////
@@ -41,19 +43,6 @@
 namespace TraceServices
 {
 
-////////////////////////////////////////////////////////////////////////////////////////////////////
-
-static uint32 GetThreadIdField(const UE::Trace::IAnalyzer::FOnEventContext& Context, const ANSICHAR* FieldName = "ThreadId")
-{
-	// Trace analysis was changed to be able to provide a suitable id. Prior to
-	// this users of Trace would send along their own thread ids. For backwards
-	// compatibility we'll bias field thread ids to avoid collision with Trace's.
-	static const uint32 Bias = 0x70000000;
-	uint32 ThreadId = Context.EventData.GetValue<uint32>(FieldName, 0);
-	ThreadId |= ThreadId ? Bias : Context.ThreadInfo.GetId();
-	return ThreadId;
-}
-
 ////////////////////////////////////////////////////////////////////////////////////////////////////
 // FTrimAnalyzer
 ////////////////////////////////////////////////////////////////////////////////////////////////////
@@ -62,6 +51,18 @@ FTrimAnalyzer::FTrimAnalyzer(FTrimParameters& InParameters, UE::Trace::FTraceWri
 	: Parameters(InParameters)
 	, TraceWriter(InTraceWriter)
 {
+#if TRIMMER_USE_THREAD_TABLE
+	ThreadTable.AddDefaulted(ThreadTableInitialSize);
+#endif
+
+	KnownEvents.SetOnUpdateThreadTime(
+		[this]
+		(uint32 ThreadId, const FOnEventContext& Context, uint64 Cycle)
+		{
+			FThreadInfo& Thread = GetOrAddThread(ThreadId);
+			UpdateThreadTime(Thread, Context, Cycle);
+		});
+
 	// Empty IncludeFilter is equivalent with "*.*" (include all).
 	if (!Parameters.Include.IsEmpty() &&
 		!Parameters.Include.Equals("*.*"))
@@ -79,11 +80,15 @@ FTrimAnalyzer::FTrimAnalyzer(FTrimParameters& InParameters, UE::Trace::FTraceWri
 
 FTrimAnalyzer::~FTrimAnalyzer()
 {
+	for (FThreadInfo* Thread : SortedThreads)
+	{
+		delete Thread;
+	}
 }
 
 ////////////////////////////////////////////////////////////////////////////////////////////////////
 
-void FTrimAnalyzer::InitNameFilter(FAnsiString InFilter, FNameFilter& OutNameFilter)
+void FTrimAnalyzer::InitNameFilter(const FAnsiString& InFilter, FNameFilter& OutNameFilter)
 {
 	TArray<FAnsiString> FilterPatterns;
 	InFilter.ParseIntoArray(FilterPatterns, ",", true);
@@ -99,7 +104,7 @@ void FTrimAnalyzer::InitNameFilter(FAnsiString InFilter, FNameFilter& OutNameFil
 		}
 		else
 		{
-			LogTraceTrimmerWarning("Ignoring malformed event name pattern (%s) in filter (%s)!", ANSI_TO_TCHAR(*Pattern), ANSI_TO_TCHAR(*InFilter));
+			LogTraceTrimmerWarning("Ignoring malformed event name pattern (%s) in filter (%s)!", *Pattern, *InFilter);
 		}
 	}
 }
@@ -111,322 +116,118 @@ void FTrimAnalyzer::OnAnalysisBegin(const FOnAnalysisContext& Context)
 	LogTraceTrimmerMessage("Trimming...");
 
 	AnalysisBeginTimestamp = FPlatformTime::Cycles64();
+	LastProgressMessageTimestamp = AnalysisBeginTimestamp;
+	ProgressMessageInterval = FPlatformTime::SecondsToCycles64(10.0);
 
 	Context.InterfaceBuilder.RouteAllEvents(0, false);
 	Context.InterfaceBuilder.RouteAllEvents(1, true);
 
-	KnownEvents[(int)EKnownEvent::Invalid].Id = InvalidEventId;
-	KnownEvents[(int)EKnownEvent::Trace_NewEvent].Id = 0;
-	for (uint32 EventId = 0; EventId <= MaxKnownEventId; ++EventId)
-	{
-		KnownEventMap[EventId] = EKnownEvent::Invalid;
-	}
-	KnownEventMap[0] = EKnownEvent::Trace_NewEvent;
-
-	RegisterCallbacksForKnownEvents();
+	KnownEvents.OnAnalysisBegin(Context);
 
 	TraceWriter.Begin(UE::Trace::ETraceWriterBeginOptions::DeclareDefaultEvents);
 }
 
 ////////////////////////////////////////////////////////////////////////////////////////////////////
 
-void FTrimAnalyzer::OnAnalysisEnd()
+static FAnsiString FormatEventCount(uint64 Count)
 {
-	if (bIsTraceNewTraceEventDetected)
+	if (Count < 1000ull)
 	{
-		TraceWriter.End();
+		return FAnsiString::Printf("%llu", Count);
 	}
-	else
-	{
-		LogTraceTrimmerError("Invalid trace ($Trace.NewTrace not detected)!");
-	}
-
-	const uint64 AnalysisDuration = FPlatformTime::Cycles64() - AnalysisBeginTimestamp;
-	const double Duration = double(AnalysisDuration) * FPlatformTime::GetSecondsPerCycle();
-	FString AnalysisDurationString;
-	if (Duration < 1.0)
+	if (Count < 1000000ull)
 	{
-		AnalysisDurationString = FString::Printf(TEXT("%.06f seconds"), Duration);
+		return FAnsiString::Printf("%lluK", Count / 1000ull);
 	}
-	else if (Duration < 60.0)
+	if (Count < 1000000000ull)
 	{
-		AnalysisDurationString = FString::Printf(TEXT("%.02f seconds"), Duration);
+		return FAnsiString::Printf("%lluM", Count / 1000000ull);
 	}
-	else if (Duration < 120.0)
-	{
-		AnalysisDurationString = FString::Printf(TEXT("%.02fs (1 minute %lli seconds)"), Duration, (int64)Duration % 60);
-	}
-	else
-	{
-		AnalysisDurationString = FString::Printf(TEXT("%.02fs (%lli minutes %lli seconds)"), Duration, (int64)Duration / 60, (int64)Duration % 60);
-	}
-	LogTraceTrimmerMessage("  done in %s", *AnalysisDurationString);
+	return FAnsiString::Printf("%lluG", Count / 1000000000ull);
 }
 
 ////////////////////////////////////////////////////////////////////////////////////////////////////
 
-void FTrimAnalyzer::OnVersion(uint32 InTransportVersion, uint32 InProtocolVersion)
+static FAnsiString FormatDuration(double Duration)
 {
-	TransportVersion = InTransportVersion;
-	ProtocolVersion = InProtocolVersion;
-}
-
-////////////////////////////////////////////////////////////////////////////////////////////////////
+	if (Duration < 1.0)
+	{
+		return FAnsiString::Printf("%.06f seconds", Duration);
+	}
 
-void FTrimAnalyzer::RegisterCallbacksForKnownEvents()
-{
-	//////////////////////////////////////////////////
-	// Misc.*
+	if (Duration < 60.0)
+	{
+		return FAnsiString::Printf("%.02f seconds", Duration);
+	}
 
-	KnownEvents[(int)EKnownEvent::Misc_RegionBeginWithId].Callback =
-		[this](const FOnEventContext& Context)
-		{
-			const uint64 Cycle = Context.EventData.GetValue<uint64>("CycleAndId");
-			UpdateThreadTime(Context, Cycle);
-		};
+	TAnsiStringBuilder<256> Str;
 
-	//////////////////////////////////////////////////
-	// CpuProfiler.*
+	Str.Appendf("%.02fs (", Duration);
 
-	KnownEvents[(int)EKnownEvent::CpuProfiler_EventBatchV3].Callback =
-		[this](const FOnEventContext& Context)
-		{
-			const FEventData& EventData = Context.EventData;
-			TArrayView<const uint8> Data = EventData.GetArrayView<uint8>("Data");
-			if (Data.GetData() != nullptr)
-			{
-				ProcessCpuEventBatchV2V3(Context, Data.GetData(), Data.Num(), 3);
-			}
-		};
+	bool bIsFirst = true;
+	int64 Seconds = (int64)Duration;
 
-	KnownEvents[(int)EKnownEvent::CpuProfiler_EventBatchV2].Callback =
-		[this](const FOnEventContext& Context)
+	auto AppendDuration =
+		[&Str, &bIsFirst, &Seconds]
+		(int64 UnitDuration, const ANSICHAR * Unit)
 		{
-			const FEventData& EventData = Context.EventData;
-			TArrayView<const uint8> Data = EventData.GetArrayView<uint8>("Data");
-			if (Data.GetData() != nullptr)
+			if (Seconds >= UnitDuration)
 			{
-				ProcessCpuEventBatchV2V3(Context, Data.GetData(), Data.Num(), 2);
-			}
-		};
+				if (!bIsFirst)
+				{
+					Str.Append(" ");
+				}
+				bIsFirst = false;
 
-	KnownEvents[(int)EKnownEvent::CpuProfiler_EventBatch].Callback =
-	KnownEvents[(int)EKnownEvent::CpuProfiler_EndCapture].Callback =
-		[this](const FOnEventContext& Context)
-		{
-			const FEventData& EventData = Context.EventData;
-			TArrayView<const uint8> Data = EventData.GetArrayView<uint8>("Data");
-			if (Data.GetData() != nullptr)
-			{
-				ProcessCpuEventBatchV1(Context, Data.GetData(), Data.Num());
+				int64 Units = Seconds / UnitDuration;
+				Seconds = Seconds % UnitDuration;
+				Str.Appendf("%lli %s", Units, Unit);
+				if (Units != 1)
+				{
+					Str.Append("s");
+				}
 			}
 		};
 
-	//////////////////////////////////////////////////
-	// GpuProfiler.*
+	AppendDuration(24 * 60 * 60, "day");
+	AppendDuration(60 * 60, "hour");
+	AppendDuration(60, "minute");
+	AppendDuration(1, "second");
 
-	KnownEvents[(int)EKnownEvent::GpuProfiler_EventBeginBreadcrumb].Callback =
-	KnownEvents[(int)EKnownEvent::GpuProfiler_EventBeginWork].Callback =
-		[this](const FOnEventContext& Context)
-		{
-			const uint64 Cycle = Context.EventData.GetValue<uint64>("GPUTimestampTOP");
-			UpdateThreadTime(Context, Cycle);
-		};
-	KnownEvents[(int)EKnownEvent::GpuProfiler_EventEndBreadcrumb].Callback =
-	KnownEvents[(int)EKnownEvent::GpuProfiler_EventEndWork].Callback =
-		[this](const FOnEventContext& Context)
-		{
-			const uint64 Cycle = Context.EventData.GetValue<uint64>("GPUTimestampBOP");
-			UpdateThreadTime(Context, Cycle);
-		};
-	KnownEvents[(int)EKnownEvent::GpuProfiler_EventWait].Callback =
-		[this](const FOnEventContext& Context)
-		{
-			const uint64 Cycle = Context.EventData.GetValue<uint64>("EndTime");
-			UpdateThreadTime(Context, Cycle);
-		};
-	KnownEvents[(int)EKnownEvent::GpuProfiler_SignalFence].Callback =
-	KnownEvents[(int)EKnownEvent::GpuProfiler_WaitFence].Callback =
-		[this](const FOnEventContext& Context)
-		{
-			const uint64 Cycle = Context.EventData.GetValue<uint64>("CPUTimestamp");
-			UpdateThreadTime(Context, Cycle);
-		};
+	Str.Append(")");
 
-	//////////////////////////////////////////////////
-	// Known events with a "Cycle" field
-
-	KnownEvents[(int)EKnownEvent::Misc_Bookmark].Callback =
-	KnownEvents[(int)EKnownEvent::Misc_ScreenshotHeader].Callback =
-	KnownEvents[(int)EKnownEvent::Misc_RegionBegin].Callback =
-	KnownEvents[(int)EKnownEvent::Misc_RegionEnd].Callback =
-	KnownEvents[(int)EKnownEvent::Misc_RegionEndWithId].Callback =
-	KnownEvents[(int)EKnownEvent::Misc_BeginFrame].Callback =
-	KnownEvents[(int)EKnownEvent::Misc_EndFrame].Callback =
-	KnownEvents[(int)EKnownEvent::Logging_LogMessage].Callback =
-	KnownEvents[(int)EKnownEvent::PlatformFile_BeginOpen].Callback =
-	KnownEvents[(int)EKnownEvent::PlatformFile_EndOpen].Callback =
-	KnownEvents[(int)EKnownEvent::PlatformFile_BeginRead].Callback =
-	KnownEvents[(int)EKnownEvent::PlatformFile_EndRead].Callback =
-	KnownEvents[(int)EKnownEvent::PlatformFile_BeginClose].Callback =
-	KnownEvents[(int)EKnownEvent::PlatformFile_EndClose].Callback =
-	KnownEvents[(int)EKnownEvent::CpuProfiler_EndThread].Callback =
-	KnownEvents[(int)EKnownEvent::Memory_Marker].Callback =
-	KnownEvents[(int)EKnownEvent::LLM_TagValue].Callback =
-	KnownEvents[(int)EKnownEvent::SlateTrace_AddWidget].Callback =
-		[this](const FOnEventContext& Context)
-		{
-			const uint64 Cycle = Context.EventData.GetValue<uint64>("Cycle");
-			UpdateThreadTime(Context, Cycle);
-		};
+	return FAnsiString(Str.ToView());
 }
 
-///////////////////////////////////////////////////////////////////////////////////////////////////
+////////////////////////////////////////////////////////////////////////////////////////////////////
 
-void FTrimAnalyzer::ProcessKnownEvent(uint32 EventId, const ANSICHAR* LoggerName, const ANSICHAR* EventName)
+void FTrimAnalyzer::OnAnalysisEnd()
 {
-	#define BEGIN_KNOWN_LOGGER(_Name) \
-		if (FCStringAnsi::Strcmp(LoggerName, _Name) == 0) \
-		{
-
-	#define KNOWN_EVENT(_Name,_Id) \
-			if (FCStringAnsi::Strcmp(EventName, _Name) == 0) \
-			{ \
-				KnownEvents[(int)EKnownEvent::_Id].Id = EventId; \
-				KnownEventMap[EventId] = EKnownEvent::_Id; \
-				break; \
-			}
+	KnownEvents.OnAnalysisEnd();
 
-	#define END_KNOWN_LOGGER() \
-			break; \
-		}
-
-	while (true)
+	if (bIsTraceNewTraceEventDetected)
 	{
-		BEGIN_KNOWN_LOGGER("$Trace")
-			KNOWN_EVENT("NewEvent", Trace_NewEvent)
-			//KNOWN_EVENT("NewTrace", Trace_NewTrace)
-			//KNOWN_EVENT("ThreadTiming", Trace_ThreadTiming)
-			//KNOWN_EVENT("ThreadInfo", Trace_ThreadInfo)
-			//KNOWN_EVENT("ThreadGroupBegin", Trace_ThreadGroupBegin)
-			//KNOWN_EVENT("ThreadGroupEnd", Trace_ThreadGroupEnd)
-			//KNOWN_EVENT("ChannelAnnounce", Trace_ChannelAnnounce_Old) // deprecated
-		END_KNOWN_LOGGER()
-
-		//BEGIN_KNOWN_LOGGER("Trace")
-			//KNOWN_EVENT("ChannelAnnounce", Trace_ChannelAnnounce)
-			//KNOWN_EVENT("ChannelAnnounce", Trace_ChannelToggle)
-		//END_KNOWN_LOGGER()
-
-		//BEGIN_KNOWN_LOGGER("Diagnostics")
-			//KNOWN_EVENT("Session", Diagnostics_Session) // deprecated
-			//KNOWN_EVENT("Session2", Diagnostics_Session2)
-		//END_KNOWN_LOGGER()
-
-		BEGIN_KNOWN_LOGGER("Misc")
-			//KNOWN_EVENT("CreateThread", Misc_CreateThread)
-			//KNOWN_EVENT("SetThreadGroup", Misc_SetThreadGroup)
-			//KNOWN_EVENT("BeginThreadGroupScope", Misc_BeginThreadGroupScope)
-			//KNOWN_EVENT("BookmarkSpec", Misc_BookmarkSpec)
-			KNOWN_EVENT("Bookmark", Misc_Bookmark)
-			KNOWN_EVENT("ScreenshotHeader", Misc_ScreenshotHeader)
-			KNOWN_EVENT("RegionBegin", Misc_RegionBegin)
-			KNOWN_EVENT("RegionBeginWithId", Misc_RegionBeginWithId)
-			KNOWN_EVENT("RegionEnd", Misc_RegionEnd)
-			KNOWN_EVENT("RegionEndWithId", Misc_RegionEndWithId)
-			//KNOWN_EVENT("BeginGameFrame", Misc_BeginGameFrame)
-			//KNOWN_EVENT("EndGameFrame", Misc_EndGameFrame)
-			//KNOWN_EVENT("BeginRenderFrame", Misc_BeginRenderFrame)
-			//KNOWN_EVENT("EndRenderFrame", Misc_EndRenderFrame)
-			KNOWN_EVENT("BeginFrame", Misc_BeginFrame)
-			KNOWN_EVENT("EndFrame", Misc_EndFrame)
-		END_KNOWN_LOGGER()
-
-		BEGIN_KNOWN_LOGGER("Logging")
-			//KNOWN_EVENT("LogCategory", Logging_LogCategory)
-			//KNOWN_EVENT("LogMessageSpec", Logging_LogMessageSpec)
-			KNOWN_EVENT("LogMessage", Logging_LogMessage)
-		END_KNOWN_LOGGER()
-
-		BEGIN_KNOWN_LOGGER("PlatformFile")
-			KNOWN_EVENT("BeginOpen", PlatformFile_BeginOpen)
-			KNOWN_EVENT("EndOpen", PlatformFile_EndOpen)
-			KNOWN_EVENT("BeginRead", PlatformFile_BeginRead)
-			KNOWN_EVENT("EndRead", PlatformFile_EndRead)
-			KNOWN_EVENT("BeginClose", PlatformFile_BeginClose)
-			KNOWN_EVENT("EndClose", PlatformFile_EndClose)
-		END_KNOWN_LOGGER()
-
-		//BEGIN_KNOWN_LOGGER("LoadTime")
-			//KNOWN_EVENT("ClassInfo", LoadTime_ClassInfo)
-			//KNOWN_EVENT("NewAsyncPackage", LoadTime_NewAsyncPackage)
-		//END_KNOWN_LOGGER()
-
-		BEGIN_KNOWN_LOGGER("GpuProfiler")
-			//KNOWN_EVENT("EventSpec", GpuProfiler_EventSpec) // deprecated
-			//KNOWN_EVENT("Frame", GpuProfiler_Frame) // deprecated
-			//KNOWN_EVENT("Frame2", GpuProfiler_Frame2) // deprecated
-			//KNOWN_EVENT("Init", GpuProfiler_Init)
-			//KNOWN_EVENT("QueueSpec", GpuProfiler_QueueSpec)
-			//KNOWN_EVENT("EventFrameBoundary", GpuProfiler_EventFrameBoundary)
-			//KNOWN_EVENT("EventBreadcrumbSpec", GpuProfiler_EventBreadcrumbSpec)
-			KNOWN_EVENT("EventBeginBreadcrumb", GpuProfiler_EventBeginBreadcrumb)
-			KNOWN_EVENT("EventEndBreadcrumb", GpuProfiler_EventEndBreadcrumb)
-			KNOWN_EVENT("EventBeginWork", GpuProfiler_EventBeginWork)
-			KNOWN_EVENT("EventEndWork", GpuProfiler_EventEndWork)
-			KNOWN_EVENT("EventWait", GpuProfiler_EventWait)
-			//KNOWN_EVENT("EventStats", GpuProfiler_EventStats)
-			KNOWN_EVENT("SignalFence", GpuProfiler_SignalFence)
-			KNOWN_EVENT("WaitFence", GpuProfiler_WaitFence)
-		END_KNOWN_LOGGER()
-
-		BEGIN_KNOWN_LOGGER("CpuProfiler")
-			//KNOWN_EVENT("EventSpec", CpuProfiler_EventSpec)
-			//KNOWN_EVENT("MetadataSpec", CpuProfiler_MetadataSpec)
-			KNOWN_EVENT("Metadata", CpuProfiler_Metadata)
-			KNOWN_EVENT("EventBatchV3", CpuProfiler_EventBatchV3)
-			KNOWN_EVENT("EventBatchV2", CpuProfiler_EventBatchV2) // deprecated
-			KNOWN_EVENT("EventBatch", CpuProfiler_EventBatch) // deprecated
-			KNOWN_EVENT("EndCapture", CpuProfiler_EndCapture) // deprecated
-			KNOWN_EVENT("EndThread", CpuProfiler_EndThread)
-		END_KNOWN_LOGGER()
-
-		//BEGIN_KNOWN_LOGGER("Stats")
-			//KNOWN_EVENT("Spec", Stats_Spec)
-		//END_KNOWN_LOGGER()
-
-		//BEGIN_KNOWN_LOGGER("Counters")
-			//KNOWN_EVENT("Spec", Counters_Spec)
-		//END_KNOWN_LOGGER()
-
-		//BEGIN_KNOWN_LOGGER("CsvProfiler")
-			//KNOWN_EVENT("RegisterCategory", CsvProfiler_RegisterCategory)
-			//KNOWN_EVENT("DefineInlineStat", CsvProfiler_DefineInlineStat)
-			//KNOWN_EVENT("DefineDeclaredStat", CsvProfiler_DefineDeclaredStat)
-			//KNOWN_EVENT("Metadata", CsvProfiler_Metadata)
-			//KNOWN_EVENT("BeginCapture", CsvProfiler_BeginCapture)
-		//END_KNOWN_LOGGER()
-
-		BEGIN_KNOWN_LOGGER("Memory")
-			KNOWN_EVENT("Marker", Memory_Marker)
-		END_KNOWN_LOGGER()
-
-		BEGIN_KNOWN_LOGGER("LLM")
-			KNOWN_EVENT("TagValue", LLM_TagValue)
-			//KNOWN_EVENT("TagsSpec", LLM_TagsSpec)
-			//KNOWN_EVENT("TrackerSpec", LLM_TrackerSpec)
-		END_KNOWN_LOGGER()
-
-		BEGIN_KNOWN_LOGGER("SlateTrace")
-			KNOWN_EVENT("AddWidget", SlateTrace_AddWidget)
-		END_KNOWN_LOGGER()
-
-		break;
+		TraceWriter.End();
 	}
+	else
+	{
+		LogTraceTrimmerError("Invalid trace ($Trace.NewTrace not detected)!");
+	}
+
+	LogTraceTrimmerMessage("Read %llu events. Written %llu events.", ReadEvents, WrittenEvents);
+	LogTraceTrimmerMessage("Analyzed Trace Session Duration: %s", *FormatDuration(SessionCurrentTime));
 
-	#undef BEGIN_KNOWN_LOGGER
-	#undef KNOWN_EVENT
-	#undef END_KNOWN_LOGGER
+	const uint64 AnalysisDuration = FPlatformTime::Cycles64() - AnalysisBeginTimestamp;
+	const double Duration = double(AnalysisDuration) * FPlatformTime::GetSecondsPerCycle();
+	LogTraceTrimmerMessage("Trimming completed in %s.", *FormatDuration(Duration));
+}
+
+////////////////////////////////////////////////////////////////////////////////////////////////////
+
+void FTrimAnalyzer::OnVersion(uint32 InTransportVersion, uint32 InProtocolVersion)
+{
+	TransportVersion = InTransportVersion;
+	ProtocolVersion = InProtocolVersion;
 }
 
 ////////////////////////////////////////////////////////////////////////////////////////////////////
@@ -506,13 +307,13 @@ bool FTrimAnalyzer::OnNewEvent(uint16 RouteId, const FEventTypeInfo& TypeInfo)
 		if (EventInfo->TypeInfo->GetId() != EventId)
 		{
 			LogTraceTrimmerWarning("Event (%s.%s) with id %u is already declared with id %u!",
-				ANSI_TO_TCHAR(LoggerName), ANSI_TO_TCHAR(EventName), EventId, EventInfo->TypeInfo->GetId());
+				LoggerName, EventName, EventId, EventInfo->TypeInfo->GetId());
 
 			EventInfo = nullptr;
 		}
 		else
 		{
-			LogTraceTrimmerWarning("Event (%s.%s) is already declared!", ANSI_TO_TCHAR(LoggerName), ANSI_TO_TCHAR(EventName));
+			LogTraceTrimmerWarning("Event (%s.%s) is already declared!", LoggerName, EventName);
 			return true;
 		}
 	}
@@ -534,7 +335,7 @@ bool FTrimAnalyzer::OnNewEvent(uint16 RouteId, const FEventTypeInfo& TypeInfo)
 
 	//////////////////////////////////////////////////
 
-	ProcessKnownEvent(EventId, LoggerName, EventName);
+	KnownEvents.OnNewEvent(EventId, LoggerName, EventName);
 
 	//////////////////////////////////////////////////
 	// Known $Trace.* events
@@ -560,7 +361,7 @@ bool FTrimAnalyzer::OnNewEvent(uint16 RouteId, const FEventTypeInfo& TypeInfo)
 		// $Trace.* events are expected to be declared automatically (by TraceWriter.Begin() or by TraceWriter.DeclareDefaultEvents()).
 		if (!WriterEventInfo)
 		{
-			LogTraceTrimmerWarning("Unknown default event (%s.%s)!", ANSI_TO_TCHAR(LoggerName), ANSI_TO_TCHAR(EventName));
+			LogTraceTrimmerWarning("Unknown default event (%s.%s)!", LoggerName, EventName);
 		}
 	}
 
@@ -578,7 +379,8 @@ bool FTrimAnalyzer::OnNewEvent(uint16 RouteId, const FEventTypeInfo& TypeInfo)
 
 	if (WriterEventInfo)
 	{
-		LogTraceTrimmerMessage("Event %u (%s.%s) was already declared (writer id %u).", EventId, ANSI_TO_TCHAR(LoggerName), ANSI_TO_TCHAR(EventName), WriterEventInfo->GetId());
+		LogTraceTrimmerMessage("Event %u (%s.%s) was already declared (writer id %u).",
+			EventId, LoggerName, EventName, WriterEventInfo->GetId());
 	}
 	else
 	{
@@ -607,7 +409,7 @@ bool FTrimAnalyzer::OnNewEvent(uint16 RouteId, const FEventTypeInfo& TypeInfo)
 				{
 					WriterRefUid = 0;
 					LogTraceTrimmerWarning("Invalid RefUid (%u) for event (%s.%s) field \"%s\"!",
-						RefUid, ANSI_TO_TCHAR(LoggerName), ANSI_TO_TCHAR(EventName), ANSI_TO_TCHAR(FieldInfo.GetName()));
+						RefUid, LoggerName, EventName, FieldInfo.GetName());
 				}
 				Builder.ReferenceField(FieldInfo.GetName(), WriterFieldType, WriterRefUid);
 			}
@@ -643,6 +445,8 @@ bool FTrimAnalyzer::OnEvent(uint16 RouteId, EStyle Style, const FOnEventContext&
 	const ANSICHAR* EventName = TypeInfo.GetName();
 	check(EventName != nullptr);
 
+	++ReadEvents;
+
 	//////////////////////////////////////////////////
 
 	bool bShouldWriteEvent = true;
@@ -656,7 +460,7 @@ bool FTrimAnalyzer::OnEvent(uint16 RouteId, EStyle Style, const FOnEventContext&
 		}
 		else
 		{
-			LogTraceTrimmerWarning("Event %u (%s.%s) is not registered!", EventId, ANSI_TO_TCHAR(LoggerName), ANSI_TO_TCHAR(EventName));
+			LogTraceTrimmerWarning("Event %u (%s.%s) is not registered!", EventId, LoggerName, EventName);
 			const bool bIsFilteredOut = !FilterByEventName(LoggerName, EventName);
 			EventInfo = &EventInfos.Add(EventId, { &TypeInfo, InvalidWriterEventId, bIsFilteredOut });
 			bShouldWriteEvent = false;
@@ -665,17 +469,8 @@ bool FTrimAnalyzer::OnEvent(uint16 RouteId, EStyle Style, const FOnEventContext&
 	check(EventInfo);
 
 	//////////////////////////////////////////////////
-	// Processing known events
-
-	check(EventId <= MaxKnownEventId);
-	EKnownEvent KnownEvent = KnownEventMap[EventId];
-	check((int)KnownEvent < (int)EKnownEvent::Count);
-	const FKnownEvent& Event = KnownEvents[(int)KnownEvent];
 
-	if (Event.Callback)
-	{
-		Event.Callback(Context);
-	}
+	KnownEvents.OnEvent(EventId, Context);
 
 	//////////////////////////////////////////////////
 	// Known $Trace.* events
@@ -723,14 +518,14 @@ bool FTrimAnalyzer::OnEvent(uint16 RouteId, EStyle Style, const FOnEventContext&
 			{
 				LogTraceTrimmerWarning("Thread %u is already registered with different properties (previous Name=\"%s\", SystemId=%u, SortHint=%d; new Name=\"%s\", SystemId=%u, SortHint=%d)!",
 					ThreadId,
-					*FString(ThreadInfo->GetName()), ThreadInfo->GetSystemId(), ThreadInfo->GetSortHint(),
-					*FString(ThreadName), SystemId, SortHint);
+					*ThreadInfo->GetName(), ThreadInfo->GetSystemId(), ThreadInfo->GetSortHint(),
+					*FAnsiString(ThreadName), SystemId, SortHint);
 			}
 			else
 			{
 				LogTraceTrimmerMessage("Thread %u is already registered (Name=\"%s\", SystemId=%u, SortHint=%d).",
 					ThreadId,
-					*FString(ThreadName), SystemId, SortHint);
+					*FAnsiString(ThreadName), SystemId, SortHint);
 			}
 			//bShouldWriteEvent = false;
 		}
@@ -738,7 +533,7 @@ bool FTrimAnalyzer::OnEvent(uint16 RouteId, EStyle Style, const FOnEventContext&
 		{
 			LogTraceTrimmerMessage("Registering thread %u (Name=\"%s\", SystemId=%u, SortHint=%d).",
 				ThreadId,
-				*FString(ThreadName), SystemId, SortHint);
+				*FAnsiString(ThreadName), SystemId, SortHint);
 
 			constexpr bool bShouldWriteThreadInfo = false;
 			TraceWriter.RegisterCustomThread(ThreadId, ThreadName, SystemId, SortHint, bShouldWriteThreadInfo);
@@ -769,7 +564,7 @@ bool FTrimAnalyzer::OnEvent(uint16 RouteId, EStyle Style, const FOnEventContext&
 	{
 		if (EventInfo->WriterEventId == InvalidWriterEventId)
 		{
-			LogTraceTrimmerWarning("Event %u (%s.%s) is not declared!", EventId, ANSI_TO_TCHAR(LoggerName), ANSI_TO_TCHAR(EventName));
+			LogTraceTrimmerWarning("Event %u (%s.%s) is not declared!", EventId, LoggerName, EventName);
 			bShouldWriteEvent = false;
 		}
 	}
@@ -780,8 +575,7 @@ bool FTrimAnalyzer::OnEvent(uint16 RouteId, EStyle Style, const FOnEventContext&
 	{
 		if (Style != EStyle::Normal)
 		{
-			LogTraceTrimmerWarning("Event (%s.%s) is dispatched as normal event, but with Style=%u!",
-				ANSI_TO_TCHAR(LoggerName), ANSI_TO_TCHAR(EventName), (uint32)Style);
+			LogTraceTrimmerWarning("Event (%s.%s) is dispatched as normal event, but with Style=%u!", LoggerName, EventName, (uint32)Style);
 		}
 	}
 	else if (RouteId == 1) // Scoped events
@@ -824,14 +618,12 @@ bool FTrimAnalyzer::OnEvent(uint16 RouteId, EStyle Style, const FOnEventContext&
 		}
 		else // unknown Style
 		{
-			LogTraceTrimmerWarning("Event (%s.%s) is dispatched as scoped event, but with Style=%u!",
-				ANSI_TO_TCHAR(LoggerName), ANSI_TO_TCHAR(EventName), (uint32)Style);
+			LogTraceTrimmerWarning("Event (%s.%s) is dispatched as scoped event, but with Style=%u!", LoggerName, EventName, (uint32)Style);
 		}
 	}
 	else // unknown RouteId
 	{
-		LogTraceTrimmerWarning("Event (%s.%s) is dispatched on unknown route %u!",
-			ANSI_TO_TCHAR(LoggerName), ANSI_TO_TCHAR(EventName), (uint32)RouteId);
+		LogTraceTrimmerWarning("Event (%s.%s) is dispatched on unknown route %u!", LoggerName, EventName, (uint32)RouteId);
 	}
 
 	//////////////////////////////////////////////////
@@ -870,41 +662,27 @@ bool FTrimAnalyzer::OnEvent(uint16 RouteId, EStyle Style, const FOnEventContext&
 		}
 
 		EventBuilder.End();
+		++WrittenEvents;
 	}
 
 	//////////////////////////////////////////////////
 
 	bool bShouldStopAnalysis = false;
 
-	// When we encounter a Sync event...
 	if (TypeInfo.IsSync())
 	{
-		double ThreadMinTime = CurrentThread.CurrentTime;
-
-		// ...we can sync times between threads.
-		for (auto& KV : ThreadMap)
-		{
-			FThreadInfo& Thread = KV.Value;
-
-			// Threads with old timestamps can be updated with time of the CurrentThread.
-			if (Thread.CurrentCycle < CurrentThread.CurrentCycle)
-			{
-				Thread.CurrentCycle = CurrentThread.CurrentCycle;
-				Thread.CurrentTime = CurrentThread.CurrentTime;
-			}
+		bShouldStopAnalysis = UpdateThreadsOnSyncEvent(CurrentThread);
+	}
 
-			// Compute also the min timestamp of all threads.
-			if (Thread.CurrentTime < ThreadMinTime)
-			{
-				ThreadMinTime = Thread.CurrentTime;
-			}
-		}
+	//////////////////////////////////////////////////
 
-		// If all threads have timestamps past the EndTime limit we can stop analysis.
-		if (ThreadMinTime > Parameters.EndTime)
-		{
-			bShouldStopAnalysis = true;
-		}
+	const uint64 CurrentTimestamp = FPlatformTime::Cycles64();
+	if (CurrentTimestamp - LastProgressMessageTimestamp > ProgressMessageInterval)
+	{
+		LastProgressMessageTimestamp = CurrentTimestamp;
+		const double TT = double(LastProgressMessageTimestamp - AnalysisBeginTimestamp) * FPlatformTime::GetSecondsPerCycle();
+		LogTraceTrimmerMessage("Trimming...   R: %-6s W: %-6s S: %.0fs   T: %.0fs",
+			*FormatEventCount(ReadEvents), *FormatEventCount(WrittenEvents), SessionCurrentTime, TT);
 	}
 
 	return !bShouldStopAnalysis;
@@ -1212,179 +990,166 @@ void FTrimAnalyzer::SetCurrentThread(uint32 ThreadId)
 
 FTrimAnalyzer::FThreadInfo& FTrimAnalyzer::GetOrAddThread(uint32 ThreadId)
 {
-	FThreadInfo* FoundThread = ThreadMap.Find(ThreadId);
-	if (FoundThread)
+#if TRIMMER_USE_THREAD_TABLE
+	// Try the fast access table.
+	if (ThreadId < uint32(ThreadTable.Num()))
 	{
-		return *FoundThread;
+		FThreadInfo* Thread = ThreadTable[ThreadId];
+		if (Thread)
+		{
+			return *Thread;
+		}
 	}
-	return ThreadMap.Add(ThreadId, { ThreadId });
-}
+#endif // TRIMMER_USE_THREAD_TABLE
 
-////////////////////////////////////////////////////////////////////////////////////////////////////
+	// Try the thread map.
+	FThreadInfo** FoundThread = ThreadMap.Find(ThreadId);
+	if (FoundThread)
+	{
+		return **FoundThread;
+	}
 
-void FTrimAnalyzer::UpdateThreadTime(const FOnEventContext& Context, uint64 Cycle)
-{
-	const uint32 ThreadId = Context.ThreadInfo.GetId();
-	FThreadInfo& Thread = GetOrAddThread(ThreadId);
-	UpdateThreadTime(Thread, Context, Cycle);
-}
+	// Create a new thread info.
+	FThreadInfo* Thread = new FThreadInfo();
+	Thread->Id = ThreadId;
 
-////////////////////////////////////////////////////////////////////////////////////////////////////
+	// Add the thread to the map.
+	ThreadMap.Add(ThreadId, Thread);
 
-void FTrimAnalyzer::UpdateThreadTime(FThreadInfo& Thread, const FOnEventContext& Context, uint64 Cycle)
-{
-	if (Cycle > Thread.CurrentCycle)
+#if TRIMMER_USE_THREAD_TABLE
+	// Add the thread to the fast access table.
+	if (ThreadId < ThreadTableMaxSize)
 	{
-		const double Time = Context.EventTime.AsSeconds(Cycle);
-
-		Thread.CurrentCycle = Cycle;
-		Thread.CurrentTime = Time;
-
-		if (Cycle > SessionCurrentCycle)
+		if (ThreadId >= uint32(ThreadTable.Num()))
 		{
-			SessionCurrentCycle = Cycle;
-			SessionCurrentTime = Time;
+			ThreadTable.AddDefaulted(int32(ThreadId) - ThreadTable.Num() + 1);
 		}
+		ThreadTable[ThreadId] = Thread;
 	}
+#endif // TRIMMER_USE_THREAD_TABLE
+
+	// Add the thread to the sorting table.
+	Thread->SortIndex = SortedThreads.Num();
+	SortedThreads.Add(Thread);
+	SortThreads(Thread->SortIndex);
+
+	return *Thread;
 }
 
 ////////////////////////////////////////////////////////////////////////////////////////////////////
 
-void FTrimAnalyzer::UpdateSessionTime(uint64 Cycle, double Time)
+void FTrimAnalyzer::SortThreads(int32 ThreadIndex)
 {
-	if (Cycle > SessionCurrentCycle)
+	const uint64 Cycle = SortedThreads[ThreadIndex]->CurrentCycle;
+
+	int32 NextIndex = ThreadIndex + 1;
+	while (NextIndex < SortedThreads.Num() &&
+		Cycle > SortedThreads[NextIndex]->CurrentCycle)
 	{
-		SessionCurrentCycle = Cycle;
-		SessionCurrentTime = Time;
+		FThreadInfo* T = SortedThreads[ThreadIndex];
+		SortedThreads[ThreadIndex] = SortedThreads[NextIndex];
+		SortedThreads[NextIndex] = T;
+		T->SortIndex = NextIndex++;
+		SortedThreads[ThreadIndex]->SortIndex = ThreadIndex++;
+	}
+
+	int32 PrevIndex = ThreadIndex - 1;
+	while (PrevIndex >= 0 &&
+		Cycle < SortedThreads[PrevIndex]->CurrentCycle)
+	{
+		FThreadInfo* T = SortedThreads[ThreadIndex];
+		SortedThreads[ThreadIndex] = SortedThreads[PrevIndex];
+		SortedThreads[PrevIndex] = T;
+		T->SortIndex = PrevIndex--;
+		SortedThreads[ThreadIndex]->SortIndex = ThreadIndex--;
 	}
 }
 
 ////////////////////////////////////////////////////////////////////////////////////////////////////
 
-void FTrimAnalyzer::ProcessCpuEventBatchV1(const FOnEventContext& Context, const uint8* BufferPtr, uint32 BufferSize)
+bool FTrimAnalyzer::UpdateThreadsOnSyncEvent(FThreadInfo& CurrentThread)
 {
-	const FEventData& EventData = Context.EventData;
+	// When we encounter a Sync event...
 
-	const uint32 ThreadId = GetThreadIdField(Context);
-	FThreadInfo& Thread = GetOrAddThread(ThreadId);
-	uint64 LastCycle = Thread.LastCycle;
+	check(SortedThreads.Num() > 0); // as CurrentThread exists, SortedThreads should contain at least one thread
 
-	const uint8* BufferEnd = BufferPtr + BufferSize;
-	while (BufferPtr < BufferEnd)
+	// Threads with old timestamps can be updated with time of the CurrentThread.
+	for (auto Thread : SortedThreads)
 	{
-		uint64 DecodedCycle = FTraceAnalyzerUtils::Decode7bit(BufferPtr);
-		LastCycle += (DecodedCycle >> 1);
-
-		if (DecodedCycle & 1ull)
+		if (Thread->CurrentCycle < CurrentThread.CurrentCycle)
 		{
-			uint32 SpecId = (uint32)FTraceAnalyzerUtils::Decode7bit(BufferPtr);
+			Thread->CurrentCycle = CurrentThread.CurrentCycle;
+			Thread->CurrentTime = CurrentThread.CurrentTime;
+		}
+		else
+		{
+			break;
 		}
-	}
-	if (!ensure(BufferPtr == BufferEnd))
-	{
-		//error
 	}
 
-	Thread.LastCycle = LastCycle;
-	UpdateThreadTime(Thread, Context, LastCycle);
-	if (ThreadId != Context.ThreadInfo.GetId()) // ThreadId has bias (using ThreadId field)
-	{
-		UpdateThreadTime(Context, LastCycle);
-	}
+	// If all threads have timestamps past the EndTime limit we can stop analysis.
+	return (SortedThreads[0]->CurrentTime > Parameters.EndTime);
 }
 
 ////////////////////////////////////////////////////////////////////////////////////////////////////
 
-void FTrimAnalyzer::ProcessCpuEventBatchV2V3(const FOnEventContext& Context, const uint8* BufferPtr, uint32 BufferSize, int32 Version = 3)
+void FTrimAnalyzer::UpdateThreadTime(FThreadInfo& Thread, const FOnEventContext& Context, uint64 Cycle)
 {
-	const FEventData& EventData = Context.EventData;
-
-	const uint32 ThreadId = Context.ThreadInfo.GetId();
-	FThreadInfo& Thread = GetOrAddThread(ThreadId);
-	uint64 LastCycle = Thread.LastCycle;
+	if (Cycle > Thread.CurrentCycle)
+	{
+		const double Time = Context.EventTime.AsSeconds(Cycle);
 
-	check(Context.EventTime.GetTimestamp() == 0);
-	const uint64 BaseCycle = Context.EventTime.AsCycle64();
+		Thread.CurrentCycle = Cycle;
+		Thread.CurrentTime = Time;
 
-	const uint8* BufferEnd = BufferPtr + BufferSize;
-	while (BufferPtr < BufferEnd)
-	{
-		uint64 DecodedCycle = FTraceAnalyzerUtils::Decode7bit(BufferPtr);
-		uint64 ActualCycle = (DecodedCycle >> 2);
+		SortThreads(Thread.SortIndex);
 
-		// ActualCycle larger or equal to LastCycle means we have a new base value.
-		if (ActualCycle < LastCycle)
+		if (Cycle > SessionCurrentCycle)
 		{
-			ActualCycle += LastCycle;
+			SessionCurrentCycle = Cycle;
+			SessionCurrentTime = Time;
 		}
+	}
+}
 
-		// If we late connect we will be joining the cycle stream mid-flow and
-		// will have missed out on it's base timestamp. Reconstruct it here.
-		if (ActualCycle < BaseCycle)
-		{
-			ActualCycle += BaseCycle;
-		}
+////////////////////////////////////////////////////////////////////////////////////////////////////
+// Async Trim Processor
+////////////////////////////////////////////////////////////////////////////////////////////////////
 
-		if (DecodedCycle & 2ull) // coroutine
+TUniquePtr<ITrimAnalysisProcessor> TrimAsync(FTrimParameters& Parameters, UE::Trace::IInDataStream& InputDataStream, UE::Trace::IOutDataStream& OutputDataStream)
+{
+	class FTrimAnalysisProcessor final : public ITrimAnalysisProcessor
+	{
+	public:
+		FTrimAnalysisProcessor(FTrimParameters& Parameters, UE::Trace::IInDataStream& InputDataStream, UE::Trace::IOutDataStream& OutputDataStream)
+			: TraceWriter(OutputDataStream)
+			, TrimAnalyzer(Parameters, TraceWriter)
 		{
-			if (DecodedCycle & 1ull) // Restore
-			{
-				uint64 CoroutineId = FTraceAnalyzerUtils::Decode7bit(BufferPtr);
-				uint32 TimerScopeDepth = (uint32)FTraceAnalyzerUtils::Decode7bit(BufferPtr);
-			}
-			else // Save
-			{
-				uint32 TimerScopeDepth = (uint32)FTraceAnalyzerUtils::Decode7bit(BufferPtr);
-			}
+			Context.AddAnalyzer(TrimAnalyzer);
+			Processor = Context.Process(InputDataStream);
 		}
-		else // normal CPU events
+
+		virtual ~FTrimAnalysisProcessor()
 		{
-			if (DecodedCycle & 1ull)
-			{
-				uint32 SpecId = (uint32)FTraceAnalyzerUtils::Decode7bit(BufferPtr);
-#if 0
-				if (Version == 3)
-				{
-					if (SpecId & 1u) // The last bit is set if this is a metadata id.
-					{
-						uint32 MetadataId = SpecId >> 1;
-					}
-					else
-					{
-						SpecId = SpecId >> 1;
-					}
-				}
-				else // Version == 2
-				{
-					//SpecId
-				}
-#endif
-			}
 		}
 
-		check(ActualCycle > 0);
-		LastCycle = ActualCycle;
-	}
-	if (!ensure(BufferPtr == BufferEnd))
-	{
-		// error
-	}
-
-	Thread.LastCycle = LastCycle;
-	UpdateThreadTime(Thread, Context, LastCycle);
-}
+		virtual bool IsActive() const override { return Processor.IsActive(); }
+		virtual void Stop() override { Processor.Stop(); }
+		virtual void Wait() override { Processor.Wait(); }
+		virtual void Pause(bool bState) override { Processor.Pause(bState); }
 
-////////////////////////////////////////////////////////////////////////////////////////////////////
+	private:
+		UE::Trace::FTraceWriter TraceWriter;
+		FTrimAnalyzer TrimAnalyzer;
+		UE::Trace::FAnalysisContext Context;
+		UE::Trace::FAnalysisProcessor Processor;
+	};
 
-UE::Trace::FAnalysisProcessor TrimAsync(FTrimParameters& Parameters, UE::Trace::IInDataStream& InputDataStream, UE::Trace::IOutDataStream& OutputDataStream)
-{
-	UE::Trace::FTraceWriter TraceWriter(OutputDataStream);
-	FTrimAnalyzer TrimAnalyzer(Parameters, TraceWriter);
-	UE::Trace::FAnalysisContext Context;
-	Context.AddAnalyzer(TrimAnalyzer);
-	return Context.Process(InputDataStream);
+	return MakeUnique<FTrimAnalysisProcessor>(Parameters, InputDataStream, OutputDataStream);
 }
 
+////////////////////////////////////////////////////////////////////////////////////////////////////
+// Sync Trim Processor
 ////////////////////////////////////////////////////////////////////////////////////////////////////
 
 int32 Trim(FTrimParameters& Parameters, UE::Trace::IInDataStream& InputDataStream, UE::Trace::IOutDataStream& OutputDataStream)
@@ -1405,3 +1170,5 @@ int32 Trim(FTrimParameters& Parameters, UE::Trace::IInDataStream& InputDataStrea
 #undef LogTraceTrimmerError
 #undef LogTraceTrimmerWarning
 #undef LogTraceTrimmerMessage
+
+#undef TRIMMER_USE_THREAD_TABLE
diff --git a/Engine/Source/Developer/TraceServices/Private/Analyzers/TrimAnalyzer.h b/Engine/Source/Developer/TraceServices/Private/Analyzers/TrimAnalyzer.h
--- a/Engine/Source/Developer/TraceServices/Private/Analyzers/TrimAnalyzer.h
+++ b/Engine/Source/Developer/TraceServices/Private/Analyzers/TrimAnalyzer.h
@@ -3,6 +3,7 @@
 #pragma once
 
 #include "CoreTypes.h"
+
 #include "Containers/Map.h"
 #include "Containers/AnsiString.h"
 #include "Containers/UnrealString.h"
@@ -11,6 +12,9 @@
 // TraceAnalysis
 #include "Trace/Analyzer.h"
 
+// TraceServices
+#include "Analyzers/KnownEventProcessor.h"
+
 namespace UE::Trace
 {
 	class FTraceWriter;
@@ -45,114 +49,10 @@ private:
 		TArray<FNameFilterPattern> Patterns;
 	};
 
-	enum class EKnownEvent : uint32
-	{
-		Invalid,
-
-		Trace_NewEvent,
-		//Trace_NewTrace,
-		//Trace_ThreadTiming,
-		//Trace_ThreadInfo,
-		//Trace_ThreadGroupBegin,
-		//Trace_ThreadGroupEnd,
-		//Trace_ChannelAnnounce_Old, // "$Trace", deprecated in 4.26
-
-		//Trace_ChannelAnnounce, // "Trace", new in 4.26
-		//Trace_ChannelToggle,
-
-		//Diagnostics_Session,
-		//Diagnostics_Session2,
-
-		//Misc_CreateThread,
-		//Misc_SetThreadGroup,
-		//Misc_BeginThreadGroupScope,
-		//Misc_BookmarkSpec,
-		Misc_Bookmark,
-		Misc_ScreenshotHeader,
-		//Misc_ScreenshotChunk,
-		Misc_RegionBegin,
-		Misc_RegionBeginWithId,
-		Misc_RegionEnd,
-		Misc_RegionEndWithId,
-		//Misc_BeginGameFrame, // deprecated
-		//Misc_EndGameFrame, // deprecated
-		//Misc_BeginRenderFrame, // deprecated
-		//Misc_EndRenderFrame, // deprecated
-		Misc_BeginFrame,
-		Misc_EndFrame,
-
-		//Logging_LogCategory,
-		//Logging_LogMessageSpec,
-		Logging_LogMessage,
-
-		PlatformFile_BeginOpen,
-		PlatformFile_EndOpen,
-		PlatformFile_BeginRead,
-		PlatformFile_EndRead,
-		PlatformFile_BeginClose,
-		PlatformFile_EndClose,
-
-		//LoadTime_ClassInfo,
-		//LoadTime_NewAsyncPackage,
-
-		//GpuProfiler_EventSpec, // old GpuProfiler, deprecated
-		//GpuProfiler_Frame, // old GpuProfiler, deprecated
-		//GpuProfiler_Frame2, // old GpuProfiler, deprecated
-		//GpuProfiler_Init,
-		//GpuProfiler_QueueSpec,
-		//GpuProfiler_EventFrameBoundary,
-		//GpuProfiler_EventBreadcrumbSpec,
-		GpuProfiler_EventBeginBreadcrumb,
-		GpuProfiler_EventEndBreadcrumb,
-		GpuProfiler_EventBeginWork,
-		GpuProfiler_EventEndWork,
-		GpuProfiler_EventWait,
-		//GpuProfiler_EventStats,
-		GpuProfiler_SignalFence,
-		GpuProfiler_WaitFence,
-
-		//CpuProfiler_EventSpec,
-		//CpuProfiler_MetadataSpec, // added in UE 5.x
-		CpuProfiler_Metadata, // added in UE 5.x
-		CpuProfiler_EventBatchV3, // added in UE 5.6
-		CpuProfiler_EventBatchV2, // backward compatibility; added in UE 5.1, removed in 5.6
-		CpuProfiler_EventBatch, // backward compatibility; removed in UE 5.1
-		CpuProfiler_EndCapture, // backward compatibility; removed in UE 5.1
-		CpuProfiler_EndThread, // "Cycle" field added in UE 5.4
-
-		//Stats_Spec,
-
-		//Counters_Spec,
-
-		//CsvProfiler_RegisterCategory,
-		//CsvProfiler_DefineInlineStat,
-		//CsvProfiler_DefineDeclaredStat,
-		//CsvProfiler_Metadata,
-		//CsvProfiler_BeginCapture,
-
-		Memory_Marker,
-
-		LLM_TagValue,
-		//LLM_TagsSpec,
-		//LLM_TrackerSpec,
-
-		SlateTrace_AddWidget,
-
-		Count
-	};
-
-	typedef TFunction<void(const FOnEventContext& Context)> KnownEventCallback;
-
-	struct FKnownEvent
-	{
-		uint32 Id = InvalidEventId;
-		KnownEventCallback Callback = nullptr;
-	};
-
 	struct FThreadInfo
 	{
 		uint32 Id = 0;
-		uint64 LastCycle = 0; // for CpuProfiler.EventBatch decoding
+		int32 SortIndex = 0;
 		uint64 CurrentCycle = 0;
 		double CurrentTime = 0.0;
 	};
@@ -165,33 +65,23 @@ public:
 	uint32 GetNumWarnings() const { return NumWarnings; }
 
 private:
-	void InitNameFilter(FAnsiString InFilter, FNameFilter& OutNameFilter);
+	void InitNameFilter(const FAnsiString& InFilter, FNameFilter& OutNameFilter);
 	bool FilterByEventName(const ANSICHAR* LoggerName, const ANSICHAR* EventName);
 
-	virtual void OnAnalysisBegin(const FOnAnalysisContext& Context) override;
+	virtual void OnAnalysisBegin(const UE::Trace::IAnalyzer::FOnAnalysisContext& Context) override;
 	virtual void OnAnalysisEnd() override;
-	virtual void OnVersion(uint32 InTransportVersion, uint32 InProtocolVersion) /* override */;
+	virtual void OnVersion(uint32 InTransportVersion, uint32 InProtocolVersion) override;
 	virtual bool OnNewEvent(uint16 RouteId, const FEventTypeInfo& TypeInfo) override;
-	virtual bool OnEvent(uint16 RouteId, EStyle, const FOnEventContext& Context) override;
+	virtual bool OnEvent(uint16 RouteId, EStyle, const UE::Trace::IAnalyzer::FOnEventContext& Context) override;
 
 	void WriteArrayField(uint32 FieldIndex, const UE::Trace::IAnalyzer::FEventFieldInfo& FieldInfo, const UE::Trace::IAnalyzer::FEventData& EventData, UE::Trace::FTraceWriterEventBuilder& EventBuilder);
 	void WriteScalarField(uint32 FieldIndex, const UE::Trace::IAnalyzer::FEventFieldInfo& FieldInfo, const UE::Trace::IAnalyzer::FEventData& EventData, UE::Trace::FTraceWriterEventBuilder& EventBuilder);
 
 	FThreadInfo& GetOrAddThread(uint32 ThreadId);
+	void SortThreads(int32 ThreadIndex);
+	bool UpdateThreadsOnSyncEvent(FThreadInfo& CurrentThread);
 	void SetCurrentThread(uint32 ThreadId);
-	void UpdateThreadTime(const FOnEventContext& Context, uint64 Cycle);
-	void UpdateThreadTime(FThreadInfo& Thread, const FOnEventContext& Context, uint64 Cycle);
-	void UpdateSessionTime(uint64 Cycle, double Time);
-
-	bool IsKnownEvent(uint32 EventId, EKnownEvent KnownEvent) const
-	{
-		return EventId == KnownEvents[(int)KnownEvent].Id;
-	}
-	void RegisterCallbacksForKnownEvents();
-	void ProcessKnownEvent(uint32 EventId, const ANSICHAR* LoggerName, const ANSICHAR* EventName);
-
-	void ProcessCpuEventBatchV1(const FOnEventContext& Context, const uint8* BufferPtr, uint32 BufferSize);
-	void ProcessCpuEventBatchV2V3(const FOnEventContext& Context, const uint8* BufferPtr, uint32 BufferSize, int32 Version);
+	void UpdateThreadTime(FThreadInfo& Thread, const UE::Trace::IAnalyzer::FOnEventContext& Context, uint64 Cycle);
 
 private:
 	FTrimParameters& Parameters;
@@ -200,10 +90,22 @@ private:
 	uint32 TransportVersion = 4;
 	uint32 ProtocolVersion = 7;
 
-	TMap<uint32, FThreadInfo> ThreadMap; // ThreadId --> FThreadInfo
+	static constexpr uint32 ThreadTableInitialSize = 1024;
+	static constexpr uint32 ThreadTableMaxSize = 4 * 1024;
+
+	TMap<uint32, FThreadInfo*> ThreadMap; // ThreadId --> FThreadInfo
+	TArray<FThreadInfo*> ThreadTable; // index == ThreadId, for ThreadId < ThreadTableMaxSize
+	TArray<FThreadInfo*> SortedThreads; // sorted threads by time
+
 	TMap<uint32, FEventInfo> EventInfos; // EventId --> FEventInfo
 	FEventInfo EmptyEventInfo;
 
+	uint64 ReadEvents = 0;
+	uint64 WrittenEvents = 0;
+
+	uint64 LastProgressMessageTimestamp = 0;
+	uint64 ProgressMessageInterval = 0;
+
 	uint64 AnalysisBeginTimestamp = 0; // realtime timestamp
 
 	uint64 SessionCycleFrequency = 0;
@@ -217,8 +119,6 @@ private:
 	FNameFilter ExcludeFilter;
 
 	static constexpr uint32 InvalidEventId = ~0u;
-	static constexpr uint32 MaxKnownEventId = 0xFFFF;
-
 	static constexpr uint32 InvalidWriterEventId = ~0u;
 
 	uint32 TraceNewTraceEventId = InvalidEventId;
@@ -226,8 +126,7 @@ private:
 
 	uint32 TraceThreadInfoEventId = InvalidEventId;
 
-	FKnownEvent KnownEvents[(int)EKnownEvent::Count];
-	EKnownEvent KnownEventMap[MaxKnownEventId + 1]; // index == EventId --> EKnownEvent
+	FKnownEventProcessor KnownEvents;
 
 	uint32 NumErrors = 0;
 	uint32 NumWarnings = 0;
diff --git a/Engine/Source/Developer/TraceServices/Public/TraceServices/TraceTrimmer.h b/Engine/Source/Developer/TraceServices/Public/TraceServices/TraceTrimmer.h
--- a/Engine/Source/Developer/TraceServices/Public/TraceServices/TraceTrimmer.h
+++ b/Engine/Source/Developer/TraceServices/Public/TraceServices/TraceTrimmer.h
@@ -4,6 +4,7 @@
 
 #include "CoreTypes.h"
 #include "Containers/AnsiString.h"
+#include "Templates/UniquePtr.h"
 
 // TraceAnalysis
 #include "Trace/Analysis.h"
@@ -12,6 +13,8 @@
 
 #include <limits>
 
+#define UE_API TRACESERVICES_API
+
 namespace TraceServices {
 
 struct FTrimParameters
@@ -46,16 +49,38 @@ struct FTrimParameters
 	FAnsiString Exclude;
 };
 
+class ITrimAnalysisProcessor
+{
+public:
+	ITrimAnalysisProcessor() = default;
+	virtual ~ITrimAnalysisProcessor() = default;
+
+	/** Checks if this object instance is valid and currently processing */
+	virtual bool IsActive() const = 0;
+
+	/** End processing a trace stream. */
+	virtual void Stop() = 0;
+
+	/** Wait for the entire stream to have been processed and analyzed. */
+	virtual void Wait() = 0;
+
+	/** Pause or resume the processing.
+	 * @param bState Pause if true, resume if false. */
+	virtual void Pause(bool bState) = 0;
+};
+
 /**
  * Trims a trace stream using the provided parameters.
- * The function starts the trimming process and returns the FAnalysisProcessor. The caller is responsible for waiting the processor to complete.
+ * The function starts the trimming process and returns immediately a unique pointer
+ * to the trim analysis processor object.
+ * The caller is responsible for waiting the processor to complete.
  *
  * @param Parameters The trim process parameters
  * @param InputDataStream The input data stream
  * @param OutputDataStream The output data stream
- * @returns the analysis processor object
+ * @returns a unique pointer to the trim analysis processor object
  */
-TRACESERVICES_API UE::Trace::FAnalysisProcessor TrimAsync(FTrimParameters& Parameters, UE::Trace::IInDataStream& InputDataStream, UE::Trace::IOutDataStream& OutputDataStream);
+UE_API TUniquePtr<ITrimAnalysisProcessor> TrimAsync(FTrimParameters& Parameters, UE::Trace::IInDataStream& InputDataStream, UE::Trace::IOutDataStream& OutputDataStream);
 
 /**
  * Trims a trace stream using the provided parameters.
@@ -66,6 +91,8 @@ TRACESERVICES_API UE::Trace::FAnalysisProcessor TrimAsync(FTrimParameters& Param
  * @param OutputDataStream The output data stream
  * @returns 0 if the trim process is completed successfully, non-zero otherwise
  */
-TRACESERVICES_API int32 Trim(FTrimParameters& Parameters, UE::Trace::IInDataStream& InputDataStream, UE::Trace::IOutDataStream& OutputDataStream);
+UE_API int32 Trim(FTrimParameters& Parameters, UE::Trace::IInDataStream& InputDataStream, UE::Trace::IOutDataStream& OutputDataStream);
 
 } // namespace TraceServices
+
+#undef UE_API
diff --git a/Engine/Source/Programs/TraceTrimmer/TraceTrimmer.Target.cs b/Engine/Source/Programs/TraceTrimmer/TraceTrimmer.Target.cs
--- a/Engine/Source/Programs/TraceTrimmer/TraceTrimmer.Target.cs
+++ b/Engine/Source/Programs/TraceTrimmer/TraceTrimmer.Target.cs
@@ -49,7 +49,6 @@ public class TraceTrimmerTarget : TargetRules
 		GlobalDefinitions.Add("UE_TRACE_ANALYSIS_DEBUG_API=0");
 		GlobalDefinitions.Add("UE_TRACE_ANALYSIS_DEBUG=0");
 		GlobalDefinitions.Add("UE_TRACE_ANALYSIS_DEBUG_LOG_IMPL=1");
-		GlobalDefinitions.Add("UE_TRACE_ANALYSIS_DEBUG_LEVEL=2");
 
 		//GlobalDefinitions.Add("_CRT_SECURE_NO_WARNINGS");
 		//GlobalDefinitions.Add("_CRT_NONSTDC_NO_WARNINGS");
