Enables multiple memory snapshots to be taken using the button from Memory View.
Fixes for the memory view save to file feature.

#Synced-CL 51753244

diff --git a/Engine/Source/Developer/TraceAnalysis/Private/Analysis/Context.cpp b/Engine/Source/Developer/TraceAnalysis/Private/Analysis/Context.cpp
--- a/Engine/Source/Developer/TraceAnalysis/Private/Analysis/Context.cpp
+++ b/Engine/Source/Developer/TraceAnalysis/Private/Analysis/Context.cpp
@@ -18,13 +18,19 @@ void FAnalysisContext::SetMessageDelegate(FMessageDelegate Delegate)
 	OnMessage = Delegate;
 }
 
+////////////////////////////////////////////////////////////////////////////////
+void FAnalysisContext::SetOnCompletedDelegate(FOnCompletedDelegate Delegate)
+{
+	OnCompleted = Delegate;
+}
+
 ////////////////////////////////////////////////////////////////////////////////
 FAnalysisProcessor FAnalysisContext::Process(IInDataStream& DataStream)
 {
 	FAnalysisProcessor Processor;
 	if (Analyzers.Num() > 0)
 	{
-		Processor.Impl = new FAnalysisProcessor::FImpl(DataStream, MoveTemp(Analyzers), MoveTemp(OnMessage));
+		Processor.Impl = new FAnalysisProcessor::FImpl(DataStream, MoveTemp(Analyzers), MoveTemp(OnMessage), MoveTemp(OnCompleted));
 	}
 	return MoveTemp(Processor);
 }
diff --git a/Engine/Source/Developer/TraceAnalysis/Private/Analysis/Processor.cpp b/Engine/Source/Developer/TraceAnalysis/Private/Analysis/Processor.cpp
--- a/Engine/Source/Developer/TraceAnalysis/Private/Analysis/Processor.cpp
+++ b/Engine/Source/Developer/TraceAnalysis/Private/Analysis/Processor.cpp
@@ -15,11 +15,12 @@ namespace UE {
 namespace Trace {
 
 ////////////////////////////////////////////////////////////////////////////////
-FAnalysisProcessor::FImpl::FImpl(IInDataStream& InDataStream, TArray<IAnalyzer*>&& InAnalyzers, FMessageDelegate&& InMessage)
+FAnalysisProcessor::FImpl::FImpl(IInDataStream& InDataStream, TArray<IAnalyzer*>&& InAnalyzers, FMessageDelegate&& InMessage, FOnCompletedDelegate&& InOnCompleted)
 : AnalysisEngine(Forward<TArray<IAnalyzer*>>(InAnalyzers), Forward<FMessageDelegate>(InMessage))
 , DataStream(InDataStream)
 , StopEvent(FPlatformProcess::GetSynchEventFromPool(true))
 , UnpausedEvent(FPlatformProcess::GetSynchEventFromPool(true))
+, OnCompleted(MoveTemp(InOnCompleted))
 {
 	Thread = FRunnableThread::Create(this, TEXT("TraceAnalysis"));
 	PauseAnalysis(false);
@@ -68,6 +69,7 @@ uint32 FAnalysisProcessor::FImpl::Run()
 
 	AnalysisEngine.End();
 	bComplete = true;
+	OnCompleted.ExecuteIfBound();
 	return 0;
 }
 
diff --git a/Engine/Source/Developer/TraceAnalysis/Private/Analysis/Processor.h b/Engine/Source/Developer/TraceAnalysis/Private/Analysis/Processor.h
--- a/Engine/Source/Developer/TraceAnalysis/Private/Analysis/Processor.h
+++ b/Engine/Source/Developer/TraceAnalysis/Private/Analysis/Processor.h
@@ -18,7 +18,7 @@ class FAnalysisProcessor::FImpl
 	: public FRunnable
 {
 public:
-						FImpl(IInDataStream& DataStream, TArray<IAnalyzer*>&& InAnalyzers, FMessageDelegate&& InMessage);
+						FImpl(IInDataStream& DataStream, TArray<IAnalyzer*>&& InAnalyzers, FMessageDelegate&& InMessage, FOnCompletedDelegate&& InOnCompleted);
 						~FImpl();
 	virtual uint32		Run() override;
 	bool				IsActive() const;
@@ -34,6 +34,7 @@ private:
 	FEvent*				UnpausedEvent;
 	FRunnableThread*	Thread = nullptr;
 	volatile bool		bComplete = false;
+	FOnCompletedDelegate OnCompleted;
 };
 
 } // namespace Trace
diff --git a/Engine/Source/Developer/TraceAnalysis/Private/DataStream.cpp b/Engine/Source/Developer/TraceAnalysis/Private/DataStream.cpp
--- a/Engine/Source/Developer/TraceAnalysis/Private/DataStream.cpp
+++ b/Engine/Source/Developer/TraceAnalysis/Private/DataStream.cpp
@@ -35,9 +35,9 @@ FFileDataStream::~FFileDataStream()
 {
 }
 
-bool FFileDataStream::Open(const TCHAR* Path)
+bool FFileDataStream::Open(const TCHAR* Path, bool bAllowWrite)
 {
-	Handle.Reset(FPlatformFileManager::Get().GetPlatformFile().OpenRead(Path));
+	Handle.Reset(FPlatformFileManager::Get().GetPlatformFile().OpenRead(Path, bAllowWrite));
 	if (Handle == nullptr)
 	{
 		return false;
diff --git a/Engine/Source/Developer/TraceAnalysis/Public/Trace/Analysis.h b/Engine/Source/Developer/TraceAnalysis/Public/Trace/Analysis.h
--- a/Engine/Source/Developer/TraceAnalysis/Public/Trace/Analysis.h
+++ b/Engine/Source/Developer/TraceAnalysis/Public/Trace/Analysis.h
@@ -22,7 +22,8 @@ enum class EAnalysisMessageSeverity
 	Error
 };
 
-DECLARE_DELEGATE_TwoParams(FMessageDelegate, EAnalysisMessageSeverity, FStringView);	
+DECLARE_DELEGATE_TwoParams(FMessageDelegate, EAnalysisMessageSeverity, FStringView);
+DECLARE_DELEGATE(FOnCompletedDelegate);
 
 /**
  * Represents the processing (e.g. analysis) of a trace stream. Instances are
@@ -76,6 +77,9 @@ public:
 	/** Adds a callback to recieve important messages. */
 	UE_API void SetMessageDelegate(FMessageDelegate Delegate);
 
+	/** Adds a callback that is invoked on the analysis thread when processing completes. */
+	UE_API void SetOnCompletedDelegate(FOnCompletedDelegate Delegate);
+
 	/** Creates and starts analysis returning an FAnalysisProcessor instance which
 	 * represents the analysis and affords some control over it.
 	 * @param DataStream Input stream of trace log data to be analysed. */
@@ -84,6 +88,7 @@ public:
 private:
 	TArray<IAnalyzer*>	Analyzers;
 	FMessageDelegate	OnMessage;
+	FOnCompletedDelegate OnCompleted;
 };
 
 } // namespace Trace
diff --git a/Engine/Source/Developer/TraceAnalysis/Public/Trace/DataStream.h b/Engine/Source/Developer/TraceAnalysis/Public/Trace/DataStream.h
--- a/Engine/Source/Developer/TraceAnalysis/Public/Trace/DataStream.h
+++ b/Engine/Source/Developer/TraceAnalysis/Public/Trace/DataStream.h
@@ -59,9 +59,10 @@ public:
 	* Open the file.
 	*
 	* @param Path - The path to the file
+	* @param bAllowWrite - If true, other processes may write to the file while it is open
 	* @returns true if the file was opened successfully
 	*/
-	UE_API bool Open(const TCHAR* Path);
+	UE_API bool Open(const TCHAR* Path, bool bAllowWrite = false);
 
 	UE_API virtual int32 Read(void* Data, uint32 Size) override;
 	UE_API virtual void Close() override;
diff --git a/Engine/Source/Developer/TraceServices/Private/AnalysisService.cpp b/Engine/Source/Developer/TraceServices/Private/AnalysisService.cpp
--- a/Engine/Source/Developer/TraceServices/Private/AnalysisService.cpp
+++ b/Engine/Source/Developer/TraceServices/Private/AnalysisService.cpp
@@ -149,12 +149,16 @@ void FAnalysisSession::Start()
 		Context.AddAnalyzer(*Analyzer);
 	}
 	Context.SetMessageDelegate(UE::Trace::FMessageDelegate::CreateRaw(this, &FAnalysisSession::OnAnalysisMessage));
+	Context.SetOnCompletedDelegate(UE::Trace::FOnCompletedDelegate::CreateRaw(this, &FAnalysisSession::OnAnalysisComplete));
 	Processor = Context.Process(*DataStream);
 }
 
 void FAnalysisSession::Stop(bool bAndWait) const
 {
-	DataStream->Close();
+	if (!bDataStreamClosed.exchange(true))
+	{
+		DataStream->Close();
+	}
 	Processor.Stop();
 	if (bAndWait)
 	{
@@ -167,6 +171,14 @@ void FAnalysisSession::Wait() const
 	Processor.Wait();
 }
 
+void FAnalysisSession::OnAnalysisComplete()
+{
+	if (!bDataStreamClosed.exchange(true))
+	{
+		DataStream->Close();
+	}
+}
+
 void FAnalysisSession::EnumerateMetadata(TFunctionRef<void(const FTraceSessionMetadata& Metadata)> Callback) const
 {
 	Lock.ReadAccessCheck();
@@ -291,42 +303,13 @@ TSharedPtr<const IAnalysisSession> FAnalysisService::Analyze(const TCHAR* Sessio
 
 TSharedPtr<const IAnalysisSession> FAnalysisService::StartAnalysis(const TCHAR* SessionUri)
 {
-	struct FFileDataStream
-		: public UE::Trace::IInDataStream
-	{
-		virtual int32 Read(void* Data, uint32 Size) override
-		{
-			if (Remaining <= 0)
-			{
-				return 0;
-			}
-			if (Size > Remaining)
-			{
-				Size = static_cast<uint32>(Remaining);
-			}
-			Remaining -= Size;
-			if (!Handle->Read((uint8*)Data, Size))
-			{
-				return 0;
-			}
-			return Size;
-		}
-
-		TUniquePtr<IFileHandle> Handle;
-		uint64 Remaining;
-	};
-
-	IPlatformFile& FileSystem = IPlatformFile::GetPlatformPhysical();
-	IFileHandle* Handle = FileSystem.OpenRead(SessionUri, true);
-	if (!Handle)
+	UE::Trace::FFileDataStream* FileStream = new UE::Trace::FFileDataStream();
+	if (!FileStream->Open(SessionUri, /*bAllowWrite=*/true))
 	{
+		delete FileStream;
 		return nullptr;
 	}
 
-	FFileDataStream* FileStream = new FFileDataStream();
-	FileStream->Handle = TUniquePtr<IFileHandle>(Handle);
-	FileStream->Remaining = Handle->Size();
-
 	TUniquePtr<UE::Trace::IInDataStream> DataStream(FileStream);
 	return StartAnalysis(~0, SessionUri, MoveTemp(DataStream));
 }
diff --git a/Engine/Source/Developer/TraceServices/Private/AnalysisServicePrivate.h b/Engine/Source/Developer/TraceServices/Private/AnalysisServicePrivate.h
--- a/Engine/Source/Developer/TraceServices/Private/AnalysisServicePrivate.h
+++ b/Engine/Source/Developer/TraceServices/Private/AnalysisServicePrivate.h
@@ -92,6 +92,7 @@ private:
 	virtual IEditableProvider* EditProviderPrivate(const FName& Name) override;
 
 	void OnAnalysisMessage(UE::Trace::EAnalysisMessageSeverity Severity, FStringView Message);
+	void OnAnalysisComplete();
 
 private:
 	mutable FAnalysisSessionLock Lock;
@@ -107,6 +108,7 @@ private:
 	TMap<FName, TTuple<TSharedPtr<IProvider>, TSharedPtr<IEditableProvider>>> Providers;
 	TArray<FAnalysisMessage> PendingMessages;
 	mutable std::atomic<uint32> PendingMessagesCount;
+	mutable std::atomic<bool> bDataStreamClosed { false };
 	mutable TUniquePtr<UE::Trace::IInDataStream> DataStream;
 	mutable UE::Trace::FAnalysisProcessor Processor;
 };
