[Trace Services] Analysis to read Delta7Bit with correct event tag name for any future traces that choose this compression scheme

#Synced-CL 52446069

diff --git a/Engine/Source/Developer/TraceServices/Private/Analyzers/CallstacksAnalysis.cpp b/Engine/Source/Developer/TraceServices/Private/Analyzers/CallstacksAnalysis.cpp
--- a/Engine/Source/Developer/TraceServices/Private/Analyzers/CallstacksAnalysis.cpp
+++ b/Engine/Source/Developer/TraceServices/Private/Analyzers/CallstacksAnalysis.cpp
@@ -34,8 +34,7 @@ void FCallstacksAnalyzer::OnAnalysisBegin(const FOnAnalysisContext& Context)
 	FInterfaceBuilder& Builder = Context.InterfaceBuilder;
 	Builder.RouteEvent(RouteId_Callstack, "Memory", "CallstackSpec");
 	Builder.RouteEvent(RouteId_CallstackXORAndRLE, "Memory", "CallstackSpecXORAndRLE");
-	// NOTE: Update Event Tag to 'CallstackSpecDelta7bit' before UE 5.8 
-	Builder.RouteEvent(RouteId_CallstackDelta7bit, "Memory", "CallstackSpecZigZagAndVLE");
+	Builder.RouteEvent(RouteId_CallstackDelta7bit, "Memory", "CallstackSpecDelta7bit");
 	Builder.RouteEvent(RouteId_CallstackDeltaVarInt, "Memory", "CallstackSpecDeltaVarInt");
 }
 
