[Insights] Updated use of UE_LOGF in CpuProfilerTraceAnalysis and GpuProfilerTraceAnalysis.

#rnx
#Synced-CL 52677305

diff --git a/Engine/Source/Developer/TraceServices/Private/Analyzers/CpuProfilerTraceAnalysis.cpp b/Engine/Source/Developer/TraceServices/Private/Analyzers/CpuProfilerTraceAnalysis.cpp
--- a/Engine/Source/Developer/TraceServices/Private/Analyzers/CpuProfilerTraceAnalysis.cpp
+++ b/Engine/Source/Developer/TraceServices/Private/Analyzers/CpuProfilerTraceAnalysis.cpp
@@ -19,11 +19,16 @@
 #define CPUPROFILER_DEBUG_LOGF(Format, ...) //{ if (ThreadState.ThreadId == 2) FPlatformMisc::LowLevelOutputDebugStringf(Format, __VA_ARGS__); }
 #define CPUPROFILER_DEBUG_BEGIN_EVENT(Time, Event) { ++TotalScopeCount; }
 #define CPUPROFILER_DEBUG_END_EVENT(Time)
-#define CPUPROFILER_LOG_API_L1(Format, ...) UE_LOG(LogTraceServices, Log, Format, ##__VA_ARGS__)
-#define CPUPROFILER_LOG_API_L2(Format, ...) //UE_LOG(LogTraceServices, Verbose, Format, ##__VA_ARGS__)
-#define CPUPROFILER_LOG_API_L3(Format, ...) //UE_LOG(LogTraceServices, VeryVerbose, Format, ##__VA_ARGS__)
+#define CPUPROFILER_LOG_API_L1(Format, ...) UE_LOGF(LogTraceServices, Log, Format, ##__VA_ARGS__)
+#define CPUPROFILER_LOG_API_L2(Format, ...) //UE_LOGF(LogTraceServices, Verbose, Format, ##__VA_ARGS__)
+#define CPUPROFILER_LOG_API_L3(Format, ...) //UE_LOGF(LogTraceServices, VeryVerbose, Format, ##__VA_ARGS__)
+
+// debug functionality; adds detection and extra logging for timer name changes
 #define CPUPROFILER_DETECT_TIMER_NAME_CHANGE 0
 
+// debug functionality; filters the trace events by ThreadId
+#define CPUPROFILER_FILTER_BY_THREAD(ThreadId, ReturnOrBreak) //{ if (ThreadId != 2) { ReturnOrBreak; } }
+
 namespace TraceServices
 {
 
@@ -149,6 +154,7 @@ bool FCpuProfilerAnalyzer::OnEvent(uint16 RouteId, EStyle Style, const FOnEventC
 	case RouteId_EndThread:
 	{
 		const uint32 ThreadId = FTraceAnalyzerUtils::GetThreadIdField(Context);
+		CPUPROFILER_FILTER_BY_THREAD(ThreadId, break);
 		FThreadState& ThreadState = GetOrAddThreadState(ThreadId);
 
 		if (ThreadState.LastCycle == ~0ull)
@@ -183,6 +189,7 @@ bool FCpuProfilerAnalyzer::OnEvent(uint16 RouteId, EStyle Style, const FOnEventC
 	case RouteId_EventBatchV2: // backward compatibility
 	{
 		const uint32 ThreadId = Context.ThreadInfo.GetId();
+		CPUPROFILER_FILTER_BY_THREAD(ThreadId, break);
 		FThreadState& ThreadState = GetOrAddThreadState(ThreadId);
 
 		if (ThreadState.LastCycle == ~0ull)
@@ -212,6 +219,7 @@ bool FCpuProfilerAnalyzer::OnEvent(uint16 RouteId, EStyle Style, const FOnEventC
 	case RouteId_EndCapture: // backward compatibility
 	{
 		const uint32 ThreadId = FTraceAnalyzerUtils::GetThreadIdField(Context);
+		CPUPROFILER_FILTER_BY_THREAD(ThreadId, break);
 		FThreadState& ThreadState = GetOrAddThreadState(ThreadId);
 
 		if (ThreadState.LastCycle == ~0ull)
@@ -734,6 +742,7 @@ void FCpuProfilerAnalyzer::OnCpuScopeEnter(const FOnEventContext& Context)
 	}
 
 	uint32 ThreadId = Context.ThreadInfo.GetId();
+	CPUPROFILER_FILTER_BY_THREAD(ThreadId, return);
 	FThreadState& ThreadState = GetOrAddThreadState(ThreadId);
 
 	if (ThreadState.bShouldIgnorePendingEvents)
@@ -795,6 +804,7 @@ void FCpuProfilerAnalyzer::OnCpuScopeLeave(const FOnEventContext& Context)
 	}
 
 	uint32 ThreadId = Context.ThreadInfo.GetId();
+	CPUPROFILER_FILTER_BY_THREAD(ThreadId, return);
 	FThreadState& ThreadState = GetOrAddThreadState(ThreadId);
 
 	if (ThreadState.bShouldIgnorePendingEvents)
@@ -870,7 +880,7 @@ void FCpuProfilerAnalyzer::OnEventSpec(const FOnEventContext& Context)
 	}
 	const TCHAR* FileName = !File.IsEmpty() ? *File : nullptr;
 
-	CPUPROFILER_LOG_API_L2(TEXT("[CpuProfiler] EventSpec Id=%u Name=\"%s\" File=\"%s\" Line=%u"),
+	CPUPROFILER_LOG_API_L2(TEXT("[CpuProfiler] EventSpec Id=%u Name=\"%ls\" File=\"%ls\" Line=%u"),
 		SpecId, TimerName, FileName ? FileName : TEXT("N/A"), Line);
 
 	const TCHAR* StoredTimerName = Session.StoreString(TimerName);
@@ -890,7 +900,7 @@ void FCpuProfilerAnalyzer::OnMetadataSpec(const FOnEventContext& Context)
 	EventData.GetString("NameFormat", NameFormat);
 	TArrayView<const uint8> FieldNames = EventData.GetArrayView<uint8>("FieldNames");
 
-	CPUPROFILER_LOG_API_L2(TEXT("[CpuProfiler] MetadataSpec Id=%u Name=\"%s\" NameFormat=\"%s\" Fields[]={%u bytes}"),
+	CPUPROFILER_LOG_API_L2(TEXT("[CpuProfiler] MetadataSpec Id=%u Name=\"%ls\" NameFormat=\"%ls\" Fields[]={%u bytes}"),
 		SpecId, *Name, *NameFormat, FieldNames.Num());
 
 	if (Name.Compare(NameFormat) == 0)
@@ -1041,7 +1051,7 @@ uint32 FCpuProfilerAnalyzer::AddCpuTimer(uint32 SpecId, const TCHAR* TimerName,
 	// Map the SpecId to the timer.
 	SpecIdToTimerIdMap.Add(SpecId, TimerId);
 
-	CPUPROFILER_LOG_API_L3(TEXT("[CpuProfiler] --> AddCpuTimer(SpecId=%u, Name=\"%s\") --> TimerId=%u"), SpecId, TimerName, TimerId);
+	CPUPROFILER_LOG_API_L3(TEXT("[CpuProfiler] --> AddCpuTimer(SpecId=%u, Name=\"%ls\") --> TimerId=%u"), SpecId, TimerName, TimerId);
 
 	return TimerId;
 }
@@ -1166,7 +1176,7 @@ void FCpuProfilerAnalyzer::SetTimerName(uint32 SpecId, uint32 TimerId, const TCH
 	}
 	else
 	{
-		CPUPROFILER_LOG_API_L2(TEXT("[CpuProfiler] --> SetTimerName(SpecId=%u, TimerId=%u, Name=\"%s\""), SpecId, TimerId, TimerName);
+		CPUPROFILER_LOG_API_L2(TEXT("[CpuProfiler] --> SetTimerName(SpecId=%u, TimerId=%u, Name=\"%ls\""), SpecId, TimerId, TimerName);
 	}
 #endif
 }
@@ -1196,6 +1206,7 @@ FCpuProfilerAnalyzer::FThreadState& FCpuProfilerAnalyzer::GetOrAddThreadState(ui
 
 } // namespace TraceServices
 
+#undef CPUPROFILER_FILTER_BY_THREAD
 #undef CPUPROFILER_DETECT_TIMER_NAME_CHANGE
 #undef CPUPROFILER_LOG_API_L3
 #undef CPUPROFILER_LOG_API_L2
diff --git a/Engine/Source/Developer/TraceServices/Private/Analyzers/GpuProfilerTraceAnalysis.cpp b/Engine/Source/Developer/TraceServices/Private/Analyzers/GpuProfilerTraceAnalysis.cpp
--- a/Engine/Source/Developer/TraceServices/Private/Analyzers/GpuProfilerTraceAnalysis.cpp
+++ b/Engine/Source/Developer/TraceServices/Private/Analyzers/GpuProfilerTraceAnalysis.cpp
@@ -14,9 +14,9 @@
 #include "Model/CountersPrivate.h"
 #include "Model/TimingProfilerPrivate.h"
 
-#define GPUPROFILER_LOG_API_L1(Format, ...) UE_LOG(LogTraceServices, Log, Format, ##__VA_ARGS__)
-#define GPUPROFILER_LOG_API_L2(Format, ...) //UE_LOG(LogTraceServices, Verbose, Format, ##__VA_ARGS__)
-#define GPUPROFILER_LOG_API_L3(Format, ...) //UE_LOG(LogTraceServices, VeryVerbose, Format, ##__VA_ARGS__)
+#define GPUPROFILER_LOG_API_L1(Format, ...) UE_LOGF(LogTraceServices, Log, Format, ##__VA_ARGS__)
+#define GPUPROFILER_LOG_API_L2(Format, ...) //UE_LOGF(LogTraceServices, Verbose, Format, ##__VA_ARGS__)
+#define GPUPROFILER_LOG_API_L3(Format, ...) //UE_LOGF(LogTraceServices, VeryVerbose, Format, ##__VA_ARGS__)
 
 namespace TraceServices
 {
@@ -100,7 +100,7 @@ bool FGpuProfilerAnalyzer::OnEvent(uint16 RouteId, EStyle Style, const FOnEventC
 	case RouteId_Init:
 	{
 		Version = (uint32)EventData.GetValue<uint8>("Version");
-		GPUPROFILER_LOG_API_L1(TEXT("[GpuProfiler] Init Version=%u"), Version);
+		GPUPROFILER_LOG_API_L1("[GpuProfiler] Init Version=%u", Version);
 		break;
 	}
 
@@ -115,7 +115,7 @@ bool FGpuProfilerAnalyzer::OnEvent(uint16 RouteId, EStyle Style, const FOnEventC
 		FString Name;
 		EventData.GetString("TypeString", Name);
 
-		GPUPROFILER_LOG_API_L2(TEXT("[GpuProfiler] QueueSpec QueueId=%u, GPU=%u, Index=%u, Type=%d \"%s\""), QueueId, (uint32)GPU, (uint32)Index, (uint32)Type, *Name);
+		GPUPROFILER_LOG_API_L2("[GpuProfiler] QueueSpec QueueId=%u, GPU=%u, Index=%u, Type=%d \"%ls\"", QueueId, (uint32)GPU, (uint32)Index, (uint32)Type, *Name);
 
 		const TCHAR* PersistentName = Session.StoreString(Name);
 
@@ -134,7 +134,7 @@ bool FGpuProfilerAnalyzer::OnEvent(uint16 RouteId, EStyle Style, const FOnEventC
 		const uint32 QueueId = EventData.GetValue<uint32>("QueueId");
 		const uint32 FrameNumber = EventData.GetValue<uint32>("FrameNumber");
 
-		GPUPROFILER_LOG_API_L3(TEXT("[GpuProfiler] EventFrameBoundary QueueId=%u, FrameNumber=%u"), QueueId, FrameNumber);
+		GPUPROFILER_LOG_API_L3("[GpuProfiler] EventFrameBoundary QueueId=%u, FrameNumber=%u", QueueId, FrameNumber);
 
 		FQueue& Queue = GetOrAddQueue(QueueId);
 		Queue.FrameNumber = FrameNumber;
@@ -171,7 +171,7 @@ bool FGpuProfilerAnalyzer::OnEvent(uint16 RouteId, EStyle Style, const FOnEventC
 		const double Time = Context.EventTime.AsSeconds(GPUTimestampTOP);
 		TArray<uint8> Metadata(EventData.GetArrayView<uint8>("Metadata"));
 
-		GPUPROFILER_LOG_API_L3(TEXT("[GpuProfiler] EventBeginBreadcrumb QueueId=%u, Time=%f, SpecId=\"%u\""), QueueId, Time, SpecId);
+		GPUPROFILER_LOG_API_L3("[GpuProfiler] EventBeginBreadcrumb QueueId=%u, Time=%f, SpecId=%u", QueueId, Time, SpecId);
 
 		uint32 TimerId = GetOrAddTimer(SpecId, TEXT("<unknown>"));
 
@@ -210,7 +210,7 @@ bool FGpuProfilerAnalyzer::OnEvent(uint16 RouteId, EStyle Style, const FOnEventC
 		}
 		const double Time = Context.EventTime.AsSeconds(GPUTimestampBOP);
 
-		GPUPROFILER_LOG_API_L3(TEXT("[GpuProfiler] EventEndBreadcrumb QueueId=%u, Time=%f"), QueueId, Time);
+		GPUPROFILER_LOG_API_L3("[GpuProfiler] EventEndBreadcrumb QueueId=%u, Time=%f", QueueId, Time);
 
 		FQueue& Queue = GetOrAddQueue(QueueId);
 		EndEvent(Queue, 0, Time, InvalidTimerId);
@@ -224,7 +224,7 @@ bool FGpuProfilerAnalyzer::OnEvent(uint16 RouteId, EStyle Style, const FOnEventC
 		double Time = Context.EventTime.AsSeconds(GPUTimestampTOP);
 		//const uint64 CPUTimestamp = EventData.GetValue<uint64>("CPUTimestamp"); // not used
 
-		GPUPROFILER_LOG_API_L3(TEXT("[GpuProfiler] EventBeginWork QueueId=%u, GPUTimestampTOP=%f"), QueueId, Time);
+		GPUPROFILER_LOG_API_L3("[GpuProfiler] EventBeginWork QueueId=%u, GPUTimestampTOP=%f", QueueId, Time);
 
 		if (GpuWorkTimerId == InvalidTimerId)
 		{
@@ -258,7 +258,7 @@ bool FGpuProfilerAnalyzer::OnEvent(uint16 RouteId, EStyle Style, const FOnEventC
 		const uint64 GPUTimestampBOP = EventData.GetValue<uint64>("GPUTimestampBOP");
 		const double Time = Context.EventTime.AsSeconds(GPUTimestampBOP);
 
-		GPUPROFILER_LOG_API_L3(TEXT("[GpuProfiler] EventEndWork QueueId=%u, GPUTimestampBOP=%f"), QueueId, Time);
+		GPUPROFILER_LOG_API_L3("[GpuProfiler] EventEndWork QueueId=%u, GPUTimestampBOP=%f", QueueId, Time);
 
 		FQueue& Queue = GetOrAddQueue(QueueId);
 		EndEvent(Queue, 1, Time, GpuWorkTimerId);
@@ -274,7 +274,7 @@ bool FGpuProfilerAnalyzer::OnEvent(uint16 RouteId, EStyle Style, const FOnEventC
 		const uint64 EndTimestamp = EventData.GetValue<uint64>("EndTime");
 		const double EndTime = Context.EventTime.AsSeconds(EndTimestamp);
 
-		GPUPROFILER_LOG_API_L3(TEXT("[GpuProfiler] EventWait QueueId=%u, StartTime=%f, EndTime=%f"), QueueId, StartTime, EndTime);
+		GPUPROFILER_LOG_API_L3("[GpuProfiler] EventWait QueueId=%u, StartTime=%f, EndTime=%f", QueueId, StartTime, EndTime);
 
 		if (GpuWaitTimerId == InvalidTimerId)
 		{
@@ -393,9 +393,9 @@ void FGpuProfilerAnalyzer::BeginEvent(FQueue& Queue, int32 StackIndex, double Be
 			if (++ErrorData.NumWarnings < ErrorData.NumMaxWarnings)
 			{
 				uint32 OriginalBeginEventTimerId = TimingProfilerProvider.GetOriginalTimerIdFromMetadata(BeginEventTimerId);
-				UE_LOGF(LogTraceServices, Warning, "[GpuProfiler] Queue %u %ls BEGIN %d \"%ls\" : Time %f < %f (last) !!!",
+				UE_LOGF(LogTraceServices, Warning, "[GpuProfiler] Queue %u %s BEGIN %d \"%ls\" : Time %f < %f (last) !!!",
 					Queue.Id,
-					StackIndex == 1 ? TEXT("WORK") : TEXT("BREADCRUMB"),
+					StackIndex == 1 ? "WORK" : "BREADCRUMB",
 					int32(BeginEventTimerId), GetTimerName(OriginalBeginEventTimerId),
 					BeginEventTime, LastTimestamp);
 			}
@@ -428,9 +428,9 @@ void FGpuProfilerAnalyzer::EndEvent(FQueue& Queue, int32 StackIndex, double EndE
 		++ErrorData.NumMismatchedEvents;
 		if (++ErrorData.NumWarnings < ErrorData.NumMaxWarnings)
 		{
-			UE_LOGF(LogTraceServices, Warning, "[GpuProfiler] Queue %u %ls END %d \"%ls\" (%f) != BEGIN %d \"%ls\" (%f) !!!",
+			UE_LOGF(LogTraceServices, Warning, "[GpuProfiler] Queue %u %s END %d \"%ls\" (%f) != BEGIN %d \"%ls\" (%f) !!!",
 				Queue.Id,
-				StackIndex == 1 ? TEXT("WORK") : TEXT("BREADCRUMB"),
+				StackIndex == 1 ? "WORK" : "BREADCRUMB",
 				int32(EndEventTimerId), GetTimerName(EndEventTimerId), EndEventTime,
 				int32(BeginEvent.TimerId), GetTimerName(BeginEventTimerId), BeginEvent.Time);
 		}
@@ -443,17 +443,17 @@ void FGpuProfilerAnalyzer::EndEvent(FQueue& Queue, int32 StackIndex, double EndE
 		{
 			if (EndEventTimerId == InvalidTimerId)
 			{
-				UE_LOGF(LogTraceServices, Warning, "[GpuProfiler] Queue %u %ls END : Time %f < %f (BEGIN %d \"%ls\") !!!",
+				UE_LOGF(LogTraceServices, Warning, "[GpuProfiler] Queue %u %s END : Time %f < %f (BEGIN %d \"%ls\") !!!",
 					Queue.Id,
-					StackIndex == 1 ? TEXT("WORK") : TEXT("BREADCRUMB"),
+					StackIndex == 1 ? "WORK" : "BREADCRUMB",
 					EndEventTime, BeginEvent.Time,
 					int32(BeginEvent.TimerId), GetTimerName(BeginEventTimerId));
 			}
 			else
 			{
-				UE_LOGF(LogTraceServices, Warning, "[GpuProfiler] Queue %u %ls END %d \"%ls\" : Time %f < %f (BEGIN %d \"%ls\") !!!",
+				UE_LOGF(LogTraceServices, Warning, "[GpuProfiler] Queue %u %s END %d \"%ls\" : Time %f < %f (BEGIN %d \"%ls\") !!!",
 					Queue.Id,
-					StackIndex == 1 ? TEXT("WORK") : TEXT("BREADCRUMB"),
+					StackIndex == 1 ? "WORK" : "BREADCRUMB",
 					int32(EndEventTimerId), GetTimerName(EndEventTimerId),
 					EndEventTime, BeginEvent.Time,
 					int32(BeginEvent.TimerId), GetTimerName(BeginEventTimerId));
@@ -476,16 +476,16 @@ void FGpuProfilerAnalyzer::EndEvent(FQueue& Queue, int32 StackIndex, double EndE
 			{
 				if (EndEventTimerId == InvalidTimerId)
 				{
-					UE_LOGF(LogTraceServices, Warning, "[GpuProfiler] Queue %u %ls END : Time %f < %f (last) !!!",
+					UE_LOGF(LogTraceServices, Warning, "[GpuProfiler] Queue %u %s END : Time %f < %f (last) !!!",
 						Queue.Id,
-						StackIndex == 1 ? TEXT("WORK") : TEXT("BREADCRUMB"),
+						StackIndex == 1 ? "WORK" : "BREADCRUMB",
 						EndEventTime, LastTimestamp);
 				}
 				else
 				{
-					UE_LOGF(LogTraceServices, Warning, "[GpuProfiler] Queue %u %ls END %d \"%ls\" : Time %f < %f (last) !!!",
+					UE_LOGF(LogTraceServices, Warning, "[GpuProfiler] Queue %u %s END %d \"%ls\" : Time %f < %f (last) !!!",
 						Queue.Id,
-						StackIndex == 1 ? TEXT("WORK") : TEXT("BREADCRUMB"),
+						StackIndex == 1 ? "WORK" : "BREADCRUMB",
 						int32(EndEventTimerId), GetTimerName(EndEventTimerId),
 						EndEventTime, LastTimestamp);
 				}
@@ -510,7 +510,7 @@ void FGpuProfilerAnalyzer::OnEventBreadcrumbSpec(const FOnEventContext& Context)
 	EventData.GetString("NameFormat", NameFormat);
 	TArrayView<const uint8> FieldNames = EventData.GetArrayView<uint8>("FieldNames");
 
-	GPUPROFILER_LOG_API_L2(TEXT("[GpuProfiler] EventBreadcrumbSpec Id=%u Name=\"%s\" NameFormat=\"%s\" Fields[]={%u bytes}"),
+	GPUPROFILER_LOG_API_L2("[GpuProfiler] EventBreadcrumbSpec Id=%u Name=\"%ls\" NameFormat=\"%ls\" Fields[]={%u bytes}",
 		SpecId, *Name, *NameFormat, FieldNames.Num());
 
 	if (Name.Compare(NameFormat) == 0)
@@ -596,7 +596,7 @@ void FGpuProfilerAnalyzer::OnEventBreadcrumbSpec(const FOnEventContext& Context)
 	}
 	else
 	{
-		GPUPROFILER_LOG_API_L2(TEXT("[GpuProfiler] EventBreadcrumbSpec with a new timer (\"%s\")"), *Name);
+		GPUPROFILER_LOG_API_L2("[GpuProfiler] EventBreadcrumbSpec with a new timer (\"%ls\")", *Name);
 		TimerId = AddGpuTimer(SpecId, Name);
 	}
 
@@ -672,11 +672,11 @@ uint32 FGpuProfilerAnalyzer::AddGpuTimer(uint32 SpecId, const FString& Breadcrum
 	{
 		SpecIdToTimerIdMap.Add(SpecId, TimerId);
 
-		GPUPROFILER_LOG_API_L3(TEXT("[GpuProfiler] --> AddGpuTimer(SpecId=%u, Name=\"%s\") --> TimerId=%u"), SpecId, TimerName, TimerId);
+		GPUPROFILER_LOG_API_L3("[GpuProfiler] --> AddGpuTimer(SpecId=%u, Name=\"%ls\") --> TimerId=%u", SpecId, TimerName, TimerId);
 	}
 	else
 	{
-		GPUPROFILER_LOG_API_L3(TEXT("[GpuProfiler] --> AddGpuTimer(Name=\"%s\") --> TimerId=%u"), TimerName, TimerId);
+		GPUPROFILER_LOG_API_L3("[GpuProfiler] --> AddGpuTimer(Name=\"%ls\") --> TimerId=%u", TimerName, TimerId);
 	}
 
 	return TimerId;
